Adding ART scripts to repository

diff --git a/.gitignore b/.gitignore
--- a/.gitignore
+++ b/.gitignore
@@ -52,7 +52,14 @@
 # Ignore various other dependencies that are never checked in
 /Engine/Content/**
 /Engine/Documentation/**
-/Engine/Extras/**
+
+# Ignore all binary types in the extras folder
+/Engine/Extras/*/**
+!/Engine/Extras/Maya_AnimationRiggingTools/**
+/Engine/Extras/Maya_AnimationRiggingTools/**/*.pdf
+/Engine/Extras/Maya_AnimationRiggingTools/**/*.pyd
+/Engine/Extras/Maya_AnimationRiggingTools/MayaTools/General/ART/**
+/Engine/Extras/Maya_AnimationRiggingTools/MayaTools/General/Icons/**
 
 # Ignore any saved local files
 /Engine/Saved/**
diff --git a/Engine/Extras/Maya_AnimationRiggingTools/MayaTools/General/Scripts/ART_Help.py b/Engine/Extras/Maya_AnimationRiggingTools/MayaTools/General/Scripts/ART_Help.py
new file mode 100644
--- /dev/null
+++ b/Engine/Extras/Maya_AnimationRiggingTools/MayaTools/General/Scripts/ART_Help.py
@@ -0,0 +1,133 @@
+import maya.cmds as cmds
+from functools import partial
+import os, cPickle
+
+
+class ART_LearningVideos():
+    
+    def __init__(self):
+        
+        #check to see if the UI exists, if so, delete
+        if cmds.window("ART_LearningVideos_UI", exists = True):
+            cmds.deleteUI("ART_LearningVideos_UI")
+            
+        self.widgets = {}
+        
+        #get access to our maya tools
+        toolsPath = cmds.internalVar(usd = True) + "mayaTools.txt"
+        if os.path.exists(toolsPath):
+            
+            f = open(toolsPath, 'r')
+            self.mayaToolsDir = f.readline()
+            f.close()
+            
+            
+        #create window
+        self.widgets["window"] = cmds.window("ART_LearningVideos_UI", title = "Learning Videos", w = 350, h = 500, sizeable = False)
+        
+        #main layout
+        self.widgets["mainLayout"] = cmds.columnLayout(adj = True)
+        
+        #banner 
+        cmds.image(w = 350, h = 50, image = self.mayaToolsDir + "/General/Icons/ART/artHelpBanner.bmp", parent = self.widgets["mainLayout"])
+        
+        
+        #tab layout
+        self.widgets["tabs"] = cmds.tabLayout(parent = self.widgets["mainLayout"])
+        
+        
+        
+        #Create Tabs
+        
+        #scroll layout [Tab 1]
+        self.widgets["scrollLayout"] = cmds.scrollLayout(parent = self.widgets["tabs"], hst = 0, h = 500)
+        
+        #row column layout
+        self.widgets["rowColLayout"] = cmds.rowColumnLayout(nc = 2, rs = [1, 5], cw = [(1, 100), (2, 215)], cat = [(1, "left", 5), (2, "both", 10)], parent = self.widgets["scrollLayout"])
+        
+        
+        #scroll layout [Tab 2]
+        self.widgets["scrollLayout2"] = cmds.scrollLayout(parent = self.widgets["tabs"], hst = 0, h = 500)
+        
+        #row column layout
+        self.widgets["rowColLayout2"] = cmds.rowColumnLayout(nc = 2, rs = [1,5], cw = [(1, 100), (2, 215)], cat = [(1, "left", 5), (2, "both", 10)], parent = self.widgets["scrollLayout2"])        
+        
+        
+        #RIGGING TAB
+        
+        #skeleton settings
+        cmds.symbolButton(w = 100, h = 100, parent = self.widgets["rowColLayout"], image = self.mayaToolsDir + "/General/Icons/ART/movie.bmp", c = partial(self.playMovie, "skeleton_creation"))
+        cmds.image(w = 210, h = 100, parent = self.widgets["rowColLayout"], image = self.mayaToolsDir + "/General/Icons/ART/rigHelp2.bmp")
+        
+        #joint mover basics
+        cmds.symbolButton(w = 100, h = 100, parent = self.widgets["rowColLayout"], image = self.mayaToolsDir + "/General/Icons/ART/movie.bmp", c = partial(self.playMovie, "skeleton_placement"))
+        cmds.image(w = 210, h = 100, parent = self.widgets["rowColLayout"], image = self.mayaToolsDir + "/General/Icons/ART/rigHelp3.bmp")
+        
+        #joint mover tools
+        cmds.symbolButton(w = 100, h = 100, parent = self.widgets["rowColLayout"], image = self.mayaToolsDir + "/General/Icons/ART/movie.bmp", c = partial(self.playMovie, "joint_mover_tools"))
+        cmds.image(w = 210, h = 100, parent = self.widgets["rowColLayout"], image = self.mayaToolsDir + "/General/Icons/ART/rigHelp4.bmp")  
+        
+        #physique mode
+        cmds.symbolButton(w = 100, h = 100, parent = self.widgets["rowColLayout"], image = self.mayaToolsDir + "/General/Icons/ART/movie.bmp", c = partial(self.playMovie, "customizing_your_mannequin"))
+        cmds.image(w = 210, h = 100, parent = self.widgets["rowColLayout"], image = self.mayaToolsDir + "/General/Icons/ART/rigHelp5.bmp")
+        
+        #deformation setup
+        cmds.symbolButton(w = 100, h = 100, parent = self.widgets["rowColLayout"], image = self.mayaToolsDir + "/General/Icons/ART/movie.bmp", c = partial(self.playMovie, "deformation_setup"))
+        cmds.image(w = 210, h = 100, parent = self.widgets["rowColLayout"], image = self.mayaToolsDir + "/General/Icons/ART/rigHelp6.bmp")
+        
+        #Publish
+        cmds.symbolButton(w = 100, h = 100, parent = self.widgets["rowColLayout"], image = self.mayaToolsDir + "/General/Icons/ART/movie.bmp", c = partial(self.playMovie, "publishing_your_character"))
+        cmds.image(w = 210, h = 100, parent = self.widgets["rowColLayout"], image = self.mayaToolsDir + "/General/Icons/ART/rigHelp7.bmp")
+        
+        #Editing
+        cmds.symbolButton(w = 100, h = 100, parent = self.widgets["rowColLayout"], image = self.mayaToolsDir + "/General/Icons/ART/movie.bmp", c = partial(self.playMovie, "editing_your_character"))
+        cmds.image(w = 210, h = 100, parent = self.widgets["rowColLayout"], image = self.mayaToolsDir + "/General/Icons/ART/rigHelp8.bmp")      
+        
+        
+        
+        
+        
+        #ANIMATION TAB
+        
+        #getting started
+        cmds.symbolButton(w = 100, h = 100, parent = self.widgets["rowColLayout2"], image = self.mayaToolsDir + "/General/Icons/ART/movie.bmp", c = partial(self.playMovie, "animation_getting_started"))
+        cmds.image(w = 210, h = 100, parent = self.widgets["rowColLayout2"], image = self.mayaToolsDir + "/General/Icons/ART/animHelp1.bmp")      
+        
+        #rig overview
+        cmds.symbolButton(w = 100, h = 100, parent = self.widgets["rowColLayout2"], image = self.mayaToolsDir + "/General/Icons/ART/movie.bmp", c = partial(self.playMovie, "animation_rig_overview"))
+        cmds.image(w = 210, h = 100, parent = self.widgets["rowColLayout2"], image = self.mayaToolsDir + "/General/Icons/ART/animHelp2.bmp")  
+        
+        #interface overview
+        cmds.symbolButton(w = 100, h = 100, parent = self.widgets["rowColLayout2"], image = self.mayaToolsDir + "/General/Icons/ART/movie.bmp", c = partial(self.playMovie, "animation_ui_overview"))
+        cmds.image(w = 210, h = 100, parent = self.widgets["rowColLayout2"], image = self.mayaToolsDir + "/General/Icons/ART/animHelp3.bmp")  
+        
+        #import/export
+        cmds.symbolButton(w = 100, h = 100, parent = self.widgets["rowColLayout2"], image = self.mayaToolsDir + "/General/Icons/ART/movie.bmp", c = partial(self.playMovie, "animation_import_export"))
+        cmds.image(w = 210, h = 100, parent = self.widgets["rowColLayout2"], image = self.mayaToolsDir + "/General/Icons/ART/animHelp4.bmp")  
+        
+        #space switching
+        cmds.symbolButton(w = 100, h = 100, parent = self.widgets["rowColLayout2"], image = self.mayaToolsDir + "/General/Icons/ART/movie.bmp", c = partial(self.playMovie, "animation_space_switching"))
+        cmds.image(w = 210, h = 100, parent = self.widgets["rowColLayout2"], image = self.mayaToolsDir + "/General/Icons/ART/animHelp5.bmp")  
+        
+        #pose editor
+        cmds.symbolButton(w = 100, h = 100, parent = self.widgets["rowColLayout2"], image = self.mayaToolsDir + "/General/Icons/ART/movie.bmp", c = partial(self.playMovie, "animation_pose_editor"))
+        cmds.image(w = 210, h = 100, parent = self.widgets["rowColLayout2"], image = self.mayaToolsDir + "/General/Icons/ART/animHelp6.bmp")   
+        
+        #matching and visibility
+        cmds.symbolButton(w = 100, h = 100, parent = self.widgets["rowColLayout2"], image = self.mayaToolsDir + "/General/Icons/ART/movie.bmp", c = partial(self.playMovie, "animation_matching"))
+        cmds.image(w = 210, h = 100, parent = self.widgets["rowColLayout2"], image = self.mayaToolsDir + "/General/Icons/ART/animHelp7.bmp")   
+        
+        
+        
+        
+        #edit tab names
+        cmds.tabLayout(self.widgets["tabs"], edit = True, tabLabel = [(self.widgets["scrollLayout"], "Rigging"), (self.widgets["scrollLayout2"], "Animation")])
+        #show the window
+        cmds.showWindow(self.widgets["window"])
+        
+        
+    
+    def playMovie(self, movieFile, *args):
+        
+        moviePath = self.mayaToolsDir + "/General/ART/Help/LearningVideos/" + movieFile + ".wmv"
+        cmds.launch(mov = moviePath)
\ No newline at end of file
diff --git a/Engine/Extras/Maya_AnimationRiggingTools/MayaTools/General/Scripts/ART_Settings.py b/Engine/Extras/Maya_AnimationRiggingTools/MayaTools/General/Scripts/ART_Settings.py
new file mode 100644
--- /dev/null
+++ b/Engine/Extras/Maya_AnimationRiggingTools/MayaTools/General/Scripts/ART_Settings.py
@@ -0,0 +1,257 @@
+import maya.cmds as cmds
+import maya.mel as mel
+from functools import partial
+import os, cPickle
+
+class ART_Settings():
+    
+    def __init__(self):
+        
+        #get access to our maya tools
+        toolsPath = cmds.internalVar(usd = True) + "mayaTools.txt"
+        if os.path.exists(toolsPath):
+            
+            f = open(toolsPath, 'r')
+            self.mayaToolsDir = f.readline()
+            f.close()
+            
+            
+        #check to see if window exists, if so, delete
+        if cmds.window("AnimationRiggingTool_SettingsUI", exists = True):
+            cmds.deleteUI("AnimationRiggingTool_SettingsUI")
+            
+        self.path = self.mayaToolsDir
+        self.widgets = {}
+        
+        #create the window
+        self.widgets["window"] = cmds.window("AnimationRiggingTool_SettingsUI", w = 400, h = 300, titleBarMenu = False, title = "Settings", sizeable = False)
+        
+        
+        #main layout
+        self.widgets["main"] = cmds.formLayout(w = 400, h = 300)
+        
+        #background image
+        image = self.mayaToolsDir + "/General/Icons/ART/settings.jpg"
+        self.widgets["background"] = cmds.image(w = 400, h = 300, parent = self.widgets["main"], image = image)
+        
+        
+        #close button
+        self.widgets["closeButton"] = cmds.symbolButton(w = 26, h = 20, parent = self.widgets["main"], image = self.mayaToolsDir + "/General/Icons/ART/xbutton.bmp", c = self.close)
+        cmds.formLayout(self.widgets["main"], edit = True, af = [(self.widgets["closeButton"], "top", 2), (self.widgets["closeButton"], "right", 2)])
+        
+        #tools path
+        self.widgets["toolsPathLabel"] = cmds.text(h = 40, w = 100, label = "Tools Path: ", font = "boldLabelFont", parent = self.widgets["main"])
+        self.widgets["toolsPathField"] = cmds.textField(h = 40, w = 225, text = "", parent = self.widgets["main"], editable = False)
+        self.widgets["toolsPathBrowseButton"] = cmds.symbolButton(w = 40, h = 40, parent = self.widgets["main"], image = self.mayaToolsDir + "/General/Icons/ART/browse.bmp", c = self.browseNewToolsPath)
+        
+        cmds.formLayout(self.widgets["main"], edit = True, af = [(self.widgets["toolsPathLabel"], "top", 40), (self.widgets["toolsPathLabel"], "left", 10)])
+        cmds.formLayout(self.widgets["main"], edit = True, af = [(self.widgets["toolsPathField"], "top", 40), (self.widgets["toolsPathField"], "right", 50)])
+        cmds.formLayout(self.widgets["main"], edit = True, af = [(self.widgets["toolsPathBrowseButton"], "top", 40), (self.widgets["toolsPathBrowseButton"], "right", 5)])
+        
+        #Source control settings
+        divider = cmds.separator(w = 400, h = 10, style = "out")
+        self.widgets["useSourceControl"] = cmds.checkBox( h = 40, label = "Use Source Control", v = False, parent = self.widgets["main"], onc = partial(self.useSourceControl, True), ofc = partial(self.useSourceControl, False))
+    
+        cmds.formLayout(self.widgets["main"], edit = True, af = [(divider, "top", 95), (divider, "left", 0)])
+        cmds.formLayout(self.widgets["main"], edit = True, af = [(self.widgets["useSourceControl"], "top", 105), (self.widgets["useSourceControl"], "left", 10)])
+        
+        #source control test button
+        self.widgets["testConnection"] = cmds.button(w = 225, h = 40, label = "Test Perforce Connection", parent = self.widgets["main"], enable = False, c = self.testConnection)
+        cmds.formLayout(self.widgets["main"], edit = True, af = [(self.widgets["testConnection"], "top", 105), (self.widgets["testConnection"], "right", 25)]) 
+        
+        
+        #favorite project option
+	favoriteProjLabel = cmds.text(label = "Set Favorite Project: ", parent = self.widgets["main"])
+	cmds.formLayout(self.widgets["main"], edit = True, af = [(favoriteProjLabel, "top", 170), (favoriteProjLabel, "left", 10)]) 
+	self.widgets["favoriteProject_OM"] = cmds.optionMenu(w = 225, h = 40, label = "", parent = self.widgets["main"])
+	cmds.menuItem(label = "None", parent = self.widgets["favoriteProject_OM"])
+	
+        cmds.formLayout(self.widgets["main"], edit = True, af = [(self.widgets["favoriteProject_OM"], "top", 160), (self.widgets["favoriteProject_OM"], "right", 25)]) 
+	
+	
+	#get projects and add to menu
+	self.findProjects()
+	
+	
+        #save settings button
+        self.widgets["saveSettings"] = cmds.button(w = 380, h = 40, label = "Save Settings and Close", parent = self.widgets["main"], c = self.saveSettings)
+        cmds.formLayout(self.widgets["main"], edit = True, af = [(self.widgets["saveSettings"], "bottom", 10), (self.widgets["saveSettings"], "right", 10)])   
+        
+        #show window
+        cmds.showWindow(self.widgets["window"])
+        
+        #get tools path
+        self.getToolsPath()
+        
+        
+        #load settings from disk
+	settingsLocation = self.mayaToolsDir + "/General/Scripts/projectSettings.txt"
+	if os.path.exists(settingsLocation):
+            f = open(settingsLocation, 'r')
+            settings = cPickle.load(f)	    
+	    useSource = settings.get("UseSourceControl")
+	    favoriteProject = settings.get("FavoriteProject")
+	    
+	    f.close()
+	    
+	    try:
+		#set the checkbox to the settings value
+		cmds.checkBox(self.widgets["useSourceControl"], edit = True, v = useSource)
+		self.useSourceControl(useSource)
+		
+	    except:
+		pass
+	    
+	    try:
+		#set the optionMenu to favorite project
+		cmds.optionMenu(self.widgets["favoriteProject_OM"], edit = True, v = favoriteProject)
+		
+	    except:
+		pass
+	    
+        
+        
+
+    # # # # # # # # # # # # # # # # # # # # # # # # # # # # # # # # # # # # #
+    def close(self, *args):
+        
+        cmds.deleteUI("AnimationRiggingTool_SettingsUI")
+        
+        
+
+    # # # # # # # # # # # # # # # # # # # # # # # # # # # # # # # # # # # # #
+    def getToolsPath(self, *args):
+        
+        toolsPath = cmds.internalVar(usd = True) + "mayaTools.txt"
+        if os.path.exists(toolsPath):
+            
+            f = open(toolsPath, 'r')
+            self.mayaToolsDir = f.readline()
+            f.close()
+            
+        #edit text field
+        cmds.textField(self.widgets["toolsPathField"], edit = True, text = self.mayaToolsDir)
+        
+
+    # # # # # # # # # # # # # # # # # # # # # # # # # # # # # # # # # # # # #
+    def browseNewToolsPath(self, *args):
+        
+        newPath = cmds.fileDialog2(fm = 3)
+        if newPath != None:
+            newPath = newPath[0]
+            
+            if newPath.rpartition("/")[2] != "MayaTools":
+                    cmds.warning("Selected directory is not valid. Please locate the MayaTools directory.")      
+                    return
+            
+            else:
+                #create file that contains this path
+                path = cmds.internalVar(usd = True) + "mayaTools.txt"
+                
+                f = open(path, 'w')
+                f.write(newPath)
+                f.close()
+                
+                self.getToolsPath()
+                self.close()
+                cmds.confirmDialog(title = "Settings", icon = "warning", message = "New tools path will not be used until Maya has been restarted.")
+                
+
+    # # # # # # # # # # # # # # # # # # # # # # # # # # # # # # # # # # # # #
+    def useSourceControl(self, state, *args):
+        
+        if state:
+            cmds.button(self.widgets["testConnection"], edit = True, enable = True)
+	    cmds.menuItem("perforceSubmitMenuItem", edit = True, enable = True)
+	    cmds.menuItem("perforceAddAndSubmitMenuItem", edit = True, enable = True)
+	    cmds.menuItem("perforceCheckOutMenuItem", edit = True, enable = True)
+	    cmds.menuItem("perforceFileHistoryMenuItem", edit = True, enable = True)
+	    cmds.menuItem("perforceGetLatestMenuItem", edit = True, enable = True)
+	    cmds.menuItem("perforceProjectList", edit = True, enable = True)
+
+        if state == False:
+            cmds.button(self.widgets["testConnection"], edit = True, enable = False)
+	    cmds.menuItem("perforceSubmitMenuItem", edit = True, enable = False)
+	    cmds.menuItem("perforceAddAndSubmitMenuItem", edit = True, enable = False)
+	    cmds.menuItem("perforceCheckOutMenuItem", edit = True, enable = False)
+	    cmds.menuItem("perforceFileHistoryMenuItem", edit = True, enable = False)
+	    cmds.menuItem("perforceGetLatestMenuItem", edit = True, enable = False)
+	    cmds.menuItem("perforceProjectList", edit = True, enable = False)
+            
+
+    # # # # # # # # # # # # # # # # # # # # # # # # # # # # # # # # # # # # #
+    def testConnection(self, *args):
+        
+        
+        try:
+            from P4 import P4, P4Exception
+            
+        except:
+            cmds.confirmDialog(title = "Test Connection", icon = "critical", message = "Perforce Python modules not found.")
+            return
+        
+        
+        try:
+            p4 = P4()
+            p4.connect()
+	    
+	    
+	    #client info
+            spec = p4.run( "client", "-o" )[0]
+	    
+            client = spec.get("Client")
+            owner = spec.get("Owner")
+	    host = spec.get("Host")
+	    
+            p4.disconnect()
+            
+            cmds.confirmDialog(title = "Test Connection", message = "Connection Sucessful!\n\nUsername: " + owner + "\nClient: " + client + "\nHost: " + host + ".")
+            
+            
+        except:
+            cmds.confirmDialog(title = "Test Connection", message = "Unable to connect to Perforce server.")
+        
+        
+    # # # # # # # # # # # # # # # # # # # # # # # # # # # # # # # # # # # # #
+    def saveSettings(self, *args):
+	
+	#this function will save out the user's preferences they have set in the UI to disk
+        settingsLocation = self.mayaToolsDir + "/General/Scripts/projectSettings.txt"
+	
+	try:
+	    f = open(settingsLocation, 'w')
+    
+    
+	    #create a dictionary with  values
+	    settings = {}
+	    settings["UseSourceControl"] = cmds.checkBox(self.widgets["useSourceControl"], q = True, v = True)
+	    settings["FavoriteProject"] = cmds.optionMenu(self.widgets["favoriteProject_OM"], q = True, v = True)
+	    
+	    #write our dictionary to file
+	    cPickle.dump(settings, f)
+	    f.close()
+	    
+	except:
+	    cmds.confirmDialog(title = "Error", icon = "critical", message = settingsLocation + " is not writeable. Please make sure this file is not set to read only.")
+	    
+	
+	
+	#close the UI
+	cmds.deleteUI("AnimationRiggingTool_SettingsUI")
+	
+	
+	
+    # # # # # # # # # # # # # # # # # # # # # # # # # # # # # # # # # # # # #
+    def findProjects(self):
+
+	projectPath = self.mayaToolsDir + "/General/ART/Projects/"
+	
+	try:
+	    projects = os.listdir(projectPath)
+	    
+	    if len(projects) > 0:
+		for proj in sorted(projects):
+		    cmds.menuItem(label = proj, parent = self.widgets["favoriteProject_OM"])
+		    
+	except:
+	    pass
diff --git a/Engine/Extras/Maya_AnimationRiggingTools/MayaTools/General/Scripts/ART_addCharacter_UI.py b/Engine/Extras/Maya_AnimationRiggingTools/MayaTools/General/Scripts/ART_addCharacter_UI.py
new file mode 100644
--- /dev/null
+++ b/Engine/Extras/Maya_AnimationRiggingTools/MayaTools/General/Scripts/ART_addCharacter_UI.py
@@ -0,0 +1,197 @@
+import maya.cmds as cmds
+import maya.mel as mel
+from functools import partial
+import os, cPickle
+
+
+class AddCharacter_UI():
+    
+    def __init__(self):
+        
+        #get access to our maya tools
+        toolsPath = cmds.internalVar(usd = True) + "mayaTools.txt"
+        if os.path.exists(toolsPath):
+            
+            f = open(toolsPath, 'r')
+            self.mayaToolsDir = f.readline()
+            f.close()
+            
+            
+        #create a dictionary for our UI widgets
+        self.widgets = {}
+        
+        #create window
+        if cmds.window("addCharacter_UI", exists = True):
+            cmds.deleteUI("addCharacter_UI")
+            
+        self.widgets["window"] = cmds.window("addCharacter_UI", title = "Add Character", w = 430, h = 260, sizeable = False, mxb = False, mnb = False)
+        
+        #create the main layout
+	main = cmds.columnLayout(w = 430)
+	
+	#banner image
+	cmds.image(w = 430, h = 50, image = self.mayaToolsDir + "/General/Icons/ART/artBanner430px.bmp", parent = main)
+	
+        self.widgets["formLayout"] = cmds.formLayout(w = 430, h = 260, parent = main)
+        
+        #create the widgets
+        self.widgets["project"] = cmds.optionMenu(label = "Project:", w = 200, cc = self.findCharacterRigs)
+	self.widgets["search"] = cmds.textFieldGrp(label = "Search:", cl2 = ["left", "left"], ct2 = ["left", "left"], cw2 = [40, 160], text = "coming soon", enable = False)
+        self.widgets["characterList"] = cmds.textScrollList(allowMultiSelection = False, w = 200, h = 180, selectCommand = self.findCharacterThumbnail)
+	self.widgets["thumbnailBorder"] = cmds.image(w = 210, h = 210, image = "", ebg = True, bgc = [0, 0, 0])
+        self.widgets["thumbnail"] = cmds.image(w = 200, h = 200, image = "")
+        self.widgets["addButton"] = cmds.button(label = "Add", w = 60, c = partial(self.addCharacter, False))
+	self.widgets["addButton2"] = cmds.button(label = "Add and Close", w = 80, c = partial(self.addCharacter, True))
+        self.widgets["cancelButton"] = cmds.button(label = "Cancel", w = 60, c = self.cancel)
+        
+        
+        #attach widgets
+        cmds.formLayout(self.widgets["formLayout"], edit = True, af = [(self.widgets["project"],  'top', 10),(self.widgets["project"],  'left', 10)])
+	cmds.formLayout(self.widgets["formLayout"], edit = True, af = [(self.widgets["search"],  'top', 40),(self.widgets["search"],  'left', 10)])
+	
+        cmds.formLayout(self.widgets["formLayout"], edit = True, af = [(self.widgets["characterList"],  'top', 70),(self.widgets["characterList"],  'left', 10)])
+	cmds.formLayout(self.widgets["formLayout"], edit = True, af = [(self.widgets["thumbnailBorder"],  'top', 5),(self.widgets["thumbnailBorder"],  'right', 5)])
+        cmds.formLayout(self.widgets["formLayout"], edit = True, af = [(self.widgets["thumbnail"],  'top', 10),(self.widgets["thumbnail"],  'right', 10)])
+        cmds.formLayout(self.widgets["formLayout"], edit = True, af = [(self.widgets["addButton"],  'bottom', 10),(self.widgets["addButton"],  'left', 215)])
+	cmds.formLayout(self.widgets["formLayout"], edit = True, af = [(self.widgets["addButton2"],  'bottom', 10),(self.widgets["addButton2"],  'left', 280)])
+        cmds.formLayout(self.widgets["formLayout"], edit = True, af = [(self.widgets["cancelButton"],  'bottom', 10),(self.widgets["cancelButton"],  'right', 5)])
+        
+        
+        #show the window
+        cmds.showWindow(self.widgets["window"])
+        self.findProjects()
+	
+	#set favorite project if it exists
+	settingsLocation = self.mayaToolsDir + "/General/Scripts/projectSettings.txt"
+	if os.path.exists(settingsLocation):
+            f = open(settingsLocation, 'r')
+            settings = cPickle.load(f)
+	    favoriteProject = settings.get("FavoriteProject")
+	    
+	    try:
+		cmds.optionMenu(self.widgets["project"], edit = True, v = favoriteProject)
+	    except:
+		pass
+	    
+	    self.findCharacterRigs()
+	
+        
+    def findProjects(self):
+	
+	children = cmds.optionMenu(self.widgets["project"], q = True, itemListLong = True)
+	if children != None:
+	    for child in children:
+		cmds.deleteUI(child)
+	    
+	
+	projectPath = self.mayaToolsDir + "/General/ART/Projects/"
+	try:
+	    projects = os.listdir(projectPath)
+	    
+	    if len(projects) > 0:
+		for proj in sorted(projects):
+		    cmds.menuItem(label = proj, parent = self.widgets["project"])
+		    
+		
+		self.findCharacterRigs()
+		
+	except:
+	    cmds.confirmDialog(title = "Add Character", message = "There are no projects to add a character from.")
+	    cmds.deleteUI(self.widgets["window"])
+	    
+	    
+	    
+
+    def findCharacterRigs(self, *args):
+	
+	#remove any existing child elements in the text scroll list
+	cmds.textScrollList(self.widgets["characterList"], edit = True, removeAll = True)
+
+	#find all character rigs
+	project = cmds.optionMenu(self.widgets["project"], q = True, value = True)
+	
+	rigPath = self.mayaToolsDir + "/General/ART/Projects/" + project + "/AnimRigs/"
+	rigs = os.listdir(rigPath)
+	
+	for rig in rigs:
+	    niceName = rig.partition(".mb")[0]
+	    if os.path.isfile(rigPath + rig):
+		cmds.textScrollList(self.widgets["characterList"], edit = True, append = niceName)
+	    
+	    
+	#select the first item
+	cmds.textScrollList(self.widgets["characterList"], edit = True, selectIndexedItem = 1)
+	
+	self.findCharacterThumbnail()
+    
+    def findCharacterThumbnail(self):
+	project = cmds.optionMenu(self.widgets["project"], q = True, value = True)
+	selectedCharacter = cmds.textScrollList(self.widgets["characterList"], q = True, si = True)[0]
+	thumbnailPath = self.mayaToolsDir + "/General/Icons/ART/Thumbnails/" + project
+
+	try:
+	    thumbnails = os.listdir(thumbnailPath)
+	    
+	    for thumb in thumbnails:
+		if thumb.find("_small") == -1:
+		    if thumb.find(selectedCharacter) == 0:
+			cmds.image(self.widgets["thumbnail"], edit = True, image = thumbnailPath + "/" + thumb)
+			
+	except:
+	    cmds.confirmDialog(icon = "critical", title = "Error", message = "No thumbnail found for character")
+	    return
+		
+		
+    def addCharacter(self, close, *args):
+	project = cmds.optionMenu(self.widgets["project"], q = True, value = True)
+	selectedCharacter = cmds.textScrollList(self.widgets["characterList"], q = True, si = True)[0]
+	rigPath = self.mayaToolsDir + "/General/ART/Projects/" + project + "/AnimRigs/" + selectedCharacter + ".mb"
+
+
+	#find existing namespaces in scene
+	namespaces = cmds.namespaceInfo(listOnlyNamespaces = True)
+	
+	#reference the rig file
+	cmds.file(rigPath, r = True, type = "mayaBinary", loadReferenceDepth = "all", namespace = selectedCharacter, options = "v=0")
+	
+	#clear selection and fit view
+	cmds.select(clear = True)
+	cmds.viewFit()
+	panels = cmds.getPanel(type = 'modelPanel')
+
+
+	#turn on smooth shading
+	for panel in panels:
+	    editor = cmds.modelPanel(panel, q = True, modelEditor = True)
+	    cmds.modelEditor(editor, edit = True, displayAppearance = "smoothShaded", displayTextures = True, textures = True )
+	    
+	    
+	
+	#find new namespaces in scene (this is here in case I need to do something later and I need the new name that was created)
+	newCharacterName = selectedCharacter
+	newNamespaces = cmds.namespaceInfo(listOnlyNamespaces = True)
+	
+	for name in newNamespaces:
+	    if name not in namespaces:
+		newCharacterName = name
+		
+	#launch UI
+	import ART_animationUI
+	reload(ART_animationUI)
+	ART_animationUI.AnimationUI()
+	
+	if close:
+	    cmds.deleteUI(self.widgets["window"])
+	    
+
+	
+	
+   
+	
+    def cancel(self, *args):
+	cmds.deleteUI(self.widgets["window"])
+	    
+	    
+	    
+
+        
\ No newline at end of file
