Adding ART scripts to repository

diff --git a/Engine/Extras/Maya_AnimationRiggingTools/MayaTools/General/Scripts/ART_autoRigger.py b/Engine/Extras/Maya_AnimationRiggingTools/MayaTools/General/Scripts/ART_autoRigger.py
new file mode 100644
--- /dev/null
+++ b/Engine/Extras/Maya_AnimationRiggingTools/MayaTools/General/Scripts/ART_autoRigger.py
@@ -0,0 +1,8164 @@
+#Author: [REDACTED_PERSON]
+#Date: 4.9.13
+
+
+
+import maya.cmds as cmds
+import maya.mel as mel
+import os, ast
+
+class AutoRigger():
+    
+    def __init__(self, handCtrlSpace, progressBar):
+        
+        
+        self.handCtrlSpace = handCtrlSpace
+	
+        #get access to our maya tools
+        toolsPath = cmds.internalVar(usd = True) + "mayaTools.txt"
+        if os.path.exists(toolsPath):
+            
+            f = open(toolsPath, 'r')
+            self.mayaToolsDir = f.readline()
+            f.close()
+            
+
+        #create a progress window to track the progress of the rig build
+        self.progress = 0
+        cmds.progressBar(progressBar, edit = True, progress=self.progress, status='Creating Spine Rig')
+        
+        
+
+	#build the core of the rig
+	import Modules.ART_Core
+	coreNodes = Modules.ART_Core.RigCore()
+	
+	
+        #BODY CONTROL
+        self.buildHips()
+	
+        """
+        #create the rig settings node
+        "Rig_Settings" = cmds.group(empty = True, name = "Rig_Settings")
+        cmds.setAttr("Rig_Settings" + ".tx", lock = True, keyable = False)
+        cmds.setAttr("Rig_Settings" + ".ty", lock = True, keyable = False)
+        cmds.setAttr("Rig_Settings" + ".tz", lock = True, keyable = False)
+        cmds.setAttr("Rig_Settings" + ".rx", lock = True, keyable = False)
+        cmds.setAttr("Rig_Settings" + ".ry", lock = True, keyable = False)
+        cmds.setAttr("Rig_Settings" + ".rz", lock = True, keyable = False)
+        cmds.setAttr("Rig_Settings" + ".sx", lock = True, keyable = False)
+        cmds.setAttr("Rig_Settings" + ".sy", lock = True, keyable = False)
+        cmds.setAttr("Rig_Settings" + ".sz", lock = True, keyable = False)
+        cmds.setAttr("Rig_Settings" + ".v", lock = True, keyable = False)
+        
+        
+        #build the spine rigs
+        self.createDriverSkeleton()
+        self.buildCoreComponents()
+        """
+	
+	#to be replaced by modules
+        fkControls = self.buildFKSpine()
+        ikControls = self.buildIKSpine(fkControls)
+	
+        
+        #build the leg rigs
+	
+	#first determine the leg style
+	legStyle = cmds.getAttr("SkeletonSettings_Cache.legStyle")
+	
+	if legStyle == "Standard Biped":
+	    cmds.progressBar(progressBar, edit = True, progress = 20, status='Creating Leg Rigs')
+	    self.buildFKLegs()
+	    self.buildIKLegs()
+	    self.finishLegs()
+
+	    
+	    cmds.progressBar(progressBar, edit = True, progress = 30, status='Creating Toe Rigs')
+	    self.buildToes()
+	    cmds.progressBar(progressBar, edit = True, progress = 40, status='Creating Auto Hips and Spine')
+	    self.buildAutoHips()
+	    self.autoSpine()
+	    
+	if legStyle == "Hind Leg":
+	    pass
+
+        
+    
+        #build the arms
+        cmds.progressBar(progressBar, edit = True, progress = 50, status='Creating Arm Rigs')
+	spineBones = self.getSpineJoints()
+	lastSpine = "driver_" + spineBones[-1]
+	print lastSpine
+        import Modules.ART_Arm
+        reload(Modules.ART_Arm)
+        Modules.ART_Arm.Arm(True, "", None, "l", lastSpine, 6, True)
+        Modules.ART_Arm.Arm(True, "", None, "r", lastSpine, 13, True)
+	
+	
+	"""
+        self.buildFKArms()
+        self.buildIkArms()
+	"""
+	
+        cmds.progressBar(progressBar, edit = True, progress = 60, status='Creating Finger Rigs')
+        self.buildFingers()
+        
+        #build the neck and head rig
+        cmds.progressBar(progressBar, edit = True, progress = 70, status='Creating Neck and Head Rigs')
+        self.buildNeckAndHead()
+        
+        
+        #rig extra joints
+        cmds.progressBar(progressBar, edit = True, progress = 80, status='Creating Rigs for Custom Joints')
+        createdControls = self.rigLeafJoints()
+        createdJiggleNodes = self.rigJiggleJoints()
+	createdChainNodes = self.rigCustomJointChains()
+        
+        
+        
+
+
+        #clean up the hierarchy
+	cmds.progressBar(progressBar, edit = True, progress = 90, status='Cleaning up Scene')
+	
+        bodyGrp = cmds.group(empty = True, name = "body_grp")
+	for obj in ["spine_splineIK_curve", "splineIK_spine_01_splineIK", "body_anim_space_switcher_follow"]:
+	    if cmds.objExists(obj):
+		cmds.parent(obj, bodyGrp)
+
+	if cmds.objExists("autoHips_sys_grp"):
+	    cmds.parent("autoHips_sys_grp", "body_anim")
+	    
+	    
+	rigGrp = "ctrl_rig"
+	cmds.parent([bodyGrp, "leg_sys_grp"], rigGrp)
+
+	"""
+        rigGrp = cmds.group(empty = True, name = "ctrl_rig")
+        cmds.parent([bodyGrp, "leg_sys_grp", "Rig_Settings"], rigGrp)
+	cmds.parent(rigGrp, "offset_anim")
+	"""
+        
+        
+        #Arms
+	"""
+        if cmds.objExists("arm_rig_master_grp_l"):
+            cmds.setAttr("Rig_Settings.lArmMode", 1)
+            
+	    if cmds.objExists("lowerarm_l_roll_grp"):
+		cmds.parent("lowerarm_l_roll_grp", "arm_rig_master_grp_l")
+		
+            cmds.parent("arm_rig_master_grp_l", "ctrl_rig")
+            
+	
+        if cmds.objExists("arm_rig_master_grp_r"):
+            cmds.setAttr("Rig_Settings.rArmMode", 1)
+            		
+	    if cmds.objExists("lowerarm_r_roll_grp"):
+		cmds.parent("lowerarm_r_roll_grp", "arm_rig_master_grp_r")
+                
+            cmds.parent("arm_rig_master_grp_r", "ctrl_rig")
+            
+        if cmds.objExists("arm_rig_master_grp_r") and cmds.objExists("arm_rig_master_grp_l"):
+            armSysGrp = cmds.group(empty = True, name = "arm_sys_grp")
+            cmds.parent(armSysGrp, "ctrl_rig")
+            cmds.parent(["arm_rig_master_grp_r", "arm_rig_master_grp_l", "ik_wrist_l_anim_space_switcher_follow", "ik_wrist_r_anim_space_switcher_follow"], armSysGrp)
+	    
+	    #arm twists
+	    if cmds.objExists("upperarm_twist_grp_l"):
+		cmds.parent("upperarm_twist_grp_l", armSysGrp)
+		
+	    if cmds.objExists("upperarm_twist_grp_r"):
+		cmds.parent("upperarm_twist_grp_r", armSysGrp)	    
+	"""
+	
+            
+	if cmds.objExists("neck_01_fk_anim_grp"):
+	    cmds.parent("neck_01_fk_anim_grp", "ctrl_rig")
+            
+	#Fingers
+        if cmds.objExists("finger_sys_grp_l"):
+            cmds.parent("finger_sys_grp_l", "ctrl_rig")
+            
+        if cmds.objExists("finger_sys_grp_r"):
+            cmds.parent("finger_sys_grp_r", "ctrl_rig")
+	    
+	    
+	#Custom Joints (leaf, jiggle, chain)
+        if len(createdControls) > 0:
+            for each in createdControls:
+                cmds.parent(each, "ctrl_rig")
+		
+        if len(createdJiggleNodes) > 0:
+            for each in createdJiggleNodes:
+                cmds.parent(each, "ctrl_rig")
+		
+        if len(createdChainNodes) > 0:
+            for each in createdChainNodes:
+                cmds.parent(each, "ctrl_rig")
+                
+        cmds.parent("head_sys_grp", "ctrl_rig")
+	
+	#finish grouping everything under 1 character grp
+	if cmds.objExists("Proxy_Geo_Skin_Grp"):
+	    try:
+		cmds.parent("Proxy_Geo_Skin_Grp", "rig_grp")
+	    except:
+		pass
+	    
+	if cmds.objExists("dynHairChain"):
+	    try:
+		cmds.parent("dynHairChain", "rig_grp")
+	    except:
+		pass
+        
+        
+        #add world spaces to each space switch control
+	self.addSpaces()
+        
+        #Hide all joints
+	joints = cmds.ls(type = 'joint')
+	for joint in joints:
+	    if cmds.getAttr(joint + ".v", settable = True):
+		cmds.setAttr(joint + ".v", 0)
+	cmds.progressBar(progressBar, edit = True, progress = 100, status='Cleaning up Scene')
+	
+	#delete the joint mover
+	cmds.select("root_mover_grp", r = True, hi = True)
+        cmds.select("Skeleton_Settings", add = True)
+        nodes = cmds.ls(sl = True, transforms = True)
+        cmds.select(clear = True)
+	
+        for node in nodes:
+            cmds.lockNode(node, lock = False)
+	cmds.lockNode("JointMover", lock = False)
+	cmds.delete("JointMover")
+	
+	#find and delete junk nodes/clean scene
+	for obj in ["invis_legs_Rig_Settings", "invis_legs_Rig_Settings1", "invis_legs_spine_splineIK_curve", "invis_legs_spine_splineIK_curve1","invis_legs_master_anim_space_switcher_follow", "invis_legs_master_anim_space_switcher_follow1"]:
+	    try:
+		cmds.select("*" + obj + "*")
+		selection = cmds.ls(sl = True)
+		for each in selection:
+		    if cmds.objExists(each):
+			cmds.delete(each)
+	    except:
+		pass
+	    
+	    
+	cmds.select(all = True)
+	selection = cmds.ls(sl = True)
+	
+	for each in selection:
+	    if each.find("invis_") == 0:
+		
+		try:
+		    parent = cmds.listRelatives(each, parent = True)
+		    if parent == None:
+			cmds.delete(each)
+		except:
+		    pass
+		
+
+		    
+
+	#set default rotate Orders
+	self.setDefaultRotateOrders()
+
+        #end progress window
+        cmds.select(clear = True)
+	
+	
+        
+        
+        
+    # # # # # # # # # # # # # # # # # # # # # # # # # # # # # # # # # # # # # # # # # # # # # # # # # # # # # # # # # # # # # # # 
+    # # # # # # # # # # # # # # # # # # # # # # # # # # # # # # # # # # # # # # # # # # # # # # # # # # # # # # # # # # # # # # # 
+    # # # # # # # # # # # # # # # # # # # # # # # # # # # # # # # # # # # # # # # # # # # # # # # # # # # # # # # # # # # # # # # 
+    # # # # # # # # # # # # # # # # # # # # # # # # # # # # # # # # # # # # # # # # # # # # # # # # # # # # # # # # # # # # # # # 
+    # # # # # # # # # # # # # # # # # # # # # # # # # # # # # # # # # # # # # # # # # # # # # # # # # # # # # # # # # # # # # # # 
+    def setDefaultRotateOrders(self):
+	
+	cmds.setAttr("body_anim.rotateOrder", 5)
+	cmds.setAttr("hip_anim.rotateOrder", 5)
+	if cmds.objExists("mid_ik_anim"):
+	    cmds.setAttr("mid_ik_anim.rotateOrder", 5)
+	    cmds.setAttr("chest_ik_anim.rotateOrder", 5)
+	cmds.setAttr("head_fk_anim.rotateOrder", 5)
+
+	
+	for control in ["neck_01_fk_anim", "neck_02_fk_anim", "neck_03_fk_anim"]:
+	    if cmds.objExists(control):
+		cmds.setAttr(control + ".rotateOrder", 5)
+		
+	for control in ["spine_01_anim", "spine_02_anim", "spine_03_anim", "spine_04_anim", "spine_05_anim"]:
+	    if cmds.objExists(control):
+		cmds.setAttr(control + ".rotateOrder", 5)
+		
+	
+	
+    # # # # # # # # # # # # # # # # # # # # # # # # # # # # # # # # # # # # # # # # # # # # # # # # # # # # # # # # # # # # # # # 
+    # # # # # # # # # # # # # # # # # # # # # # # # # # # # # # # # # # # # # # # # # # # # # # # # # # # # # # # # # # # # # # # 
+    # # # # # # # # # # # # # # # # # # # # # # # # # # # # # # # # # # # # # # # # # # # # # # # # # # # # # # # # # # # # # # # 
+    # # # # # # # # # # # # # # # # # # # # # # # # # # # # # # # # # # # # # # # # # # # # # # # # # # # # # # # # # # # # # # # 
+    # # # # # # # # # # # # # # # # # # # # # # # # # # # # # # # # # # # # # # # # # # # # # # # # # # # # # # # # # # # # # # # 
+    def addSpaces(self):
+	
+	cmds.select("*_space_switcher_follow")
+	nodes = cmds.ls(sl = True)
+	spaceSwitchers = []
+	for node in nodes:
+	    if node.find("invis") != 0:
+		spaceSwitchers.append(node)
+		
+	
+	for node in spaceSwitchers:
+	    #create a 'world' locator to constrain to
+	    worldLoc = cmds.spaceLocator(name = node + "_world_pos")[0]
+	    cmds.setAttr(worldLoc + ".v", 0)
+
+	    #position world loc to be in same place as node
+	    constraint = cmds.parentConstraint(node, worldLoc)[0]
+	    cmds.delete(constraint)
+	    
+	    #add the constraint between worldLoc and node
+	    if node == "spine_01_space_switcher_follow":
+		constraint = cmds.orientConstraint(worldLoc, node)[0]
+		
+	    else:
+		constraint = cmds.parentConstraint(worldLoc, node)[0]
+	    
+	    #add the attr to the space switcher node for that space
+	    spaceSwitchNode = node.rpartition("_follow")[0]
+	    cmds.select(spaceSwitchNode)
+	    cmds.addAttr(ln = "space_world", minValue = 0, maxValue = 1, dv = 0, keyable = True)
+
+	    
+	    #connect that attr to the constraint
+	    cmds.connectAttr(spaceSwitchNode + ".space_world", constraint + "." + worldLoc + "W0")
+	    
+	    #parent worldLoc under the offset_anim
+	    if worldLoc.find("master_anim") == 0:
+		cmds.parent(worldLoc, "rig_grp")
+		
+	    else:
+		cmds.parent(worldLoc, "offset_anim")
+		
+		
+		
+	
+	#SETUP SPECIAL CASES
+	for node in spaceSwitchers:
+	    if node == "chest_ik_anim_space_switcher_follow":
+		
+		#create a locator named world aligned
+		spaceLoc = cmds.spaceLocator(name = "chest_ik_world_aligned")[0]
+		cmds.setAttr(spaceLoc + ".v",0)
+		
+		#constrain it to the chest ik anim
+		constraint = cmds.pointConstraint("chest_ik_anim", spaceLoc)[0]
+		cmds.delete(constraint)
+		
+		#duplicate the locator 
+		worldOrientLoc = cmds.duplicate(spaceLoc, name = "chest_ik_world_orient")[0]
+		
+		#orient constrain the space loc to the world orient loc
+		cmds.orientConstraint(worldOrientLoc, spaceLoc, mo = True)
+		
+		#parent the space loc under the hip anim
+		cmds.parent(spaceLoc, "body_anim")
+		
+		#parent the worldOrientLoc under the master anim
+		cmds.parent(worldOrientLoc, "master_anim")
+		
+		#add attr to the space switcher node
+		spaceSwitchNode = node.rpartition("_follow")[0]
+		cmds.select(spaceSwitchNode)
+		cmds.addAttr(ln = "space_chest_ik_world_aligned", minValue = 0, maxValue = 1, dv = 0, keyable = True)
+		
+		#add constraint to the new object on the follow node
+		constraint = cmds.parentConstraint(spaceLoc, node, mo = True)[0]
+		
+		#hook up connections
+		targets = cmds.parentConstraint(constraint, q = True, targetList = True)
+		weight = 0
+		for i in range(len(targets)):
+		    if targets[i].find(spaceLoc) != -1:
+			weight = i
+			
+		cmds.connectAttr(spaceSwitchNode + ".space_chest_ik_world_aligned", constraint + "." + spaceLoc + "W" + str(weight))
+
+		
+		
+	    if node == "ik_wrist_l_anim_space_switcher_follow":
+		
+		spaceList = ["body_anim", "head_fk_anim"]
+		
+		if cmds.objExists("chest_ik_anim"):
+		    spaceList.append("chest_ik_anim")
+		for spaceObj in spaceList:
+		    spaceSwitchNode = node.rpartition("_follow")[0]
+		    
+		    
+		    #add attr to the space switcher node
+		    cmds.select(spaceSwitchNode)
+		    cmds.addAttr(ln = "space_" + spaceObj, minValue = 0, maxValue = 1, dv = 0, keyable = True)
+		    
+		    #add constraint to the new object on the follow node
+		    constraint = cmds.parentConstraint(spaceObj, node, mo = True)[0]
+		    
+		    #hook up connections
+		    targets = cmds.parentConstraint(constraint, q = True, targetList = True)
+		    weight = 0
+		    for i in range(len(targets)):
+			if targets[i].find(spaceObj) != -1:
+			    weight = i
+			    
+		    cmds.connectAttr(spaceSwitchNode + ".space_" + spaceObj, constraint + "." + spaceObj + "W" + str(weight))
+	    
+	    if node == "ik_wrist_r_anim_space_switcher_follow":
+		spaceList = ["body_anim", "head_fk_anim"]
+		if cmds.objExists("chest_ik_anim"):
+		    spaceList.append("chest_ik_anim")
+		for spaceObj in spaceList:
+		    spaceSwitchNode = node.rpartition("_follow")[0]
+		    
+		    
+		    #add attr to the space switcher node
+		    cmds.select(spaceSwitchNode)
+		    cmds.addAttr(ln = "space_" + spaceObj, minValue = 0, maxValue = 1, dv = 0, keyable = True)
+		    
+		    #add constraint to the new object on the follow node
+		    constraint = cmds.parentConstraint(spaceObj, node, mo = True)[0]
+		    
+		    #hook up connections
+		    targets = cmds.parentConstraint(constraint, q = True, targetList = True)
+		    weight = 0
+		    for i in range(len(targets)):
+			if targets[i].find(spaceObj) != -1:
+			    weight = i
+			    
+		    cmds.connectAttr(spaceSwitchNode + ".space_" + spaceObj, constraint + "." + spaceObj + "W" + str(weight))
+	    
+	    if node == "ik_elbow_l_anim_space_switcher_follow":
+		spaceList = ["body_anim"]
+		if cmds.objExists("chest_ik_anim"):
+		    spaceList.append("chest_ik_anim")
+		for spaceObj in spaceList:
+		    spaceSwitchNode = node.rpartition("_follow")[0]
+		    
+		    
+		    #add attr to the space switcher node
+		    cmds.select(spaceSwitchNode)
+		    cmds.addAttr(ln = "space_" + spaceObj, minValue = 0, maxValue = 1, dv = 0, keyable = True)
+		    
+		    #add constraint to the new object on the follow node
+		    constraint = cmds.parentConstraint(spaceObj, node, mo = True)[0]
+		    
+		    #hook up connections
+		    targets = cmds.parentConstraint(constraint, q = True, targetList = True)
+		    weight = 0
+		    for i in range(len(targets)):
+			if targets[i].find(spaceObj) != -1:
+			    weight = i
+			    
+		    cmds.connectAttr(spaceSwitchNode + ".space_" + spaceObj, constraint + "." + spaceObj + "W" + str(weight))
+	    
+	    if node == "ik_elbow_r_anim_space_switcher_follow":
+		spaceList = ["body_anim"]
+		if cmds.objExists("chest_ik_anim"):
+		    spaceList.append("chest_ik_anim")
+		for spaceObj in spaceList:
+		    spaceSwitchNode = node.rpartition("_follow")[0]
+		    
+		    
+		    #add attr to the space switcher node
+		    cmds.select(spaceSwitchNode)
+		    cmds.addAttr(ln = "space_" + spaceObj, minValue = 0, maxValue = 1, dv = 0, keyable = True)
+		    
+		    #add constraint to the new object on the follow node
+		    constraint = cmds.parentConstraint(spaceObj, node, mo = True)[0]
+		    
+		    #hook up connections
+		    targets = cmds.parentConstraint(constraint, q = True, targetList = True)
+		    weight = 0
+		    for i in range(int(len(targets))):
+			if targets[i].find(spaceObj) != -1:
+			    weight = i
+			    
+		    cmds.connectAttr(spaceSwitchNode + ".space_" + spaceObj, constraint + "." + spaceObj + "W" + str(weight))
+	    
+	    if node == "ik_foot_anim_l_space_switcher_follow":
+		for spaceObj in ["body_anim"]:
+		    spaceSwitchNode = node.rpartition("_follow")[0]
+		    
+		    
+		    #add attr to the space switcher node
+		    cmds.select(spaceSwitchNode)
+		    cmds.addAttr(ln = "space_" + spaceObj, minValue = 0, maxValue = 1, dv = 0, keyable = True)
+		    
+		    #add constraint to the new object on the follow node
+		    constraint = cmds.parentConstraint(spaceObj, node, mo = True)[0]
+		    
+		    #hook up connections
+		    targets = cmds.parentConstraint(constraint, q = True, targetList = True)
+		    weight = 0
+		    for i in range(int(len(targets))):
+			if targets[i].find(spaceObj) != -1:
+			    weight = i
+			    
+		    cmds.connectAttr(spaceSwitchNode + ".space_" + spaceObj, constraint + "." + spaceObj + "W" + str(weight))
+	    
+	    if node == "ik_foot_anim_r_space_switcher_follow":
+		for spaceObj in ["body_anim"]:
+		    spaceSwitchNode = node.rpartition("_follow")[0]
+		    
+		    
+		    #add attr to the space switcher node
+		    cmds.select(spaceSwitchNode)
+		    cmds.addAttr(ln = "space_" + spaceObj, minValue = 0, maxValue = 1, dv = 0, keyable = True)
+		    
+		    #add constraint to the new object on the follow node
+		    constraint = cmds.parentConstraint(spaceObj, node, mo = True)[0]
+		    
+		    #hook up connections
+		    targets = cmds.parentConstraint(constraint, q = True, targetList = True)
+		    weight = 0
+		    for i in range(int(len(targets))):
+			if targets[i].find(spaceObj) != -1:
+			    weight = i
+			    
+		    cmds.connectAttr(spaceSwitchNode + ".space_" + spaceObj, constraint + "." + spaceObj + "W" + str(weight))
+	    
+	    
+	
+	
+    # # # # # # # # # # # # # # # # # # # # # # # # # # # # # # # # # # # # # # # # # # # # # # # # # # # # # # # # # # # # # # # 
+    # # # # # # # # # # # # # # # # # # # # # # # # # # # # # # # # # # # # # # # # # # # # # # # # # # # # # # # # # # # # # # # 
+    # # # # # # # # # # # # # # # # # # # # # # # # # # # # # # # # # # # # # # # # # # # # # # # # # # # # # # # # # # # # # # # 
+    # # # # # # # # # # # # # # # # # # # # # # # # # # # # # # # # # # # # # # # # # # # # # # # # # # # # # # # # # # # # # # # 
+    # # # # # # # # # # # # # # # # # # # # # # # # # # # # # # # # # # # # # # # # # # # # # # # # # # # # # # # # # # # # # # # 
+    def buildCoreComponents(self):
+        
+        #builds the master, the root, the hips/body
+        
+        #BODY CONTROL
+        self.buildHips()
+	
+	#MASTER CONTROL
+	masterControl = self.createControl("circle", 150, "master_anim")
+	
+	constraint = cmds.pointConstraint("root", masterControl)[0]
+	cmds.delete(constraint)
+	
+	cmds.makeIdentity(masterControl, apply = True)
+	cmds.setAttr(masterControl + ".overrideEnabled", 1)
+	cmds.setAttr(masterControl + ".overrideColor", 18)
+	
+	spaceSwitchFollow = cmds.group(empty = True, name = masterControl + "_space_switcher_follow")
+	constraint = cmds.parentConstraint("root", spaceSwitchFollow)[0]
+	cmds.delete(constraint)
+	
+	spaceSwitcher = cmds.group(empty = True, name = masterControl + "_space_switcher")
+	constraint = cmds.parentConstraint("root", spaceSwitcher)[0]
+	cmds.delete(constraint)	
+	cmds.parent(spaceSwitcher, spaceSwitchFollow)
+	cmds.parent(masterControl, spaceSwitcher)
+	cmds.makeIdentity(masterControl, apply = True)
+	
+	
+	
+	#OFFSET CONTROL
+	offsetControl = self.createControl("square", 140, "offset_anim")
+	constraint = cmds.pointConstraint("root", offsetControl)[0]
+	cmds.delete(constraint)
+    
+	cmds.parent(offsetControl, masterControl)
+	cmds.makeIdentity(offsetControl, apply = True)
+	cmds.setAttr(offsetControl + ".overrideEnabled", 1)
+	cmds.setAttr(offsetControl + ".overrideColor", 17)
+	
+	#ROOT ANIM
+	rootControl = self.createControl("sphere", 10, "root_anim")
+	constraint = cmds.parentConstraint("driver_root", rootControl)[0]
+	cmds.delete(constraint)
+	cmds.parent(rootControl, masterControl)
+	cmds.makeIdentity(rootControl, apply = True)
+	cmds.parentConstraint(rootControl, "driver_root")
+	cmds.setAttr(rootControl + ".overrideEnabled", 1)
+	cmds.setAttr(rootControl + ".overrideColor", 30)
+	
+	for attr in [".sx", ".sy", ".sz", ".v"]:
+	    cmds.setAttr(masterControl + attr, lock = True, keyable = False)
+	    cmds.setAttr(offsetControl + attr, lock = True, keyable = False)
+	    cmds.setAttr(rootControl + attr, lock = True, keyable = False)
+        
+    # # # # # # # # # # # # # # # # # # # # # # # # # # # # # # # # # # # # # # # # # # # # # # # # # # # # # # # # # # # # # # # 
+    # # # # # # # # # # # # # # # # # # # # # # # # # # # # # # # # # # # # # # # # # # # # # # # # # # # # # # # # # # # # # # # 
+    # # # # # # # # # # # # # # # # # # # # # # # # # # # # # # # # # # # # # # # # # # # # # # # # # # # # # # # # # # # # # # # 
+    # # # # # # # # # # # # # # # # # # # # # # # # # # # # # # # # # # # # # # # # # # # # # # # # # # # # # # # # # # # # # # # 
+    # # # # # # # # # # # # # # # # # # # # # # # # # # # # # # # # # # # # # # # # # # # # # # # # # # # # # # # # # # # # # # # 
+    def buildHips(self):
+        
+        #create the grp and position and orient it correctly
+        bodyGrp = cmds.group(empty = True, name = "body_anim_grp")
+        bodyCtrl = self.createControl("square", 100, "body_anim")
+        constraint = cmds.parentConstraint("pelvis", bodyGrp)[0]
+        cmds.delete(constraint)
+        
+	#world alignment
+	for attr in [".rx", ".ry", ".rz"]:
+	    print cmds.getAttr(bodyGrp + attr)
+	    
+	    if cmds.getAttr(bodyGrp + attr) < 45:
+		if cmds.getAttr(bodyGrp + attr) > 0:
+		    cmds.setAttr(bodyGrp + attr, 0)
+		
+	    if cmds.getAttr(bodyGrp + attr) >= 80:
+		if cmds.getAttr(bodyGrp + attr) < 90:
+		    cmds.setAttr(bodyGrp + attr, 90)
+		    
+	    if cmds.getAttr(bodyGrp + attr) > 90:
+		if cmds.getAttr(bodyGrp + attr) < 100:
+		    cmds.setAttr(bodyGrp + attr, 90)		
+
+	    if cmds.getAttr(bodyGrp + attr) <= -80:
+		if cmds.getAttr(bodyGrp + attr) > -90:
+		    cmds.setAttr(bodyGrp + attr, -90)
+	    
+		
+	    if cmds.getAttr(bodyGrp + attr) > -90:
+		if cmds.getAttr(bodyGrp + attr) < -100:
+		    cmds.setAttr(bodyGrp + attr, -90) 
+			
+	for attr in [".rx", ".ry", ".rz"]:
+	    print cmds.getAttr(bodyGrp + attr)
+	    
+        #create space switcher
+        spaceSwitcherFollow = cmds.duplicate(bodyGrp, name = "body_anim_space_switcher_follow")[0]
+	spaceSwitcher = cmds.duplicate(bodyGrp, name = "body_anim_space_switcher")[0]
+	cmds.parent(spaceSwitcher, spaceSwitcherFollow)
+        cmds.parent(bodyGrp, spaceSwitcher)
+        
+
+        #create temp duplicate and orient control to joint
+        tempDupe = cmds.duplicate(bodyCtrl)[0]
+        constraint = cmds.parentConstraint("pelvis", bodyCtrl)[0]
+        cmds.delete(constraint)
+        
+      
+        #parent control to grp
+        cmds.parent(bodyCtrl, bodyGrp)
+        constraint = cmds.orientConstraint(tempDupe, bodyCtrl)[0]
+        cmds.delete(constraint)
+        cmds.makeIdentity(bodyCtrl, t = 1, r = 1, s = 1, apply = True)
+        
+
+        #clean up body control creation
+        cmds.delete(tempDupe)
+        
+        #set control color
+        cmds.setAttr(bodyCtrl + ".overrideEnabled", 1)
+        cmds.setAttr(bodyCtrl + ".overrideColor", 17)
+        
+        #lock attrs
+        for attr in [".sx", ".sy", ".sz", ".v"]:
+            cmds.setAttr(bodyCtrl + attr, lock = True, keyable = False)
+            
+        #build pelvis
+        self.buildPelvisControl()
+            
+            
+    # # # # # # # # # # # # # # # # # # # # # # # # # # # # # # # # # # # # # # # # # # # # # # # # # # # # # # # # # # # # # # # 
+    # # # # # # # # # # # # # # # # # # # # # # # # # # # # # # # # # # # # # # # # # # # # # # # # # # # # # # # # # # # # # # # 
+    # # # # # # # # # # # # # # # # # # # # # # # # # # # # # # # # # # # # # # # # # # # # # # # # # # # # # # # # # # # # # # # 
+    # # # # # # # # # # # # # # # # # # # # # # # # # # # # # # # # # # # # # # # # # # # # # # # # # # # # # # # # # # # # # # # 
+    # # # # # # # # # # # # # # # # # # # # # # # # # # # # # # # # # # # # # # # # # # # # # # # # # # # # # # # # # # # # # # # 
+    def buildPelvisControl(self):
+        
+        #create the grp and position and orient it correctly
+        hipGrp = cmds.group(empty = True, name = "hip_anim_grp")
+        hipCtrl = self.createControl("circle", 60, "hip_anim")
+        constraint = cmds.parentConstraint("pelvis", hipGrp)[0]
+        cmds.delete(constraint)
+        
+        #create temp duplicate and orient control to joint
+        tempDupe = cmds.duplicate(hipCtrl)[0]
+        constraint = cmds.parentConstraint("pelvis", hipCtrl)[0]
+        cmds.delete(constraint)
+        
+        #parent control to grp
+        cmds.parent(hipCtrl, hipGrp)
+        constraint = cmds.orientConstraint(tempDupe, hipCtrl)[0]
+        cmds.delete(constraint)
+        cmds.makeIdentity(hipCtrl, t = 1, r = 1, s = 1, apply = True)
+        
+        #parent the grp to the body anim
+        cmds.parent(hipGrp, "body_anim")
+        
+        
+        #clean up body control creation
+        cmds.delete(tempDupe)
+        
+        #set control color
+        cmds.setAttr(hipCtrl + ".overrideEnabled", 1)
+        cmds.setAttr(hipCtrl + ".overrideColor", 18)
+        
+        #lock attrs
+        for attr in [".sx", ".sy", ".sz", ".v"]:
+            cmds.setAttr(hipCtrl + attr, lock = True, keyable = False)
+            
+
+        
+
+    # # # # # # # # # # # # # # # # # # # # # # # # # # # # # # # # # # # # # # # # # # # # # # # # # # # # # # # # # # # # # # # 
+    # # # # # # # # # # # # # # # # # # # # # # # # # # # # # # # # # # # # # # # # # # # # # # # # # # # # # # # # # # # # # # # 
+    # # # # # # # # # # # # # # # # # # # # # # # # # # # # # # # # # # # # # # # # # # # # # # # # # # # # # # # # # # # # # # # 
+    # # # # # # # # # # # # # # # # # # # # # # # # # # # # # # # # # # # # # # # # # # # # # # # # # # # # # # # # # # # # # # # 
+    # # # # # # # # # # # # # # # # # # # # # # # # # # # # # # # # # # # # # # # # # # # # # # # # # # # # # # # # # # # # # # # 
+    def buildFKSpine(self):
+        
+        #find the number of spine bones from the skeleton settings
+        spineJoints = self.getSpineJoints()
+        fkControls = []
+        parent = None
+        
+        for joint in spineJoints:
+            if joint == "spine_01":
+		#add space switcher node to base of spine
+		spaceSwitcherFollow = cmds.group(empty = True, name = joint + "_space_switcher_follow")
+		constraint = cmds.parentConstraint(joint, spaceSwitcherFollow)[0]
+		cmds.delete(constraint)
+		spaceSwitcher = cmds.duplicate(spaceSwitcherFollow, name = joint + "_space_switcher")[0]
+		cmds.parent(spaceSwitcher, spaceSwitcherFollow)
+		
+            #create an empty group in the same space as the joint
+            group = cmds.group(empty = True, name = joint + "_anim_grp")
+            constraint = cmds.parentConstraint(joint, group)[0]
+            cmds.delete(constraint)
+	    
+	    #create an additional layer of group that has zeroed attrs
+	    offsetGroup = cmds.group(empty = True, name = joint + "_anim_offset_grp")
+	    constraint = cmds.parentConstraint(joint, offsetGroup)[0]
+	    cmds.delete(constraint)
+	    cmds.parent(offsetGroup, group)
+            
+            #create a control object in the same space as the joint
+            control = self.createControl("circle", 45, joint + "_anim")
+            tempDupe = cmds.duplicate(control)[0]
+            constraint = cmds.parentConstraint(joint, control)[0]
+            cmds.delete(constraint)
+            fkControls.append(control)
+            
+            
+            #parent the control object to the group
+            cmds.parent(control, offsetGroup)
+            constraint = cmds.orientConstraint(tempDupe, control, skip = ["x", "z"])[0]
+            cmds.delete(constraint)
+            cmds.makeIdentity(control, t = 1, r = 1, s = 1, apply = True)
+
+            
+            #setup hierarchy
+            if parent != None:
+                cmds.parent(group, parent, absolute = True)
+                
+            else:
+                cmds.parent(group, spaceSwitcher)
+		cmds.parent(spaceSwitcherFollow, "body_anim")
+                
+                
+            #set the parent to be the current spine control
+            parent = control
+            
+            #clean up
+            cmds.delete(tempDupe)
+            
+            for attr in [".sx", ".sy", ".sz", ".v"]:
+                cmds.setAttr(control + attr, lock = True, keyable = False)
+                
+            #set the control's color
+            cmds.setAttr(control + ".overrideEnabled", 1)
+            cmds.setAttr(control + ".overrideColor", 18)
+            
+            
+            #create length attr on spine controls. need to find up axis for control first
+            upAxis = self.getUpAxis(control)
+            cmds.aliasAttr("length", control + ".translate" + upAxis)
+            
+
+
+            
+        return fkControls
+            
+    # # # # # # # # # # # # # # # # # # # # # # # # # # # # # # # # # # # # # # # # # # # # # # # # # # # # # # # # # # # # # # # 
+    # # # # # # # # # # # # # # # # # # # # # # # # # # # # # # # # # # # # # # # # # # # # # # # # # # # # # # # # # # # # # # # 
+    # # # # # # # # # # # # # # # # # # # # # # # # # # # # # # # # # # # # # # # # # # # # # # # # # # # # # # # # # # # # # # # 
+    # # # # # # # # # # # # # # # # # # # # # # # # # # # # # # # # # # # # # # # # # # # # # # # # # # # # # # # # # # # # # # # 
+    # # # # # # # # # # # # # # # # # # # # # # # # # # # # # # # # # # # # # # # # # # # # # # # # # # # # # # # # # # # # # # # 
+    def buildIKSpine(self, fkControls):
+        
+        
+        numSpineBones = cmds.getAttr("Skeleton_Settings.numSpineBones")
+	
+	if numSpineBones > 2:
+	    #duplicate the spine joints we'll need for the spline IK
+	    spineJoints = self.getSpineJoints()
+	    parent = None
+	    rigJoints = []
+	    
+	    for joint in spineJoints:
+		spineBone = cmds.duplicate(joint, parentOnly = True, name = "splineIK_" + joint)[0]
+		
+		if parent != None:
+		    cmds.parent(spineBone, parent)
+		    
+		else:
+		    cmds.parent(spineBone, world = True)
+		    
+		parent = spineBone
+		rigJoints.append(str(spineBone))
+		
+	    
+	    for joint in rigJoints:
+		twistJoint = cmds.duplicate(joint, name =  "twist_" + joint, parentOnly = True)[0]
+		cmds.parent(twistJoint, joint)
+	
+		
+		
+		
+	    #create the spline IK
+	    ikNodes = cmds.ikHandle(sj = str(rigJoints[0]), ee = str(rigJoints[len(rigJoints) - 1]), sol = "ikSplineSolver", createCurve = True, simplifyCurve = True, parentCurve = False, name = str(rigJoints[0]) + "_splineIK")   
+	    ikHandle = ikNodes[0]
+	    ikCurve = ikNodes[2]
+	    ikCurve = cmds.rename(ikCurve, "spine_splineIK_curve")
+	    cmds.setAttr(ikCurve + ".inheritsTransform", 0)
+	    cmds.setAttr(ikHandle + ".v", 0)
+	    cmds.setAttr(ikCurve + ".v", 0)
+	    
+	    #create the three joints to skin the curve to
+	    botJoint = cmds.duplicate(rigJoints[0], name = "spine_splineIK_bottom_joint", parentOnly = True)[0]
+	    topJoint = cmds.duplicate(rigJoints[len(rigJoints) - 1], name = "spine_splineIK_top_joint", parentOnly = True)[0]
+	    midJoint = cmds.duplicate(topJoint, name = "spine_splineIK_mid_joint", parentOnly = True)[0]
+	    
+	    cmds.parent([botJoint, topJoint,midJoint], world = True)
+	    
+	    constraint = cmds.pointConstraint([botJoint, topJoint], midJoint)[0]
+	    cmds.delete(constraint)
+	    
+	    #skin the joints to the curve
+	    cmds.select([botJoint, topJoint, midJoint])
+	    skin = cmds.skinCluster( [botJoint, topJoint, midJoint], ikCurve, toSelectedBones = True )[0]
+	    
+	    #skin weight the curve
+	    curveShape = cmds.listRelatives(ikCurve, shapes = True)[0]
+	    numSpans = cmds.getAttr(curveShape + ".spans")
+	    degree = cmds.getAttr(curveShape + ".degree")
+	    numCVs = numSpans + degree
+	    
+	    #this should always be the case, but just to be safe
+	    if numCVs == 4:
+		cmds.skinPercent(skin, ikCurve + ".cv[0]", transformValue = [(botJoint, 1.0)])
+		cmds.skinPercent(skin, ikCurve + ".cv[1]", transformValue = [(botJoint, 0.5), (midJoint, 0.5)])
+		cmds.skinPercent(skin, ikCurve + ".cv[2]", transformValue = [(midJoint, 0.5), (topJoint, 0.5)])
+		cmds.skinPercent(skin, ikCurve + ".cv[3]", transformValue = [(topJoint, 1.0)])
+
+	    #create the controls
+	    
+	    #TOP CTRL
+	    topCtrl = self.createControl("circle", 50, "chest_ik_anim")
+	    
+	    #set the control's color
+	    cmds.setAttr(topCtrl + ".overrideEnabled", 1)
+	    cmds.setAttr(topCtrl + ".overrideColor", 17)
+	    
+	    #position the control
+	    constraint = cmds.pointConstraint(topJoint, topCtrl)[0]
+	    cmds.delete(constraint)
+	    
+	    #create the control grp
+	    topCtrlGrp = cmds.group(empty = True, name = topCtrl + "_grp")
+	    constraint = cmds.parentConstraint(topJoint, topCtrlGrp)[0]
+	    cmds.delete(constraint)
+	    
+		    
+	    #create the top control driver group
+	    topCtrlDriver = cmds.duplicate(topCtrlGrp, name = "chest_ik_anim_driver_grp")
+	    
+	    #create the space switcher group
+	    spaceSwitcherFollow = cmds.group(empty = True, name = topCtrl + "_space_switcher_follow")
+	    constraint = cmds.parentConstraint(topCtrlGrp, spaceSwitcherFollow)[0]
+	    cmds.delete(constraint)
+	    spaceSwitcher = cmds.duplicate(spaceSwitcherFollow, parentOnly = True, name = topCtrl + "_space_switcher")[0]
+	    
+	    
+	    #parent objects
+	    cmds.parent(spaceSwitcher, spaceSwitcherFollow)
+	    cmds.parent(topCtrlGrp, spaceSwitcher)
+	    cmds.parent(topCtrlDriver, topCtrlGrp)
+	    cmds.parent(topCtrl, topCtrlDriver)
+	    cmds.makeIdentity(topCtrl, t = 1, r = 1, s = 1, apply = True)
+	    cmds.parent(topJoint, topCtrl)
+	    
+	    
+	    
+	    
+	    #MID CTRL
+	    midCtrl = self.createControl("circle", 45, "mid_ik_anim")
+	    
+	    #set the control's color
+	    cmds.setAttr(midCtrl + ".overrideEnabled", 1)
+	    cmds.setAttr(midCtrl + ".overrideColor", 18)
+	    
+	    #position the control
+	    constraint = cmds.pointConstraint(midJoint, midCtrl)[0]
+	    cmds.delete(constraint)
+	    
+	    #create the control grp
+	    midCtrlGrp = cmds.group(empty = True, name = midCtrl + "_grp")
+	    constraint = cmds.parentConstraint(midJoint, midCtrlGrp)[0]
+	    cmds.delete(constraint)
+	    
+	    
+	    #mid control driver grp
+	    midCtrlDriver = cmds.duplicate(midCtrlGrp, name = "mid_ik_anim_driver_grp")
+	    midCtrlTranslateDriver = cmds.duplicate(midCtrlGrp, name = "mid_ik_anim_translate_driver_grp")
+	    
+	    #parent objects
+	    cmds.parent(midCtrl, midCtrlDriver)
+	    cmds.parent(midCtrlDriver, midCtrlTranslateDriver)
+	    cmds.parent(midCtrlTranslateDriver, midCtrlGrp)
+	    cmds.makeIdentity(midCtrl, t = 1, r = 1, s = 1, apply = True)
+	    cmds.parent(midJoint, midCtrl)
+	    cmds.parent(botJoint, "hip_anim")
+	    
+	    
+	    
+	    
+	    #ADDING STRETCH
+	    
+	    #add the attr to the top ctrl
+	    cmds.select(topCtrl)
+	    cmds.addAttr(longName='stretch', defaultValue=0, minValue=0, maxValue=1, keyable = True)
+	    cmds.addAttr(longName='squash', defaultValue=0, minValue=0, maxValue=1, keyable = True)
+	    
+	    
+	    #create the curveInfo node#find 
+	    cmds.select(ikCurve)
+	    curveInfoNode = cmds.arclen(cmds.ls(sl = True), ch = True )
+	    originalLength = cmds.getAttr(curveInfoNode + ".arcLength")
+	    
+	    #create the multiply/divide node that will get the scale factor
+	    divideNode = cmds.shadingNode("multiplyDivide", asUtility = True)
+	    divideNode_Inverse = cmds.shadingNode("multiplyDivide", asUtility = True)
+	    cmds.setAttr(divideNode + ".operation", 2)
+	    cmds.setAttr(divideNode + ".input2X", originalLength)
+	    cmds.setAttr(divideNode_Inverse + ".operation", 2)
+	    cmds.setAttr(divideNode_Inverse + ".input1X", originalLength)	
+	    
+	    #create the blendcolors node
+	    blenderNode = cmds.shadingNode("blendColors", asUtility = True)
+	    cmds.setAttr(blenderNode + ".color2R", 1)
+	    
+	    blenderNode_Inverse = cmds.shadingNode("blendColors", asUtility = True)
+	    cmds.setAttr(blenderNode_Inverse + ".color2R", 1)	    
+	    
+	    #connect attrs
+	    cmds.connectAttr(curveInfoNode + ".arcLength", divideNode + ".input1X")
+	    cmds.connectAttr(curveInfoNode + ".arcLength", divideNode_Inverse + ".input2X")
+	    cmds.connectAttr(divideNode + ".outputX", blenderNode + ".color1R")
+	    cmds.connectAttr(divideNode_Inverse + ".outputX", blenderNode_Inverse + ".color1R")
+	    
+	    cmds.connectAttr(topCtrl + ".stretch", blenderNode + ".blender")
+	    cmds.connectAttr(topCtrl + ".squash", blenderNode_Inverse + ".blender")
+	    upAxis = self.getUpAxis(topCtrl)
+	    if upAxis == "X":
+		axisB = "Y"
+		axisC = "Z"
+	    if upAxis == "Y":
+		axisB = "X"
+		axisC = "Z"
+	    if upAxis == "Z":
+		axisB = "X"
+		axisC = "Y"
+		
+	    for i in range(len(rigJoints) - 2):
+		children = cmds.listRelatives(rigJoints[i], children = True)
+		for child in children:
+		    if child.find("twist") != -1:
+			twistJoint = child
+		    
+		cmds.connectAttr(blenderNode_Inverse + ".outputR", twistJoint + ".scale" + axisB)
+		cmds.connectAttr(blenderNode_Inverse + ".outputR", twistJoint + ".scale" + axisC)
+		
+	    cmds.connectAttr(blenderNode + ".outputR", rigJoints[0] + ".scale" + upAxis)
+	    
+	    #add twist amount attrs and setup
+	    cmds.select(topCtrl)
+	    cmds.addAttr(longName='twist_amount', defaultValue=1, minValue=0, keyable = True)
+    
+	    #find number of spine joints and divide 1 by numSpineJoints
+	    num = len(spineJoints)
+	    val = 1.0/float(num)
+	    twistamount = val
+	    
+	    locGrp = cmds.group(empty = True, name = "spineIK_twist_grp")
+	    cmds.parent(locGrp, "body_anim")
+	    
+	    for i in range(int(num - 1)):
+		
+		#create a locator that will be orient constrained between the body and chest
+		locator = cmds.spaceLocator(name = spineJoints[i] + "_twistLoc")[0]
+		group = cmds.group(empty = True, name = spineJoints[i] + "_twistLocGrp")
+		constraint = cmds.parentConstraint(spineJoints[i], locator)[0]
+		cmds.delete(constraint)
+		constraint = cmds.parentConstraint(spineJoints[i], group)[0]
+		cmds.delete(constraint)	
+		cmds.parent(locator, group)
+		cmds.parent(group, locGrp)
+		cmds.setAttr(locator + ".v", 0, lock = True)
+		
+		
+		#duplicate the locator and parent it under the group. This will be the locator that takes the rotation x twist amount and gives us the final value
+		orientLoc = cmds.duplicate(locator, name = spineJoints[i] + "_orientLoc")[0]
+		
+		#create constraints between body/chest
+		constraint = cmds.orientConstraint(["body_anim", topCtrl], locator)[0]
+		
+		#set weights on constraint
+		firstValue = 1 - twistamount
+		secondValue = 1 - firstValue
+		
+		cmds.setAttr(constraint + ".body_animW0", firstValue)
+		cmds.setAttr(constraint + "." + topCtrl + "W1", secondValue)
+		
+		#factor in twist amount 
+		twistMultNode = cmds.shadingNode("multiplyDivide", asUtility = True, name = spineJoints[i] + "_twist_amount")
+		
+		#expose the twistAmount on the control as an attr
+		cmds.connectAttr(topCtrl + ".twist_amount", twistMultNode + ".input2X")
+		cmds.connectAttr(topCtrl + ".twist_amount", twistMultNode + ".input2Y")
+		cmds.connectAttr(topCtrl + ".twist_amount", twistMultNode + ".input2Z")
+		cmds.connectAttr(locator + ".rotate", twistMultNode + ".input1")
+		
+		cmds.connectAttr(twistMultNode + ".output", orientLoc + ".rotate")
+		
+		#constrain the spine joint to the orientLoc
+		if upAxis == "X":
+		    skipped = ["y", "z"]
+		if upAxis == "Y":
+		    skipped = ["x", "z"]
+		if upAxis == "Z":
+		    skipped = ["x", "y"]
+		    
+		
+		cmds.orientConstraint(orientLoc, "twist_splineIK_" + spineJoints[i], skip = skipped)
+		
+		twistamount = twistamount + val
+		
+	    
+	    
+	    #parent the components to the body anim
+	    cmds.parent(midCtrlGrp, "body_anim")
+	    cmds.parent(midCtrl, world = True)
+	    cmds.parent(midJoint, world = True)
+	    
+	    for attr in [".rx", ".ry", ".rz"]:
+		cmds.setAttr(midCtrlGrp + attr, 0)
+		
+	    cmds.parent(midCtrl, midCtrlDriver)
+	    cmds.makeIdentity(midCtrl, t = 1, r = 1, s = 0, apply = True)
+	    cmds.parent(midJoint, midCtrl)
+	    
+
+	    cmds.parent(spaceSwitcherFollow, "body_anim")
+	    cmds.parent(rigJoints[0], "body_anim")
+		
+	    #world alignment
+	    cmds.parent(topCtrl, world = True)
+	    cmds.parent(topJoint, world = True)
+	    for attr in [".rx", ".ry", ".rz"]:
+		if cmds.getAttr(spaceSwitcherFollow + attr) < 45:
+		    if cmds.getAttr(spaceSwitcherFollow + attr) > 0:
+			cmds.setAttr(spaceSwitcherFollow + attr, 0)
+		    
+		if cmds.getAttr(spaceSwitcherFollow + attr) >= 80:
+		    if cmds.getAttr(spaceSwitcherFollow + attr) < 90:
+			cmds.setAttr(spaceSwitcherFollow + attr, 90)
+			
+		if cmds.getAttr(spaceSwitcherFollow + attr) > 90:
+		    if cmds.getAttr(spaceSwitcherFollow + attr) < 100:
+			cmds.setAttr(spaceSwitcherFollow + attr, 90)		
+
+		if cmds.getAttr(spaceSwitcherFollow + attr) <= -80:
+		    if cmds.getAttr(spaceSwitcherFollow + attr) > -90:
+			cmds.setAttr(spaceSwitcherFollow + attr, -90)
+		
+		    
+		if cmds.getAttr(spaceSwitcherFollow + attr) > -90:
+		    if cmds.getAttr(spaceSwitcherFollow + attr) < -100:
+			cmds.setAttr(spaceSwitcherFollow + attr, -90)	    
+		
+	    cmds.parent(topCtrl, topCtrlDriver)
+	    cmds.makeIdentity(topCtrl, t = 1, r = 1, s = 0, apply = True)
+	    cmds.parent(topJoint, topCtrl)
+	    
+	    
+	    #hookup spine driver joints
+	    driverJoints = []
+	    for joint in rigJoints:
+		driverJoint = joint.partition("splineIK_")[2]
+		driverJoint = "driver_" + driverJoint
+		driverJoints.append(driverJoint)
+		
+		
+	    #hookup spine to driver
+	    self.hookupSpine(rigJoints, fkControls)
+	    
+	    #control driver joints
+	    children = cmds.listRelatives(rigJoints[len(rigJoints) -1], children = True)
+	    for child in children:
+		if child.find("twist") != -1:
+		    twistJoint = child
+	    
+	    topSpineJointConstraint = cmds.pointConstraint(topJoint, twistJoint, mo = True)[0]
+		    
+		    
+	    #connect attr on top spine joint constraint
+	    target = cmds.pointConstraint(topSpineJointConstraint, q = True, weightAliasList = True)[0]
+	    cmds.connectAttr(topCtrl + ".stretch", topSpineJointConstraint + "." + target)
+	    
+	    
+	    #create stretch meter attr
+	    cmds.select(topCtrl)
+	    cmds.addAttr(longName='stretchFactor',keyable = True)
+	    cmds.connectAttr(divideNode + ".outputX", topCtrl + ".stretchFactor")
+	    cmds.setAttr(topCtrl + ".stretchFactor", lock = True)
+	    
+	    
+	    cmds.select(midCtrl)
+	    cmds.addAttr(longName='stretchFactor',keyable = True)
+	    cmds.connectAttr(topCtrl + ".stretchFactor", midCtrl + ".stretchFactor")
+	    cmds.setAttr(midCtrl + ".stretchFactor", lock = True)
+	    
+	    
+	    #lock and hide attrs that should not be keyable
+	    for control in [topCtrl, midCtrl]:
+		for attr in [".sx", ".sy", ".sz", ".v"]:
+		    cmds.setAttr(control + attr, keyable = False, lock = True)
+		    
+	    
+	    
+
+
+    # # # # # # # # # # # # # # # # # # # # # # # # # # # # # # # # # # # # # # # # # # # # # # # # # # # # # # # # # # # # # # # 
+    # # # # # # # # # # # # # # # # # # # # # # # # # # # # # # # # # # # # # # # # # # # # # # # # # # # # # # # # # # # # # # # 
+    # # # # # # # # # # # # # # # # # # # # # # # # # # # # # # # # # # # # # # # # # # # # # # # # # # # # # # # # # # # # # # # 
+    # # # # # # # # # # # # # # # # # # # # # # # # # # # # # # # # # # # # # # # # # # # # # # # # # # # # # # # # # # # # # # # 
+    # # # # # # # # # # # # # # # # # # # # # # # # # # # # # # # # # # # # # # # # # # # # # # # # # # # # # # # # # # # # # # # 
+    def buildFKLegs(self):
+        
+        #need to create the leg joints for each side based on the driver thigh, calf, and foot
+        
+        for side in ["l", "r"]:
+            
+            ball = False
+            
+            #create joints
+            fkThighJoint = cmds.duplicate("driver_thigh_" + side, name = "fk_leg_thigh_" + side, parentOnly = True)[0]
+            fkCalfJoint = cmds.duplicate("driver_calf_" + side, name = "fk_leg_calf_" + side, parentOnly = True)[0]
+            fkFootJoint = cmds.duplicate("driver_foot_" + side, name = "fk_leg_foot_" + side, parentOnly = True)[0]
+            
+            if cmds.objExists("driver_ball_" + side):
+                ball = True
+                fkBallJoint = cmds.duplicate("driver_ball_" + side, name = "fk_leg_ball_" + side, parentOnly = True)[0]
+                
+            for joint in [fkThighJoint, fkCalfJoint, fkFootJoint]:
+                cmds.parent(joint, world = True)
+                
+            if ball:
+                cmds.parent(fkBallJoint, fkFootJoint)
+            
+            
+            cmds.parent(fkFootJoint, fkCalfJoint)
+            cmds.parent(fkCalfJoint, fkThighJoint)
+            cmds.makeIdentity(fkThighJoint, t = 0, r = 1, s = 0, apply = True)
+            
+            
+            #create controls for each joint
+            
+            
+            #THIGH
+            fkThighCtrl = self.createControl("circle", 30, "fk_thigh_" + side + "_anim")
+            cmds.setAttr(fkThighCtrl + ".ry", -90)
+            cmds.makeIdentity(fkThighCtrl, r = 1, apply =True)
+            
+            fkThighCtrlGrp = cmds.group(empty = True, name = "fk_thigh_" + side + "_anim_grp")
+            
+            
+            constraint = cmds.parentConstraint(fkThighJoint, fkThighCtrlGrp)[0]
+            cmds.delete(constraint)
+            constraint = cmds.parentConstraint(fkThighJoint, fkThighCtrl)[0]
+            cmds.delete(constraint)
+            
+            fkThighOrientGrp = cmds.duplicate(fkThighCtrlGrp, parentOnly = True, name = "fk_thigh_" + side + "_orient_grp")
+            fkThighWorldNode = cmds.duplicate(fkThighOrientGrp, parentOnly = True, name = "fk_thigh_" + side + "_world")
+            cmds.orientConstraint(fkThighWorldNode, fkThighOrientGrp)
+            cmds.parent(fkThighWorldNode, "body_anim")
+            cmds.parent(fkThighCtrl, fkThighCtrlGrp)
+            cmds.parent(fkThighCtrlGrp, fkThighOrientGrp)
+            
+            
+            
+            
+            #get the distance between the hip and knee
+            thighPos = cmds.xform("driver_thigh_" + side, q = True, ws = True, t = True)
+            kneePos = cmds.xform("driver_calf_" + side, q = True, ws = True, t = True)
+            dist = (thighPos[2] - kneePos[2]) / 2
+
+            
+            #move the ctrl to the position of dist
+            upAxis = self.getUpAxis(fkThighCtrl)
+            
+            if side == "l":
+                cmds.setAttr(fkThighCtrl + ".translate" + upAxis, dist * -1)
+                
+            else:
+                cmds.setAttr(fkThighCtrl + ".translate" + upAxis, dist)
+                
+            #get the pivot of the thigh and set the pivot of the ctrl to that position
+            piv = cmds.xform(fkThighJoint, q = True, ws = True, rotatePivot = True)
+            cmds.xform(fkThighCtrl, ws = True, piv = (piv[0], piv[1], piv[2]))
+            
+            
+            #lock attrs that should not be animated
+            cmds.setAttr(fkThighCtrl + ".tx", lock = True, keyable = False)
+            cmds.setAttr(fkThighCtrl + ".ty", lock = True, keyable = False)
+            cmds.setAttr(fkThighCtrl + ".tz", lock = True, keyable = False)
+            cmds.setAttr(fkThighCtrl + ".sx", lock = True, keyable = False)
+            cmds.setAttr(fkThighCtrl + ".sy", lock = True, keyable = False)
+            cmds.setAttr(fkThighCtrl + ".sz", lock = True, keyable = False)
+            cmds.setAttr(fkThighCtrl + ".v", lock = True, keyable = False)
+            
+            
+            
+            
+            
+            
+            #CALF
+            fkCalfCtrl = self.createControl("semiCircle", 5, "fk_calf_" + side + "_anim")
+            cmds.makeIdentity(fkCalfCtrl, s = 1, apply = True)
+            cmds.setAttr(fkCalfCtrl + ".sx", .5)
+            cmds.setAttr(fkCalfCtrl + ".sy", .75)
+            cmds.setAttr(fkCalfCtrl + ".rx", 180)
+            cmds.setAttr(fkCalfCtrl + ".ry", -90)
+            cmds.makeIdentity(fkCalfCtrl, s = 1, apply = True)
+            
+            fkCalfCtrlGrp = cmds.group(empty = True, name = "fk_calf_" + side + "_anim_grp")
+            
+            constraint = cmds.parentConstraint(fkCalfJoint, fkCalfCtrlGrp)[0]
+            cmds.delete(constraint)
+            constraint = cmds.pointConstraint(fkCalfJoint, fkCalfCtrl)[0]
+            cmds.delete(constraint)
+            
+            cmds.parent(fkCalfCtrl, fkCalfCtrlGrp)
+            
+
+    
+            #get the pivot of the calf and set the pivot of the ctrl to that position
+            piv = cmds.xform(fkCalfJoint, q = True, ws = True, rotatePivot = True)
+            cmds.xform(fkCalfCtrl, ws = True, piv = (piv[0], piv[1], piv[2]))
+            
+            #parent the fk ctrl grp to the thigh anim
+            cmds.parent(fkCalfCtrlGrp, fkThighCtrl)
+    
+            cmds.makeIdentity(fkCalfCtrl, r = 1, apply = True)
+            
+            #lock attrs that should not be animated
+            cmds.setAttr(fkCalfCtrl + ".tx", lock = True, keyable = False)
+            cmds.setAttr(fkCalfCtrl + ".ty", lock = True, keyable = False)
+            cmds.setAttr(fkCalfCtrl + ".tz", lock = True, keyable = False)
+            cmds.setAttr(fkCalfCtrl + ".sx", lock = True, keyable = False)
+            cmds.setAttr(fkCalfCtrl + ".sy", lock = True, keyable = False)
+            cmds.setAttr(fkCalfCtrl + ".sz", lock = True, keyable = False)
+            cmds.setAttr(fkCalfCtrl + ".v", lock = True, keyable = False)
+
+            
+            
+            
+            #FOOT
+            fkFootCtrl = self.createControl("circle", 17, "fk_foot_" + side + "_anim")
+            cmds.setAttr(fkFootCtrl + ".ry", -90)
+            cmds.makeIdentity(fkFootCtrl, r = 1, apply =True)
+            
+            fkFootCtrlGrp = cmds.group(empty = True, name = "fk_foot_" + side + "_anim_grp")
+            
+            constraint = cmds.parentConstraint(fkFootJoint, fkFootCtrlGrp)[0]
+            cmds.delete(constraint)
+            constraint = cmds.parentConstraint(fkFootJoint, fkFootCtrl)[0]
+            cmds.delete(constraint)
+            
+            cmds.parent(fkFootCtrl, fkFootCtrlGrp)
+            
+                
+            #get the pivot of the thigh and set the pivot of the ctrl to that position
+            piv = cmds.xform(fkFootJoint, q = True, ws = True, rotatePivot = True)
+            cmds.xform(fkFootCtrl, ws = True, piv = (piv[0], piv[1], piv[2]))
+            
+            
+            #parent the fk ctrl grp to the thigh anim
+            cmds.parent(fkFootCtrlGrp, fkCalfCtrl)
+            
+            #lock attrs that should not be animated
+            cmds.setAttr(fkFootCtrl + ".tx", lock = True, keyable = False)
+            cmds.setAttr(fkFootCtrl + ".ty", lock = True, keyable = False)
+            cmds.setAttr(fkFootCtrl + ".tz", lock = True, keyable = False)
+            cmds.setAttr(fkFootCtrl + ".sx", lock = True, keyable = False)
+            cmds.setAttr(fkFootCtrl + ".sy", lock = True, keyable = False)
+            cmds.setAttr(fkFootCtrl + ".sz", lock = True, keyable = False)
+            cmds.setAttr(fkFootCtrl + ".v", lock = True, keyable = False)
+            
+            
+            if ball:
+                #BALL
+                fkBallCtrl = self.createControl("arrowOnBall", 2, "fk_ball_" + side + "_anim")
+                
+                if side == "l":
+                    cmds.setAttr(fkBallCtrl + ".rx", -90)
+                    cmds.makeIdentity(fkBallCtrl, t = 1, r = 1, s = 1, apply = True)
+                    
+                else:
+                    cmds.setAttr(fkBallCtrl + ".rx", 90)
+                    cmds.makeIdentity(fkBallCtrl, t = 1, r = 1, s = 1, apply = True)
+                
+                
+                fkBallCtrlGrp = cmds.group(empty = True, name = "fk_ball_" + side + "_anim_grp")
+                
+                constraint = cmds.parentConstraint(fkBallJoint, fkBallCtrlGrp)[0]
+                cmds.delete(constraint)
+                constraint = cmds.parentConstraint(fkBallJoint, fkBallCtrl)[0]
+                cmds.delete(constraint)
+                
+                cmds.parent(fkBallCtrl, fkBallCtrlGrp)
+                
+                    
+                #get the pivot of the thigh and set the pivot of the ctrl to that position
+                piv = cmds.xform(fkBallJoint, q = True, ws = True, rotatePivot = True)
+                cmds.xform(fkBallCtrl, ws = True, piv = (piv[0], piv[1], piv[2]))
+                
+                
+                #parent the fk ctrl grp to the thigh anim
+                cmds.parent(fkBallCtrlGrp, fkFootCtrl)
+                
+
+                #lock attrs that should not be animated
+                cmds.setAttr(fkBallCtrl + ".tx", lock = True, keyable = False)
+                cmds.setAttr(fkBallCtrl + ".ty", lock = True, keyable = False)
+                cmds.setAttr(fkBallCtrl + ".tz", lock = True, keyable = False)
+                cmds.setAttr(fkBallCtrl + ".sx", lock = True, keyable = False)
+                cmds.setAttr(fkBallCtrl + ".sy", lock = True, keyable = False)
+                cmds.setAttr(fkBallCtrl + ".sz", lock = True, keyable = False)
+                cmds.setAttr(fkBallCtrl + ".v", lock = True, keyable = False)
+                
+            
+            #hook up leg joints to follow ctrls
+            cmds.orientConstraint(fkThighCtrl, fkThighJoint)
+            cmds.orientConstraint(fkCalfCtrl, fkCalfJoint)
+            cmds.orientConstraint(fkFootCtrl, fkFootJoint)
+            
+            if ball:
+                cmds.orientConstraint(fkBallCtrl, fkBallJoint)
+                
+                
+                
+            #color the controls
+            if side == "l":
+                color = 6
+            
+            else:
+                color = 13
+                
+            cmds.setAttr(fkThighCtrl + ".overrideEnabled", 1)
+            cmds.setAttr(fkThighCtrl + ".overrideColor", color)
+
+            
+            
+
+            
+            
+    # # # # # # # # # # # # # # # # # # # # # # # # # # # # # # # # # # # # # # # # # # # # # # # # # # # # # # # # # # # # # # # 
+    # # # # # # # # # # # # # # # # # # # # # # # # # # # # # # # # # # # # # # # # # # # # # # # # # # # # # # # # # # # # # # # 
+    # # # # # # # # # # # # # # # # # # # # # # # # # # # # # # # # # # # # # # # # # # # # # # # # # # # # # # # # # # # # # # # 
+    # # # # # # # # # # # # # # # # # # # # # # # # # # # # # # # # # # # # # # # # # # # # # # # # # # # # # # # # # # # # # # # 
+    # # # # # # # # # # # # # # # # # # # # # # # # # # # # # # # # # # # # # # # # # # # # # # # # # # # # # # # # # # # # # # # 
+    def buildIKLegs(self):
+        
+        #need to create the leg joints for each side based on the driver thigh, calf, and foot
+        
+        for side in ["l", "r"]:
+            
+            #create joints
+            ikThighJoint = cmds.duplicate("driver_thigh_" + side, name = "ik_leg_thigh_" + side, parentOnly = True)[0]
+            ikCalfJoint = cmds.duplicate("driver_calf_" + side, name = "ik_leg_calf_" + side, parentOnly = True)[0]
+            ikFootJoint = cmds.duplicate("driver_foot_" + side, name = "ik_leg_foot_" + side, parentOnly = True)[0]
+            cmds.setAttr(ikThighJoint + ".v", 0)
+            
+            for joint in [ikThighJoint, ikCalfJoint, ikFootJoint]:
+                cmds.parent(joint, world = True)
+                
+            cmds.parent(ikFootJoint, ikCalfJoint)
+            cmds.parent(ikCalfJoint, ikThighJoint)
+            cmds.makeIdentity(ikThighJoint, t = 0, r = 1, s = 0, apply = True)
+            
+            #create the 2 joint chain for the no flip setup
+            cmds.select(clear = True)
+            beginJoint = cmds.joint(name = "noflip_begin_joint_" + side)
+            cmds.select(clear = True)
+            endJoint = cmds.joint(name = "noflip_end_joint_" + side)
+            cmds.select(clear = True)
+            cmds.setAttr(beginJoint + ".v", 0)
+            
+            beginPos = cmds.xform(ikThighJoint, q = True, ws = True, t = True)
+            cmds.xform(beginJoint, ws = True, t = (beginPos[0], 0, beginPos[2]))
+            endPos = cmds.xform(ikFootJoint, q = True, ws = True, t = True)
+            cmds.xform(endJoint, ws = True, relative = True, t = (endPos[0], 0, endPos[2]))
+            cmds.parent(endJoint, beginJoint)
+            cmds.makeIdentity(beginJoint, t = 0, r = 1, s = 0, apply = True)
+            
+            #set preferred angle
+            cmds.setAttr(beginJoint + ".preferredAngleX", -90)
+            
+            
+            #apply a RP IK solver to the 2 bone chain
+            ikNodes = cmds.ikHandle(name = "noflip_chain_ikHandle_" + side, solver = "ikRPsolver", sj = beginJoint, ee = endJoint)
+            for node in ikNodes:
+                cmds.setAttr(node + ".v", 0)
+            
+            #create a locator(target loc) and group it
+            targetLoc = cmds.spaceLocator(name = "noflip_target_loc_" + side)[0]
+            targetGrp = cmds.group(empty = True, name = "noflip_target_loc_grp_" + side)
+            cmds.setAttr(targetLoc + ".v", 0)
+            
+            constraint = cmds.pointConstraint(beginJoint, targetGrp)[0]
+            cmds.delete(constraint)
+            cmds.parent(targetLoc, targetGrp)
+            constraint = cmds.pointConstraint(endJoint, targetLoc)
+            cmds.delete(constraint)
+            
+            cmds.parent(ikNodes[0], targetLoc)
+            
+            #create the foot control
+            footCtrl = self.createControl("foot", 1, ("ik_foot_anim_" + side))
+            footCtrlGrp = cmds.group(empty = True, name = "ik_foot_anim_grp_" + side)
+            constraint = cmds.pointConstraint(ikFootJoint, footCtrlGrp)[0]
+            cmds.delete(constraint)
+
+            
+            
+            #position the foot control
+            footCtrlPos = cmds.xform("ball_mover_" + side + "_grp", q = True, ws = True, t = True)
+            cmds.xform(footCtrl, ws = True, t = (footCtrlPos[0], 0, 0))
+	    constraint = cmds.pointConstraint("ball_mover_" + side + "_grp", footCtrl)[0]
+	    cmds.delete(constraint)
+	    
+            cmds.makeIdentity(footCtrl, t=1, r=1, s=1, apply = True)
+            
+            if side == "r":
+                cmds.setAttr(footCtrl + ".sx", -1)
+                cmds.makeIdentity(footCtrl, t=1, r=1, s=1, apply = True)
+
+            cmds.xform(footCtrl, ws = True, piv = [endPos[0], endPos[1], endPos[2]])
+            
+            footCtrlSpaceSwitcherFollow = cmds.duplicate(footCtrlGrp, po = True, name = "ik_foot_anim_" + side + "_space_switcher_follow")[0]
+	    footCtrlSpaceSwitcher = cmds.duplicate(footCtrlGrp, po = True, name = "ik_foot_anim_" + side + "_space_switcher")[0]
+            
+            cmds.parent(footCtrlSpaceSwitcher, footCtrlSpaceSwitcherFollow)
+            cmds.parent(footCtrlGrp, footCtrlSpaceSwitcher)
+            cmds.parent(footCtrl, footCtrlGrp)
+            cmds.makeIdentity(footCtrl, t=1, r=1, s=1, apply = True)
+            
+            
+            #create the noflip pole vector loc
+            scale = self.getScaleFactor()
+            noflipVectorLoc = cmds.spaceLocator(name = "noflip_pv_loc_" + side)[0]
+            noflipVectorGrp = cmds.group(name = "noflip_pv_loc_grp_" + side, empty = True)
+            
+            constraint = cmds.pointConstraint([beginJoint, endJoint], noflipVectorLoc)[0]
+            cmds.delete(constraint)
+            constraint = cmds.pointConstraint(targetLoc, noflipVectorGrp)[0]
+            cmds.delete(constraint)
+            
+            cmds.setAttr(noflipVectorLoc + ".v", 0)
+            
+            
+            
+            
+            noflipVectorLocPos = cmds.xform(footCtrl + "_end_loc", q = True, ws = True, t = True)
+            if side == "l":
+                cmds.setAttr(noflipVectorLoc + ".ty", noflipVectorLocPos[1])
+            else:
+                cmds.setAttr(noflipVectorLoc + ".ty", noflipVectorLocPos[1] * -1)
+                
+            cmds.makeIdentity(noflipVectorLoc, t = 1, r = 1, s = 1, apply = True)
+            
+            cmds.parent(noflipVectorLoc, noflipVectorGrp)
+            if side == "l":
+                cmds.setAttr(noflipVectorLoc + ".ty", (200 * scale))
+            else:
+                cmds.setAttr(noflipVectorLoc + ".ty", (-200 * scale))
+                
+            cmds.makeIdentity(noflipVectorLoc, t = 1, r = 1, s = 1, apply = True)
+            cmds.parentConstraint(endJoint, noflipVectorGrp)
+            
+            
+            #duplicate the targetGrp to create our aim vector locator
+            aimGrp = cmds.duplicate(targetGrp, name = "noflip_aim_grp_" + side, parentOnly = True)[0]
+            aimSoftGrp = cmds.duplicate(targetGrp, name = "noflip_aim_soft_grp_" + side, parentOnly = True)[0]
+            aimLoc = cmds.duplicate(targetLoc, name = "noflip_aim_loc_" + side, parentOnly = True)[0]
+            cmds.parent(aimSoftGrp, aimGrp)
+            cmds.parent(aimLoc, aimSoftGrp)
+            cmds.setAttr(aimGrp + ".v", 0)
+            
+            if side == "r":
+                cmds.setAttr(aimGrp + ".ry", 90)
+            
+            else:
+                cmds.setAttr(aimGrp + ".ry", -90)
+                
+            
+            #connectAttrs of targetLoc and aimLoc
+            cmds.connectAttr(targetLoc + ".tx", aimLoc + ".tx")
+            cmds.connectAttr(targetLoc + ".tz", aimLoc + ".tz")
+            
+            
+            #pole vector constraint between aimLoc and ikNodes[0] (2bone chain ik handle)
+            cmds.poleVectorConstraint(aimLoc, ikNodes[0])
+            
+            if side == "l":
+                cmds.setAttr(ikNodes[0] + ".twist", 180)
+                
+            twistAmt = cmds.getAttr(beginJoint + ".rz")
+            cmds.setAttr(ikNodes[0] + ".twist", twistAmt * -1)
+            
+            
+            #create RP IK on the actual IK leg chain
+            #set preferred angle first
+            cmds.setAttr(ikThighJoint + ".preferredAngleZ", 90)
+            cmds.setAttr(ikCalfJoint + ".preferredAngleZ", 90)
+            
+            ikNodesLeg = cmds.ikHandle(name = "foot_ikHandle_" + side, solver = "ikRPsolver", sj = ikThighJoint, ee = ikFootJoint)
+            footIK = ikNodesLeg[0]
+            cmds.setAttr(footIK + ".v", 0)
+            
+            cmds.parent(footIK, targetLoc)
+            
+            
+            #create pole vector constraint between knee loc and full ik leg rp ik handle
+            cmds.poleVectorConstraint(noflipVectorLoc, footIK)
+            
+            #set limits on the aimLoc in Z space
+            minTz = cmds.getAttr(aimLoc + ".tz")
+            maxTz = cmds.xform(aimGrp, q = True, ws = True, t = True)[0]
+            
+            if side == "l":
+                maxTz = maxTz * -1
+                
+            cmds.transformLimits(aimLoc, etz = (True, True), tz = (minTz, maxTz))
+            
+            
+            #create the twist control
+            kneeCtrl = self.createControl("arrow", 2, ("ik_knee_anim_" + side))
+            constraint = cmds.pointConstraint(ikCalfJoint, kneeCtrl)[0]
+            cmds.delete(constraint)
+            kneeCtrlGrp = cmds.group(name = "ik_knee_anim_grp_" + side, empty = True)
+            constraint = cmds.parentConstraint(ikCalfJoint, kneeCtrlGrp)[0]
+            cmds.delete(constraint)
+            
+            cmds.parent(kneeCtrl, kneeCtrlGrp)
+            cmds.makeIdentity(kneeCtrl, t = 1, r = 1, s = 1, apply = True)
+            upAxis = self.getUpAxis(kneeCtrl)
+            
+            cmds.pointConstraint(ikCalfJoint, kneeCtrlGrp, mo = True)
+	    cmds.setAttr(kneeCtrl + ".overrideEnabled", 1)
+	    cmds.setAttr(kneeCtrl + ".overrideDisplayType", 2)
+
+            
+
+
+            
+            #Create foot rig
+            
+            #create joints for ball and toe in IK leg skeleton
+            cmds.select(clear = True)
+            ikBallJoint = cmds.joint(name = "ik_leg_ball_" + side)
+            cmds.select(clear = True)
+            
+            ikToeJoint = cmds.joint(name = "ik_leg_toe_" + side)
+            cmds.select(clear = True)
+            
+            #position joints
+            constraint = cmds.parentConstraint("ball_" + side + "_lra", ikBallJoint)[0]
+            cmds.delete(constraint)
+            constraint = cmds.pointConstraint("jointmover_toe_" + side + "_end", ikToeJoint)[0]
+            cmds.delete(constraint)
+            constraint = cmds.orientConstraint(ikBallJoint, ikToeJoint)[0]
+            cmds.delete(constraint)
+            
+            #parent joints into IK leg hierarchy
+            cmds.parent(ikToeJoint, ikBallJoint)
+            cmds.parent(ikBallJoint, ikFootJoint)
+            
+            cmds.makeIdentity(ikBallJoint, r = 1, apply = True)
+            
+            #create SC IK for ankle to ball and ball to toe
+            ballIKNodes = cmds.ikHandle(name = "ikHandle_ball_" + side, solver = "ikSCsolver", sj = ikFootJoint, ee = ikBallJoint)
+            toeIKNodes = cmds.ikHandle(name = "ikHandle_toe_" + side, solver = "ikSCsolver", sj = ikBallJoint, ee = ikToeJoint)
+            
+            cmds.setAttr(ballIKNodes[0] + ".v", 0)
+            cmds.setAttr(toeIKNodes[0] + ".v", 0)
+            
+            
+            
+            #create the locators we need for all of the pivot points
+            toeTipPivot = cmds.spaceLocator(name = "ik_foot_toe_tip_pivot_" + side)[0]
+            insidePivot = cmds.spaceLocator(name = "ik_foot_inside_pivot_" + side)[0]
+            outsidePivot = cmds.spaceLocator(name = "ik_foot_outside_pivot_" + side)[0]
+            heelPivot = cmds.spaceLocator(name = "ik_foot_heel_pivot_" + side)[0]
+            toePivot = cmds.spaceLocator(name = "ik_foot_toe_pivot_" + side)[0]
+            ballPivot = cmds.spaceLocator(name = "ik_foot_ball_pivot_" + side)[0]
+	    masterBallPivot = cmds.spaceLocator(name = "master_foot_ball_pivot_" + side)[0]
+            
+            #create the controls
+	    heelControl = self.createControl("arrowOnBall", 1.5, "heel_ctrl_" + side)
+            toeWiggleControl = self.createControl("arrowOnBall", 2, "toe_wiggle_ctrl_" + side)
+            toeControl = self.createControl("arrowOnBall", 1.5, "toe_tip_ctrl_" + side)
+            
+            if side == "l":
+                cmds.setAttr(toeControl + ".rx", -90)
+                cmds.setAttr(toeControl + ".rz", -90)
+                cmds.makeIdentity(toeControl, t = 1, r = 1, s = 1, apply = True)
+                
+            else:
+                cmds.setAttr(toeControl + ".rx", 90)
+                cmds.setAttr(toeControl + ".rz", -90)
+                cmds.makeIdentity(toeControl, t = 1, r = 1, s = 1, apply = True)
+     
+            if side == "l":
+                cmds.setAttr(toeWiggleControl + ".rx", -90)
+                cmds.makeIdentity(toeWiggleControl, t = 1, r = 1, s = 1, apply = True)
+                
+            else:
+                cmds.setAttr(toeWiggleControl + ".rx", 90)
+                cmds.makeIdentity(toeWiggleControl, t = 1, r = 1, s = 1, apply = True)
+                
+                
+            cmds.setAttr(heelControl + ".rx", -90)
+            cmds.makeIdentity(heelControl, t = 1, r = 1, s = 1, apply = True)
+ 
+ 
+            #position and orient controls
+            constraint = cmds.parentConstraint("jointmover_" + side + "_heel_loc", heelControl)[0]
+            cmds.delete(constraint)
+            constraint = cmds.parentConstraint("ball_" + side + "_lra", toeWiggleControl)[0]
+            cmds.delete(constraint)
+            constraint = cmds.pointConstraint("jointmover_toe_" + side + "_end", toeControl)[0]
+            cmds.delete(constraint)
+            constraint = cmds.orientConstraint(toeWiggleControl, toeControl)[0]
+            cmds.delete(constraint)
+
+                
+            #position the pivots
+            constraint = cmds.pointConstraint(heelControl, heelPivot)[0]
+	    cmds.delete(constraint)
+	    
+	    constraint = cmds.orientConstraint(heelControl, heelPivot)[0]
+            cmds.delete(constraint)
+
+            constraint = cmds.parentConstraint(toeWiggleControl, ballPivot)[0]
+            cmds.delete(constraint)
+            
+            constraint = cmds.parentConstraint(toeControl, toeTipPivot)[0]
+            cmds.delete(constraint)
+            
+            constraint = cmds.parentConstraint(toeControl, toePivot)[0]
+            cmds.delete(constraint)
+            
+            constraint = cmds.parentConstraint("inside_pivot_" + side + "_mover", insidePivot)[0]
+            cmds.delete(constraint)
+            
+            constraint = cmds.parentConstraint("outside_pivot_" + side + "_mover", outsidePivot)[0]
+            cmds.delete(constraint)
+	    
+            constraint = cmds.parentConstraint(ballPivot, masterBallPivot)[0]
+            cmds.delete(constraint)	    
+	    
+	    
+
+            
+            #create groups for each pivot and parent the pivot to the corresponding group
+            for piv in [heelPivot, ballPivot, toeTipPivot, toePivot, insidePivot, outsidePivot, masterBallPivot]:
+                pivGrp = cmds.group(empty = True, name = piv + "_grp")
+                constraint = cmds.parentConstraint(piv, pivGrp)[0]
+                cmds.delete(constraint)
+                cmds.parent(piv, pivGrp)
+                shape = cmds.listRelatives(piv, shapes = True)[0]
+                cmds.setAttr(shape + ".v", 0)
+                
+
+                
+            #create groups for each control and parent the control to the corresponding group
+            for ctrl in [heelControl, toeWiggleControl, toeControl]:
+                grp = cmds.group(empty = True, name = ctrl + "_grp")
+                constraint = cmds.parentConstraint(ctrl, grp)[0]
+                cmds.delete(constraint)
+		cmds.parent(ctrl, grp)
+		
+		if side == "r":
+		    if ctrl == heelControl:
+			cmds.setAttr(grp + ".rx", (cmds.getAttr(grp + ".rx")) *1)
+			cmds.setAttr(grp + ".ry", (cmds.getAttr(grp + ".ry")) *1)
+		    
+                
+                
+            #setup pivot hierarchy
+            cmds.parent(toeWiggleControl + "_grp", toePivot)
+            cmds.parent(ballPivot + "_grp", toePivot)
+            cmds.parent(toePivot + "_grp", heelPivot)
+            cmds.parent(heelPivot + "_grp", outsidePivot)
+            cmds.parent(outsidePivot + "_grp", insidePivot)
+            cmds.parent(insidePivot + "_grp", toeTipPivot)
+            
+            
+            
+            #setup foot roll
+            cmds.setAttr(heelControl + ".rz", 0)
+            cmds.setAttr(heelPivot + ".rz", 0)
+            cmds.setAttr(toePivot + ".rz", 0)
+            cmds.setAttr(ballPivot + ".rz", 0)
+            cmds.setDrivenKeyframe([heelPivot + ".rz", toePivot + ".rz", ballPivot + ".rz"], cd = heelControl + ".rz", itt = "linear", ott = "linear")
+            
+            
+            if side == "l":
+		cmds.setAttr(heelControl + ".rz", -90)
+		cmds.setAttr(heelPivot + ".rz", 0)
+		cmds.setAttr(toePivot + ".rz", 0)
+		cmds.setAttr(ballPivot + ".rz", -90)
+		cmds.setDrivenKeyframe([heelPivot + ".rz", toePivot + ".rz", ballPivot + ".rz"], cd = heelControl + ".rz", itt = "linear", ott = "linear")
+		
+		cmds.setAttr(heelControl + ".rz", 90)
+		cmds.setAttr(heelPivot + ".rz", 90)
+		cmds.setAttr(toePivot + ".rz", 0)
+		cmds.setAttr(ballPivot + ".rz", 0)
+		cmds.setDrivenKeyframe([heelPivot + ".rz", toePivot + ".rz", ballPivot + ".rz"], cd = heelControl + ".rz", itt = "linear", ott = "linear")
+		
+		cmds.setAttr(heelControl + ".rz", 0)
+		cmds.setAttr(heelPivot + ".rz", 0)
+		cmds.setAttr(toePivot + ".rz", 0)
+		cmds.setAttr(ballPivot + ".rz", 0)
+		
+		
+	    if side == "r":
+		cmds.setAttr(heelControl + ".rz", -90)
+		cmds.setAttr(heelPivot + ".rz", 0)
+		cmds.setAttr(toePivot + ".rz", 0)
+		cmds.setAttr(ballPivot + ".rz", -90)
+		cmds.setDrivenKeyframe([heelPivot + ".rz", toePivot + ".rz", ballPivot + ".rz"], cd = heelControl + ".rz", itt = "linear", ott = "linear")
+		
+		cmds.setAttr(heelControl + ".rz", 90)
+		cmds.setAttr(heelPivot + ".rz", 90)
+		cmds.setAttr(toePivot + ".rz", 0)
+		cmds.setAttr(ballPivot + ".rz", 0)
+		cmds.setDrivenKeyframe([heelPivot + ".rz", toePivot + ".rz", ballPivot + ".rz"], cd = heelControl + ".rz", itt = "linear", ott = "linear")
+		
+		cmds.setAttr(heelControl + ".rz", 0)
+		cmds.setAttr(heelPivot + ".rz", 0)
+		cmds.setAttr(toePivot + ".rz", 0)
+		cmds.setAttr(ballPivot + ".rz", 0)
+            
+            
+            #setup heel rotate X and Y
+	    
+	    if side == "l":
+		cmds.connectAttr(heelControl + ".rx", ballPivot + ".ry")
+		
+	    if side == "r":
+                heelControlRXMultNode = cmds.shadingNode("multiplyDivide", asUtility = True, name = heelControl + "_RX_MultNode")
+                cmds.connectAttr(heelControl + ".rx", heelControlRXMultNode + ".input1X")
+                cmds.setAttr(heelControlRXMultNode + ".input2X", -1)
+                cmds.connectAttr(heelControlRXMultNode + ".outputX", ballPivot + ".ry")
+		
+		
+		
+            
+            if side == "l":
+                heelControlRYMultNode = cmds.shadingNode("multiplyDivide", asUtility = True, name = heelControl + "_RY_MultNode")
+                cmds.connectAttr(heelControl + ".ry", heelControlRYMultNode + ".input1X")
+                cmds.setAttr(heelControlRYMultNode + ".input2X", -1)
+                cmds.connectAttr(heelControlRYMultNode + ".outputX", ballPivot + ".rx")
+                
+            else:
+                cmds.connectAttr(heelControl + ".ry", ballPivot + ".rx")
+                
+            
+            #setup toe control Y and Z rotates
+            cmds.connectAttr(toeControl + ".ry", toeTipPivot + ".ry")
+            cmds.connectAttr(toeControl + ".rz", toeTipPivot + ".rz")
+            
+            
+            #setup the toe control RX (side to side)
+	    if side == "l":
+		cmds.setAttr(toeControl + ".rx", 0)
+		cmds.setAttr(insidePivot + ".rx", 0)
+		cmds.setAttr(outsidePivot + ".rx", 0)
+		cmds.setDrivenKeyframe([insidePivot + ".rx", outsidePivot + ".rx"], cd = toeControl + ".rx", itt = "linear", ott = "linear")
+		
+		cmds.setAttr(toeControl + ".rx", -90)
+		cmds.setAttr(insidePivot + ".rx", 0)
+		cmds.setAttr(outsidePivot + ".rx", -90)
+		cmds.setDrivenKeyframe([insidePivot + ".rx", outsidePivot + ".rx"], cd = toeControl + ".rx", itt = "linear", ott = "linear")
+		
+		cmds.setAttr(toeControl + ".rx", 90)
+		cmds.setAttr(insidePivot + ".rx", 90)
+		cmds.setAttr(outsidePivot + ".rx", 0)
+		cmds.setDrivenKeyframe([insidePivot + ".rx", outsidePivot + ".rx"], cd = toeControl + ".rx", itt = "linear", ott = "linear")
+		
+		cmds.setAttr(toeControl + ".rx", 0)
+		cmds.setAttr(insidePivot + ".rx", 0)
+		cmds.setAttr(outsidePivot + ".rx", 0)
+		
+	    if side == "r":
+		cmds.setAttr(toeControl + ".rx", 0)
+		cmds.setAttr(insidePivot + ".rx", 0)
+		cmds.setAttr(outsidePivot + ".rx", 0)
+		cmds.setDrivenKeyframe([insidePivot + ".rx", outsidePivot + ".rx"], cd = toeControl + ".rx", itt = "linear", ott = "linear")
+		
+		cmds.setAttr(toeControl + ".rx", -90)
+		cmds.setAttr(insidePivot + ".rx", 0)
+		cmds.setAttr(outsidePivot + ".rx", 90)
+		cmds.setDrivenKeyframe([insidePivot + ".rx", outsidePivot + ".rx"], cd = toeControl + ".rx", itt = "linear", ott = "linear")
+		
+		cmds.setAttr(toeControl + ".rx", 90)
+		cmds.setAttr(insidePivot + ".rx", -90)
+		cmds.setAttr(outsidePivot + ".rx", 0)
+		cmds.setDrivenKeyframe([insidePivot + ".rx", outsidePivot + ".rx"], cd = toeControl + ".rx", itt = "linear", ott = "linear")
+		
+		cmds.setAttr(toeControl + ".rx", 0)
+		cmds.setAttr(insidePivot + ".rx", 0)
+		cmds.setAttr(outsidePivot + ".rx", 0)
+            
+            
+    
+            #parent the IK nodes into the foot rig setup
+            cmds.parent(footIK, ballPivot)
+            cmds.parent(ballIKNodes[0], ballPivot)
+            cmds.parent(toeIKNodes[0], toeWiggleControl)
+            cmds.pointConstraint(footCtrl, targetLoc, mo = True)
+	    cmds.parent([toeTipPivot + "_grp", heelControl + "_grp", toeControl + "_grp"], masterBallPivot)
+	    cmds.parent(masterBallPivot + "_grp", footCtrl)
+	
+	    
+	    #add the heel pivot and ball pivot attrs to the foot control
+	    cmds.select(heelControl)
+	    
+	    cmds.addAttr(longName= ( "heelPivot" ), defaultValue=0,  keyable = True)            
+            cmds.addAttr(longName= ( "ballPivot" ), defaultValue=0,  keyable = True)
+	    
+	    #setup heel and ball pivot
+	    
+	    if side == "r":
+		heelPivotMultNode = cmds.shadingNode("multiplyDivide", asUtility = True, name = heelPivot + "_MultNode")
+		cmds.connectAttr(heelControl + ".heelPivot", heelPivotMultNode + ".input1X")
+		cmds.setAttr(heelPivotMultNode + ".input2X", -1)
+		cmds.connectAttr(heelPivotMultNode + ".outputX", heelPivot + ".rx")
+		
+	    else:
+		cmds.connectAttr(heelControl + ".heelPivot", heelPivot + ".rx")
+		
+	    cmds.connectAttr(heelControl + ".ballPivot", masterBallPivot + ".ry")
+
+            
+            #clean up the hierarchy
+            ctrlGrp = cmds.group(name = "leg_ctrl_grp_" + side, empty = True)
+            cmds.parent([ikThighJoint, targetGrp, aimGrp, noflipVectorGrp], ctrlGrp)
+            legGroup = cmds.group(name = "leg_group_" + side, empty = True)
+            constraint = cmds.pointConstraint("driver_pelvis", legGroup)[0]
+            cmds.delete(constraint)
+            cmds.parent([footCtrlSpaceSwitcherFollow, beginJoint, ctrlGrp], legGroup)
+	    cmds.orientConstraint("body_anim_grp", ctrlGrp, mo = True)
+	    cmds.pointConstraint("body_anim", ctrlGrp, mo = True)
+	    
+	    #constrain aimGrp
+	    cmds.pointConstraint("body_anim", aimGrp, mo = True)
+	    cmds.orientConstraint("offset_anim", aimGrp, mo = True)
+	    
+	    
+	    #cmds.parentConstraint("driver_pelvis", beginJoint, mo = True)
+            ikGrp = cmds.group(name = "ik_leg_grp_" + side, empty = True)
+            cmds.parent(ikGrp, legGroup)
+            cmds.parent([kneeCtrlGrp, footCtrlSpaceSwitcherFollow], ikGrp)
+            
+            
+            
+            #color the controls
+            if side == "l":
+                color = 6
+            
+            else:
+                color = 13
+                
+            for control in [footCtrl]:
+                cmds.setAttr(control + ".overrideEnabled", 1)
+                cmds.setAttr(control + ".overrideColor", color)
+                
+                
+                
+            #connect ik twist attr to ik leg twist
+            cmds.select(footCtrl)
+	    cmds.addAttr(longName=("knee_twist"), at = 'double', keyable = True)
+            if side == "r":
+                cmds.connectAttr(footCtrl + ".knee_twist", footIK + ".twist")
+                
+            else:
+                twistMultNode = cmds.shadingNode("multiplyDivide", name = "ik_leg_" + side + "_twistMultNode", asUtility = True)
+                cmds.connectAttr(footCtrl + ".knee_twist", twistMultNode + ".input1X")
+                cmds.setAttr(twistMultNode + ".input2X", -1)
+                cmds.connectAttr(twistMultNode + ".outputX", footIK + ".twist")
+		
+	                     
+	                     
+	    #add stretchy IK to legs
+	    cmds.select(footCtrl)
+            cmds.addAttr(longName=("stretch"), at = 'double',min = 0, max = 1, dv = 0, keyable = True)
+	    cmds.addAttr(longName=("squash"), at = 'double',min = 0, max = 1, dv = 0, keyable = True)
+            cmds.addAttr(longName=("toeCtrlVis"), at = 'bool', dv = 0, keyable = True)
+            stretchMultNode = cmds.shadingNode("multiplyDivide", asUtility = True, name = "ikLeg_stretchToggleMultNode_" + side)
+
+
+
+	    #need to get the total length of the leg chain
+	    totalDist = abs(cmds.getAttr(ikCalfJoint + ".tx" ) + cmds.getAttr(ikFootJoint + ".tx"))
+	    
+		
+	    #create a distanceBetween node
+	    distBetween = cmds.shadingNode("distanceBetween", asUtility = True, name = side + "_ik_leg_distBetween")
+	    
+	    #get world positions of upper arm and ik
+	    baseGrp = cmds.group(empty = True, name = "ik_leg_base_grp_" + side)
+	    endGrp = cmds.group(empty = True, name = "ik_leg_end_grp_" + side)
+	    cmds.pointConstraint(ikThighJoint, baseGrp)
+	    cmds.pointConstraint(footCtrl, endGrp)
+	    
+	    #hook in group translates into distanceBetween node inputs
+	    cmds.connectAttr(baseGrp + ".translate", distBetween + ".point1")
+	    cmds.connectAttr(endGrp + ".translate", distBetween + ".point2")
+	    
+	    #create a condition node that will compare original length to current length
+	    #if second term is greater than, or equal to the first term, the chain needs to stretch
+	    ikLegCondition = cmds.shadingNode("condition", asUtility = True, name = side + "_ik_leg_stretch_condition")
+	    cmds.setAttr(ikLegCondition + ".operation", 3)
+	    cmds.connectAttr(distBetween + ".distance", ikLegCondition + ".secondTerm")
+	    cmds.setAttr(ikLegCondition + ".firstTerm", totalDist)
+	    
+	    #hook up the condition node's return colors
+	    cmds.setAttr(ikLegCondition + ".colorIfTrueR", totalDist)
+	    cmds.connectAttr(distBetween + ".distance", ikLegCondition + ".colorIfFalseR")
+	    
+	    
+	    #create the mult/divide node(set to divide) that will take the original creation length as a static value in input2x, and the connected length into 1x.
+	    legDistMultNode = cmds.shadingNode("multiplyDivide", asUtility = True, name = "leg_dist_multNode_" + side)
+	    cmds.setAttr(legDistMultNode + ".operation", 2) #divide
+	    cmds.setAttr(legDistMultNode + ".input2X", totalDist)
+	    cmds.connectAttr(ikLegCondition + ".outColorR", legDistMultNode + ".input1X")
+	    
+	    #create a stretch toggle mult node that multiplies the stretch factor by the bool of the stretch attr. (0 or 1), this way our condition reads
+	    #if this result is greater than the original length(impossible if stretch bool is off, since result will be 0), than take this result and plug it
+	    #into the scale of our IK arm joints
+	    stretchToggleCondition = cmds.shadingNode("condition", asUtility = True, name = "leg_stretch_toggle_condition_" + side)
+	    cmds.setAttr(stretchToggleCondition + ".operation", 0)
+	    cmds.connectAttr(footCtrl + ".stretch", stretchToggleCondition + ".firstTerm")
+	    cmds.setAttr(stretchToggleCondition + ".secondTerm", 1)
+	    cmds.connectAttr(legDistMultNode + ".outputX", stretchToggleCondition + ".colorIfTrueR")
+	    cmds.setAttr(stretchToggleCondition + ".colorIfFalseR", 1)
+	    
+	    
+	    #set up the squash nodes
+	    squashMultNode = cmds.shadingNode("multiplyDivide", asUtility = True, name = side + "_ik_leg_squash_mult")
+	    cmds.setAttr(squashMultNode + ".operation", 2)
+	    cmds.setAttr(squashMultNode + ".input1X", totalDist)
+	    cmds.connectAttr(ikLegCondition + ".outColorR", squashMultNode + ".input2X")
+	    
+	    
+	    #create a stretch toggle mult node that multiplies the stretch factor by the bool of the stretch attr. (0 or 1), this way our condition reads
+	    #if this result is greater than the original length(impossible if stretch bool is off, since result will be 0), than take this result and plug it
+	    #into the scale of our IK arm joints
+	    squashToggleCondition = cmds.shadingNode("condition", asUtility = True, name = "leg_squash_toggle_condition_" + side)
+	    cmds.setAttr(squashToggleCondition + ".operation", 0)
+	    cmds.connectAttr(footCtrl + ".squash", squashToggleCondition + ".firstTerm")
+	    cmds.setAttr(squashToggleCondition + ".secondTerm", 1)
+	    cmds.connectAttr(squashMultNode + ".outputX", squashToggleCondition + ".colorIfTrueR")
+	    cmds.setAttr(squashToggleCondition + ".colorIfFalseR", 1)
+	    
+	    
+	    
+	    #connect to arm scale
+	    cmds.connectAttr(stretchToggleCondition + ".outColorR", ikThighJoint + ".sx")
+	    cmds.connectAttr(stretchToggleCondition + ".outColorR", ikCalfJoint + ".sx")
+	    
+	    cmds.connectAttr(squashToggleCondition + ".outColorR", ikCalfJoint + ".sy")
+	    cmds.connectAttr(squashToggleCondition + ".outColorR", ikCalfJoint + ".sz")
+	    
+	    cmds.connectAttr(squashToggleCondition + ".outColorR", ikThighJoint + ".sy")
+	    cmds.connectAttr(squashToggleCondition + ".outColorR", ikThighJoint + ".sz")
+	    
+	    
+	    #add base and end groups to arm grp
+	    cmds.parent([baseGrp, endGrp], ctrlGrp)
+            
+            
+            #lock attrs on control that shouldn't be animated
+            
+            for control in [toeControl, heelControl, toeWiggleControl]:
+                cmds.setAttr(control + ".tx", lock = True, keyable = False)
+                cmds.setAttr(control + ".ty", lock = True, keyable = False)
+                cmds.setAttr(control + ".tz", lock = True, keyable = False)
+                cmds.setAttr(control + ".sx", lock = True, keyable = False)
+                cmds.setAttr(control + ".sy", lock = True, keyable = False)
+                cmds.setAttr(control + ".sz", lock = True, keyable = False)
+                cmds.setAttr(control + ".v", lock = True, keyable = False)
+                
+
+            #lock attrs on foot control that should not be animated
+            cmds.setAttr(footCtrl + ".sx", lock = True, keyable = False)
+            cmds.setAttr(footCtrl + ".sy", lock = True, keyable = False)
+            cmds.setAttr(footCtrl + ".sz", lock = True, keyable = False)
+            cmds.setAttr(footCtrl + ".v", lock = True, keyable = False)
+            
+            #lock attrs on knee control that should not be animated
+	    cmds.connectAttr(footCtrl + ".knee_twist", kneeCtrl + ".rx")
+
+		
+            cmds.setAttr(kneeCtrl + ".rx", lock = False, keyable = False)
+            cmds.setAttr(kneeCtrl + ".ry", lock = True, keyable = False)
+            cmds.setAttr(kneeCtrl + ".rz", lock = True, keyable = False)
+            cmds.setAttr(kneeCtrl + ".tx", lock = True, keyable = False)
+            cmds.setAttr(kneeCtrl + ".ty", lock = True, keyable = False)
+            cmds.setAttr(kneeCtrl + ".tz", lock = True, keyable = False)
+            cmds.setAttr(kneeCtrl + ".sx", lock = True, keyable = False)
+            cmds.setAttr(kneeCtrl + ".sy", lock = True, keyable = False)
+            cmds.setAttr(kneeCtrl + ".sz", lock = True, keyable = False)
+            cmds.setAttr(kneeCtrl + ".v", lock = True, keyable = False)
+            
+            
+            
+
+
+        
+    # # # # # # # # # # # # # # # # # # # # # # # # # # # # # # # # # # # # # # # # # # # # # # # # # # # # # # # # # # # # # # # 
+    # # # # # # # # # # # # # # # # # # # # # # # # # # # # # # # # # # # # # # # # # # # # # # # # # # # # # # # # # # # # # # # 
+    # # # # # # # # # # # # # # # # # # # # # # # # # # # # # # # # # # # # # # # # # # # # # # # # # # # # # # # # # # # # # # # 
+    # # # # # # # # # # # # # # # # # # # # # # # # # # # # # # # # # # # # # # # # # # # # # # # # # # # # # # # # # # # # # # # 
+    # # # # # # # # # # # # # # # # # # # # # # # # # # # # # # # # # # # # # # # # # # # # # # # # # # # # # # # # # # # # # # # 
+    def buildFingers(self):
+        
+        #find out which finger joints need to be rigged
+        for side in ["l", "r"]:
+            
+            
+            #create a list to hold all ctrl groups that are created
+            ctrlGroups = []
+            ikGrps = []
+            joints = []
+            fkOrients = []
+	    metaJoints = []
+	    ikJoints = []
+	    
+            
+            children = cmds.listRelatives("driver_hand_" + side, children = True, type = 'joint')
+            allChildren = cmds.listRelatives("driver_hand_" + side, allDescendents = True, type = 'joint')
+            
+            #find out how many finger joints we have for each finger
+            thumbMeta = [False, None]
+            indexMeta = [False, None]
+            middleMeta = [False, None]
+            ringMeta = [False, None]
+            pinkyMeta = [False, None]
+            numThumbJoints = [0, "thumb"]
+            numIndexJoints = [0, "index"]
+            numMiddleJoints = [0, "middle"]
+            numRingJoints = [0, "ring"]
+            numPinkyJoints = [0, "pinky"]
+            
+            if allChildren:
+                for finger in allChildren:
+                    #find if metatarsals exist
+                    if finger.find("meta") != -1:
+                        if finger.partition("driver_")[2].find("index") == 0:
+                            indexMeta = [True, "index"]
+                        if finger.partition("driver_")[2].find("middle") == 0:
+                            middleMeta = [True, "middle"]
+                        if finger.partition("driver_")[2].find("ring") == 0:
+                            ringMeta = [True, "ring"]
+                        if finger.partition("driver_")[2].find("pinky") == 0:
+                            pinkyMeta = [True, "pinky"]
+                            
+                            
+                    #get num fingers -meta
+                    if finger.partition("driver_")[2].find("thumb") == 0:
+                        numThumbJoints[0] += 1
+                    if finger.partition("driver_")[2].find("index") == 0:
+                        numIndexJoints[0] += 1
+                    if finger.partition("driver_")[2].find("middle") == 0:
+                        numMiddleJoints[0] += 1
+                    if finger.partition("driver_")[2].find("ring") == 0:
+                        numRingJoints[0] += 1
+                    if finger.partition("driver_")[2].find("pinky") == 0:
+                        numPinkyJoints[0] += 1
+                
+            
+                #subtract metacarpals (only if they exist!)
+                if indexMeta[0] == True:
+                    numIndexJoints[0] -= 1
+                if middleMeta[0] == True:
+                    numMiddleJoints[0] -= 1
+                if ringMeta[0] == True:
+                    numRingJoints[0] -= 1
+                if pinkyMeta[0] == True:
+                    numPinkyJoints[0] -= 1
+                
+            
+    
+    
+            #duplicate the driver joints to be used as the rig joints
+            if children:
+                for child in children:
+                    for mode in ["fk", "ik"]:
+                        dupeChildNodes = cmds.duplicate(child, name = "temp")
+                        
+                        #parent root joint of each finger to world if not already child of world
+                        parent = cmds.listRelatives(dupeChildNodes[0], parent = True)[0]
+                        if parent != None:
+                            cmds.parent(dupeChildNodes[0], world = True)
+                            
+                        #rename duped joints
+                        for node in dupeChildNodes:
+                            if node == "temp":
+                                niceName = child.partition("driver_")[2]
+                                joint = cmds.rename(node, "rig_" + mode + "_" + niceName)
+				if mode == "ik":
+				    ikJoints.append(joint)
+				else:
+				    joints.append(joint)
+                                
+                            else:
+                                niceName = node.partition("driver_")[2]
+                                cmds.rename("rig_*|" + node, "rig_" + mode + "_" + niceName)
+                            
+                            
+
+                        
+            #if the metacarpal fingers exist, create a control for them
+            for meta in [indexMeta, middleMeta, ringMeta, pinkyMeta]:
+                if meta[0] == True:
+                    #create the control object for the metacarpal
+                    ctrlName = meta[1]
+                    ctrlName = ctrlName + "_metacarpal_ctrl_" + side
+                    control = self.createControl("square", 1, ctrlName)
+                    constraint = cmds.parentConstraint("rig_fk_" + meta[1] + "_metacarpal_" + side, control)[0]
+                    cmds.delete(constraint)
+                    
+                    cmds.setAttr(control + ".sx", 0)
+                    cmds.setAttr(control + ".sy", 15)
+                    cmds.setAttr(control + ".sz", 15)
+                    cmds.makeIdentity(control, t = 1, r = 1, s = 1, apply = True)
+                    
+                    #create the group node and parent ctrl to it
+                    ctrlGrp = cmds.group(empty = True, name = ctrlName + "_grp")
+		    constraint = cmds.parentConstraint("rig_fk_" + meta[1] + "_metacarpal_" + side, ctrlGrp)[0]
+		    metaJoints.append(ctrlGrp)
+                        
+                    cmds.delete(constraint)
+                    cmds.parent(control, ctrlGrp)
+                    cmds.makeIdentity(control, t = 1, r = 1, s = 1, apply = True)
+                    
+                    #parent constrain the rig joint to the control
+		    cmds.parentConstraint(control, "rig_fk_" + meta[1] + "_metacarpal_" + side, mo = True)
+		    cmds.parentConstraint(control, "rig_ik_" + meta[1] + "_metacarpal_" + side, mo = True)
+                    
+                    #lock attrs on control that shouldn't be animated
+                    cmds.setAttr(control + ".tx", lock = True, keyable = False)
+                    cmds.setAttr(control + ".ty", lock = True, keyable = False)
+                    cmds.setAttr(control + ".tz", lock = True, keyable = False)
+                    cmds.setAttr(control + ".sx", lock = True, keyable = False)
+                    cmds.setAttr(control + ".sy", lock = True, keyable = False)
+                    cmds.setAttr(control + ".sz", lock = True, keyable = False)
+                    cmds.setAttr(control + ".v", lock = True, keyable = False)
+                    
+                    
+                    #color the controls
+                    if side == "l":
+                        color = 6
+                    
+                    else:
+                        color = 13
+                        
+                    cmds.setAttr(control + ".overrideEnabled", 1)
+                    cmds.setAttr(control + ".overrideColor", color)
+                    
+            
+
+            
+            #Create the FK orient joints
+            
+            #first create a group for the IK handles to go into. Then setup the constraints on this group and set driven keys
+            ikHandlesGrp = cmds.group(empty = True, name = "fkOrient_ikHandles_" + side + "_grp")
+            constraint = cmds.parentConstraint("ik_wrist_" + side + "_anim", "fk_wrist_" + side + "_anim", ikHandlesGrp, mo = True)[0]
+            
+            cmds.setAttr("Rig_Settings" + "." + side + "ArmMode", 0)
+            cmds.setAttr(constraint + ".ik_wrist_" + side + "_anim" + "W0", 0)
+            cmds.setAttr(constraint + ".fk_wrist_" + side + "_anim" + "W1", 1)
+            cmds.setDrivenKeyframe([constraint + ".ik_wrist_" + side + "_anim" + "W0", constraint + ".fk_wrist_" + side + "_anim" + "W1"], cd = "Rig_Settings" + "." + side + "ArmMode", itt = "linear", ott = "linear")
+
+            cmds.setAttr("Rig_Settings" + "." + side + "ArmMode", 1)
+            cmds.setAttr(constraint + ".ik_wrist_" + side + "_anim" + "W0", 1)
+            cmds.setAttr(constraint + ".fk_wrist_" + side + "_anim" + "W1", 0)
+            cmds.setDrivenKeyframe([constraint + ".ik_wrist_" + side + "_anim" + "W0", constraint + ".fk_wrist_" + side + "_anim" + "W1"], cd = "Rig_Settings" + "." + side + "ArmMode", itt = "linear", ott = "linear")
+
+
+            
+            for fingers in [numIndexJoints, numMiddleJoints, numRingJoints, numPinkyJoints, numThumbJoints]:
+                if fingers[0] > 0:
+                    
+                    #setup metaCtrl name
+                    if fingers[1] == "thumb":
+                        metaCtrl = fingers[1] + "_01_" + side
+                    
+                    else:
+                        metaCtrl = fingers[1] + "_metacarpal_ctrl_" + side
+                        
+                    #create the base and end joint
+                    baseJoint = cmds.duplicate("rig_fk_" + fingers[1] + "_01_" + side, po = True, name = "rig_fkOrient_" + fingers[1] + "_01_" + side)[0]
+                    endJoint = cmds.duplicate("rig_fk_" + fingers[1] + "_0" + str(fingers[0]) + "_" + side, po = True, name = "rig_fkOrient_" + fingers[1] + "_0" + str(fingers[0]) + "_" + side)[0]
+                    
+                    
+                    #position the end joint
+                    scaleFactor = self.getScaleFactor()
+                    if side == "l":
+                        cmds.parent(endJoint, "rig_fk_" + fingers[1] + "_0" + str(fingers[0]) + "_" + side)
+                        cmds.setAttr(endJoint + ".tx", 5 * scaleFactor)
+                        
+                    else:
+                        cmds.parent(endJoint, "rig_fk_" + fingers[1] + "_0" + str(fingers[0]) + "_" + side)
+                        cmds.setAttr(endJoint + ".tx", -5 * scaleFactor)
+                    
+                    #parent the end joint to the base joint
+                    cmds.parent(endJoint, baseJoint)
+                    
+                    #create SC ik handles for each chain
+                    ikNodes = cmds.ikHandle(sol = "ikSCsolver", name = baseJoint + "_ikHandle", sj = baseJoint, ee = endJoint)[0]
+                    cmds.parent(ikNodes, ikHandlesGrp)
+                    cmds.setAttr(ikNodes + ".v", 0)
+                    
+                    #parent our orient joint to the metacarpal if it exists
+                    if cmds.objExists(metaCtrl):
+                        if fingers[1] == "thumb":
+                            fkOrients.append(baseJoint)
+                            
+                        else:
+                            cmds.parent(baseJoint, metaCtrl)
+                        
+                    else:
+                        fkOrients.append(baseJoint)
+                    
+                    
+                    
+            #Create FK controls for the fingers
+            fkControls = []
+            for fingers in [numIndexJoints, numMiddleJoints, numRingJoints, numPinkyJoints, numThumbJoints]:
+                for i in range(int(fingers[0])):
+                    
+                    #create an FK control per finger
+                    ctrlName = fingers[1] + "_finger_fk_ctrl_" + str(i + 1) + "_" + side
+                    control = self.createControl("circle", 3, ctrlName)
+                    ctrlGrp = cmds.group(empty = True, name = control + "_grp")
+
+
+		    metaCtrl = fingers[1] + "_metacarpal_ctrl_" + side
+                    
+                    if cmds.objExists(metaCtrl) == False:
+                        if (i + 1) == 1:
+                            ctrlGroups.append(ctrlGrp)
+                    
+                    #add the created control to the controls list
+                    fkControls.append(control)
+                    
+                    #position control
+                    constraint = cmds.parentConstraint("rig_fk_" + fingers[1] + "_0" + str(i+1) + "_" + side, control)[0]
+                    grpConstraint = cmds.parentConstraint("rig_fk_" + fingers[1] + "_0" + str(i+1) + "_" +  side, ctrlGrp)[0]
+                    cmds.delete([constraint, grpConstraint])
+                    cmds.makeIdentity(control, t = 0, r = 1, s = 0, apply = True)
+                    
+                    
+                    #duplicate the ctrl group to create the driven group
+                    drivenGrp = cmds.duplicate(ctrlGrp, parentOnly = True, name = control + "_driven_grp")[0]
+                    ctrlGroups.append(drivenGrp)
+                    cmds.parent(drivenGrp, ctrlGrp)
+                
+                    #parent control to grp
+                    cmds.parent(control, drivenGrp)
+                    cmds.makeIdentity(control, t = 0, r = 1, s = 0, apply = True)
+                    cmds.setAttr(control + ".ry", -90)
+                    cmds.makeIdentity(control, t = 1, r = 1, s = 1, apply = True)
+                    
+                    #constrain finger joint to control
+                    cmds.parentConstraint(control, "rig_fk_" + fingers[1] + "_0" + str(i+1) + "_" +  side, mo = True)
+                    
+                    
+                    #if we aren't the root of the finger chain, then parent our ctrlGrp to the previous fk control
+                    if i != 0:
+                        cmds.parent(ctrlGrp, ctrlParent)
+                        
+                    else:
+			
+                        #if the control grp is the root of the finger chain, need to parent the ctrl grp to the metaCtrl
+                        if cmds.objExists(metaCtrl):
+			    cmds.parent(ctrlGrp, metaCtrl)
+
+                            
+                            #setup set driven keys for the orientation options
+                            cmds.select(control)
+			    cmds.addAttr(longName= ( "sticky" ), defaultValue=0, minValue=0, maxValue=1, keyable = True)
+
+			    #setup the constraint between the fk finger orient joint and the ctrlGrp
+			    constraint = cmds.parentConstraint("rig_fkOrient_" + fingers[1] + "_01_" + side, ctrlGrp, mo = True)[0]
+			    
+                            #set driven keyframes on constraint
+                            cmds.setAttr(control + ".sticky", 1)
+                            cmds.setAttr(constraint + "." + "rig_fkOrient_" + fingers[1] + "_01_" + side + "W0", 1)
+                            cmds.setDrivenKeyframe(constraint + "." + "rig_fkOrient_" + fingers[1] + "_01_" + side + "W0", cd = control + ".sticky", itt = "linear", ott = "linear")
+
+                            cmds.setAttr(control + ".sticky", 0)
+                            cmds.setAttr(constraint + "." + "rig_fkOrient_" + fingers[1] + "_01_" + side + "W0", 0)
+                            cmds.setDrivenKeyframe(constraint + "." + "rig_fkOrient_" + fingers[1] + "_01_" + side + "W0", cd = control + ".sticky", itt = "linear", ott = "linear")
+			    
+			    if fingers[1] == "thumb":
+				cmds.setAttr(control + ".sticky", 1)
+				
+			    else:
+				cmds.setAttr(control + ".sticky", 0)
+				
+                            ctrlGroups.append(ctrlGrp)
+
+
+
+                        #if the meta carpal does not exist, simply parent the root group under the base joint
+                        else:
+                            ctrlGroups.append(ctrlGrp)
+                            
+                            constraint = cmds.parentConstraint("rig_fkOrient_" + fingers[1] + "_01_" + side, ctrlGrp, mo = True)[0]
+                            
+                            #setup set driven keys for the orientation options
+                            cmds.select(control)
+                            cmds.addAttr(longName= ( "sticky" ), defaultValue=0, minValue=0, maxValue=1, keyable = True)
+                            
+                            #set driven keyframes on constraint
+                            cmds.setAttr(control + ".sticky", 1)
+                            cmds.setAttr(constraint + "." + "rig_fkOrient_" + fingers[1] + "_01_" + side + "W0", 1)
+                            cmds.setDrivenKeyframe(constraint + "." + "rig_fkOrient_" + fingers[1] + "_01_" + side + "W0", cd = control + ".sticky", itt = "linear", ott = "linear")
+
+                            cmds.setAttr(control + ".sticky", 0)
+                            cmds.setAttr(constraint + "." + "rig_fkOrient_" + fingers[1] + "_01_" + side + "W0", 0)
+                            cmds.setDrivenKeyframe(constraint + "." + "rig_fkOrient_" + fingers[1] + "_01_" + side + "W0", cd = control + ".sticky", itt = "linear", ott = "linear")
+
+                    #set the control parent for the next ctrl in the chain to the current control
+                    ctrlParent = control
+                    
+                    #lock attrs on control that shouldn't be animated
+                    cmds.setAttr(control + ".sx", lock = True, keyable = False)
+                    cmds.setAttr(control + ".sy", lock = True, keyable = False)
+                    cmds.setAttr(control + ".sz", lock = True, keyable = False)
+                    cmds.setAttr(control + ".v", lock = True, keyable = False)
+                    
+                    
+                    #color the controls
+                    if side == "l":
+                        color = 6
+                    
+                    else:
+                        color = 13
+                    cmds.setAttr(control + ".overrideEnabled", 1)
+                    cmds.setAttr(control + ".overrideColor", color)
+                    
+                    
+                    
+            #setup the hand roll feature
+            
+            #create our 4 locators(pivots) and position
+            pinkyPiv = cmds.spaceLocator(name = "hand_" + side + "_pinky_pivot")[0]
+            thumbPiv = cmds.spaceLocator(name = "hand_" + side + "_thumb_pivot")[0]
+            midPiv = cmds.spaceLocator(name = "hand_" + side + "_mid_pivot")[0]
+            tipPiv = cmds.spaceLocator(name = "hand_" + side + "_tip_pivot")[0]
+            
+            for piv in [pinkyPiv, thumbPiv, midPiv, tipPiv]:
+                cmds.setAttr(piv + ".v", 0)
+            
+            constraint = cmds.parentConstraint(side + "_hand_pinky_pivot", pinkyPiv)[0]
+            cmds.delete(constraint)
+            constraint = cmds.parentConstraint(side + "_hand_thumb_pivot", thumbPiv)[0]
+            cmds.delete(constraint)
+            constraint = cmds.parentConstraint(side + "_hand_mid_pivot", midPiv)[0]
+            cmds.delete(constraint)
+            constraint = cmds.parentConstraint(side + "_hand_tip_pivot", tipPiv)[0]
+            cmds.delete(constraint)
+            
+            #create the control groups for the pivots so our values are zeroed
+            for each in [pinkyPiv, thumbPiv, midPiv, tipPiv]:
+                group = cmds.group(empty = True, name = each + "_grp")
+                constraint = cmds.parentConstraint(each, group)[0]
+                cmds.delete(constraint)
+                cmds.parent(each, group)
+                
+            #setup hierarchy
+            cmds.parent(thumbPiv + "_grp", pinkyPiv)
+            cmds.parent(tipPiv + "_grp", thumbPiv)
+            cmds.parent(midPiv + "_grp", tipPiv)
+            
+            #parent the arm IK handles under the midPiv locator
+            cmds.parent(["arm_ikHandle_" + side, "invis_arm_ikHandle_" + side], midPiv)
+            cmds.parent(pinkyPiv + "_grp", "ik_wrist_" + side + "_anim")
+            
+            #add attrs to the IK hand control (side, roll, tip pivot)
+            cmds.select("ik_wrist_" + side + "_anim")
+            cmds.addAttr(longName= ( "side" ), defaultValue=0, keyable = True)
+            cmds.addAttr(longName= ( "mid_bend" ), defaultValue=0, keyable = True)
+            cmds.addAttr(longName= ( "mid_swivel" ), defaultValue=0, keyable = True)
+            cmds.addAttr(longName= ( "tip_pivot" ), defaultValue=0, keyable = True)
+            cmds.addAttr(longName= ( "tip_swivel" ), defaultValue=0, keyable = True)
+            
+            
+            #hook up attrs to pivot locators
+            cmds.connectAttr("ik_wrist_" + side + "_anim.mid_bend", midPiv + ".rz")
+            cmds.connectAttr("ik_wrist_" + side + "_anim.tip_pivot", tipPiv + ".rz")
+            
+	    cmds.connectAttr("ik_wrist_" + side + "_anim.mid_swivel", midPiv + ".ry")
+	    cmds.connectAttr("ik_wrist_" + side + "_anim.tip_swivel", tipPiv + ".ry")
+
+                
+            
+            #set driven keys for the side to side attr
+            
+            if side == "l":
+                thumbVal = 180
+                pinkyVal = -180
+            
+            else:
+                thumbVal = 180
+                pinkyVal = -180
+                
+            cmds.setAttr("ik_wrist_" + side + "_anim.side", 0)
+            cmds.setAttr(pinkyPiv + ".rx", 0)
+            cmds.setAttr(thumbPiv + ".rx", 0)
+            cmds.setDrivenKeyframe([pinkyPiv + ".rx", thumbPiv + ".rx"], cd = "ik_wrist_" + side + "_anim.side", itt = "linear", ott = "linear")
+            
+            cmds.setAttr("ik_wrist_" + side + "_anim.side", 180)
+            cmds.setAttr(pinkyPiv + ".rx", pinkyVal)
+            cmds.setAttr(thumbPiv + ".rx", 0)
+            cmds.setDrivenKeyframe([pinkyPiv + ".rx", thumbPiv + ".rx"], cd = "ik_wrist_" + side + "_anim.side", itt = "linear", ott = "linear")
+            
+            cmds.setAttr("ik_wrist_" + side + "_anim.side", -180)
+            cmds.setAttr(pinkyPiv + ".rx", 0)
+            cmds.setAttr(thumbPiv + ".rx", thumbVal)
+            cmds.setDrivenKeyframe([pinkyPiv + ".rx", thumbPiv + ".rx"], cd = "ik_wrist_" + side + "_anim.side", itt = "linear", ott = "linear")
+
+            cmds.setAttr("ik_wrist_" + side + "_anim.side", 0)
+                    
+                    
+                    
+                    
+                    
+            #If there are enough finger joints on each finger, create IK rig
+            ikCtrls = []
+            poleVectorLocs = []
+            modeGrps = []
+            for fingers in [numIndexJoints, numMiddleJoints, numRingJoints, numPinkyJoints, numThumbJoints]:
+                
+                if fingers[0] == 3:
+                    
+                    #set preferred angles on joints so IK will create properly
+                    cmds.setAttr("rig_ik_" + fingers[1] + "_01_" + side + ".preferredAngleZ", 45)
+                    cmds.setAttr("rig_ik_" + fingers[1] + "_02_" + side + ".preferredAngleZ", 45)
+                    cmds.setAttr("rig_ik_" + fingers[1] + "_03_" + side + ".preferredAngleZ", 45)
+                    
+
+                    #create a tip joint
+                    tipJoint = cmds.duplicate("rig_ik_" + fingers[1] + "_03_" + side, po = True, name = "rig_ik_" + fingers[1] + "_tip_" + side)[0]
+                    cmds.parent(tipJoint, "rig_ik_" + fingers[1] + "_03_" + side)
+                    
+                    #position tip joint
+                    if side == "l":
+                        cmds.setAttr(tipJoint + ".tx", 5 * scaleFactor)
+                    else:
+                        cmds.setAttr(tipJoint + ".tx", -5 * scaleFactor)
+                    
+                    #create the IK handle 
+                    ikNodes = cmds.ikHandle(sol = "ikRPsolver", name = fingers[1] + "_" + side + "_ikHandle", sj = "rig_ik_" + fingers[1] + "_01_" + side, ee = "rig_ik_" + fingers[1] + "_03_" + side)[0]
+                    ikTipNodes = cmds.ikHandle(sol = "ikSCsolver", name = fingers[1] + "_" + side + "_end_ikHandle", sj = "rig_ik_" + fingers[1] + "_03_" + side, ee = tipJoint)[0]
+                    cmds.setAttr(ikNodes + ".v", 0)
+		    
+                    cmds.parent(ikTipNodes, ikNodes)
+                    
+                    #create a pole vector locator and position it
+                    poleVector = cmds.spaceLocator(name = fingers[1] + "_" + side + "_poleVector")[0]
+                    constraint = cmds.parentConstraint("rig_ik_" + fingers[1] + "_02_" + side, poleVector)[0]
+                    cmds.delete(constraint)
+                    #color the control
+                    if side == "l":
+                        color = 6
+                    
+                    else:
+                        color = 13
+                        
+                    cmds.setAttr(poleVector + ".overrideEnabled", 1)
+                    cmds.setAttr(poleVector + ".overrideColor", color)
+                    
+                    #create a pole vector group
+                    pvGrp = cmds.group(empty = True, name = poleVector + "_grp")
+                    constraint = cmds.parentConstraint(poleVector, pvGrp)[0]
+                    cmds.delete(constraint)
+                    
+                    
+                    #parent to the joint, and move out away from finger
+                    cmds.parent(poleVector, "rig_ik_" + fingers[1] + "_02_" + side)
+                    
+                    if side == "l":
+                        cmds.setAttr(poleVector + ".ty", -20 * scaleFactor)
+                        
+                    else:
+                        cmds.setAttr(poleVector + ".ty", 20 * scaleFactor)
+                    
+                    cmds.makeIdentity(poleVector, t =1, r =1, s = 1, apply = True)
+                    cmds.parent(poleVector, pvGrp, absolute = True)
+                    cmds.makeIdentity(poleVector, t =1, r =1, s = 1, apply = True)
+                    
+                    #lock pole vector attrs
+		    for attr in [".sx", ".sy", ".sz", ".v"]:
+			if attr == ".v":
+			    cmds.setAttr(poleVector + attr, keyable = False)
+			    
+			else:
+			    cmds.setAttr(poleVector + attr, lock = True, keyable = False)
+			
+			
+                    #create the IK finger controls
+                    ctrlName = fingers[1] + "_" + side + "_ik_anim"
+                    control = self.createControl("circle", 3, ctrlName)
+                    ctrlGrp = cmds.group(empty = True, name = control + "_grp")
+                    ikCtrls.append(control)
+                    
+                    #position control
+                    constraint = cmds.parentConstraint(tipJoint, control)[0]
+                    grpConstraint = cmds.parentConstraint(tipJoint, ctrlGrp)[0]
+                    cmds.delete([constraint, grpConstraint])
+                    cmds.makeIdentity(control, t = 0, r = 1, s = 0, apply = True)
+                    
+                    #parent control to grp
+                    cmds.parent(control, ctrlGrp)
+                    cmds.makeIdentity(control, t = 0, r = 1, s = 0, apply = True)
+                    cmds.setAttr(control + ".ry", -90)
+                    cmds.makeIdentity(control, t = 1, r = 1, s = 1, apply = True)
+                    cmds.parent(ikNodes, control)
+                    
+                    #setup the pole vector constraint and add the locator to the poleVectorLocs list
+                    cmds.poleVectorConstraint(poleVector, ikNodes)
+                    poleVectorLocs.append(pvGrp)
+                    
+                    #add attr to show pole vector control
+                    cmds.select(control)
+                    cmds.addAttr(longName= ( "poleVectorVis" ), defaultValue=0, minValue=0, maxValue=1, keyable = True)
+                    cmds.connectAttr(control + ".poleVectorVis", poleVector + ".v")
+                    
+                    #create a tip locator with finger mode attrs
+                    fingerModeCtrl = cmds.spaceLocator(name = fingers[1] + "_finger_" + side + "_mode_anim")[0]
+                    fingerModeCtrlGrp = cmds.group(empty = True, name = fingers[1] + "_finger_" + side + "_mode_grp")
+                    modeGrps.append(fingerModeCtrlGrp)
+		    cmds.setAttr(fingerModeCtrl + ".v", 0)
+                    
+                    
+                    constraint = cmds.parentConstraint(tipJoint, fingerModeCtrl)[0]
+                    cmds.delete(constraint)
+                    constraint = cmds.parentConstraint(tipJoint, fingerModeCtrlGrp)[0]
+                    cmds.delete(constraint)
+                    
+                    cmds.parent(fingerModeCtrl, fingerModeCtrlGrp)
+                    
+                    #lock attrs
+                    for attr in [".tx", ".ty", ".tz", ".rx", ".ry", ".rz", ".sx", ".sy", ".sz", ".v"]:
+                        cmds.setAttr(fingerModeCtrl + attr, lock = True, keyable = False)
+			
+		    for attr in [".sx", ".sy", ".sz", ".v"]:
+			cmds.setAttr(control + attr, lock = True, keyable = False)
+                        
+                        
+                    #scale up the fingerModeCtrl
+                    shape = cmds.listRelatives(fingerModeCtrl, shapes = True)[0]
+                    cmds.setAttr(shape + ".localScaleX", 3 * scaleFactor)
+                    cmds.setAttr(shape + ".localScaleY", 3 * scaleFactor)
+                    cmds.setAttr(shape + ".localScaleZ", 3 * scaleFactor)
+                    
+                    
+                    
+                    #constrain the fingerModeCtrlGrp to the driver base knuckle
+                    if cmds.objExists("driver_" + fingers[1] + "_03_" + side):
+                        cmds.parentConstraint("driver_" + fingers[1] + "_03_" + side, fingerModeCtrlGrp, mo = True)
+                        
+                    else:
+                        if cmds.objExists("driver_" + fingers[1] + "_02_" + side):
+                            cmds.parentConstraint("driver_" + fingers[1] + "_02_" + side, fingerModeCtrlGrp, mo = True)
+                            
+                        else:
+                            if cmds.objExists("driver_" + fingers[1] + "_01_" + side):
+                                cmds.parentConstraint("driver_" + fingers[1] + "_01_" + side, fingerModeCtrlGrp, mo = True)
+                    
+                    #color the control
+                    if side == "l":
+                        color = 6
+                    
+                    else:
+                        color = 13
+                        
+                    cmds.setAttr(fingerModeCtrl + ".overrideEnabled", 1)
+                    cmds.setAttr(fingerModeCtrl + ".overrideColor", color)
+                    
+                    
+                    #add attr for finger mode(fk/ik) on both IK and FK control
+                    cmds.select(fingerModeCtrl)
+                    cmds.addAttr(longName= "FK_IK", defaultValue = 0, minValue = 0, maxValue = 1, keyable = True)
+                    
+                    
+            
+            
+            #take all of the pole vector groups and add them to a master pv group
+            masterPvGrp = cmds.group(empty = True, name = "fingers_" + side + "_poleVectors_grp")
+            for pv in poleVectorLocs:
+                cmds.parent(pv, masterPvGrp, absolute = True)
+                
+
+            
+            #create a global IK control if there are any IK fingers
+            if ikCtrls:
+                ctrlName =  side + "_global_ik_anim"
+                control = self.createControl("square", 20, ctrlName)
+                ctrlGrp = cmds.group(empty = True, name = control + "_grp")
+                
+                #position control
+                constraint = cmds.pointConstraint(midPiv, control)[0]
+                grpConstraint = cmds.pointConstraint(midPiv, ctrlGrp)[0]
+                cmds.delete([constraint, grpConstraint])
+
+                #parent control to grp
+                cmds.parent(control, ctrlGrp)
+                #cmds.setAttr(control + ".rx", -90)
+                
+                #freeze rots
+                cmds.makeIdentity(control, t = 0, r = 1, s = 0, apply = True)
+                
+                #translate down in y 13
+		cmds.setAttr(control + ".tz", -13)
+                cmds.makeIdentity(control, t = 1, r = 1, s = 1, apply = True)
+		
+		cmds.parent(control, world = True)
+		constraint = cmds.pointConstraint(control, ctrlGrp)[0]
+		cmds.delete(constraint)
+		
+		cmds.parent(control, ctrlGrp)
+                cmds.makeIdentity(control, t = 1, r = 1, s = 1, apply = True)
+		
+                #create a space switcher grp
+                spaceSwitcherFollow = cmds.duplicate(ctrlGrp, po = True, name = ctrlName + "_space_switcher_follow")[0]
+		spaceSwitcher = cmds.duplicate(ctrlGrp, po = True, name = ctrlName + "_space_switcher")[0]
+		cmds.parent(spaceSwitcher, spaceSwitcherFollow)
+                cmds.parent(ctrlGrp, spaceSwitcher)
+                cmds.parent(spaceSwitcherFollow, "ik_wrist_" + side + "_anim")
+                
+                
+                #parent ik control grps to this global control
+                for ctrl in ikCtrls:
+                    parent = cmds.listRelatives(ctrl, parent = True)[0]
+                    cmds.parent(parent, control)
+                
+                #parent constrain the master pv group to the global control
+                cmds.parentConstraint(control, masterPvGrp, mo = True)
+		
+		#lock attrs
+		for attr in [".sx", ".sy", ".sz", ".v"]:
+		    cmds.setAttr(control + attr, lock = True, keyable = False)
+                    
+
+
+
+            #clean up hand rig hierarchy
+            jointsGrp = cmds.group(empty = True, name = "hand_fk_joints_grp_" + side)
+	    handDrivenGrp = cmds.group(empty = True, name = "hand_driven_grp_" + side)
+	    handDrivenGrpMaster = cmds.group(empty = True, name = "hand_driven_grp_master_" + side)
+            fkCtrlGrp = cmds.group(empty = True, name = "fk_finger_controls_" + side + "_grp")
+
+            constraint = cmds.parentConstraint("ik_wrist_" + side + "_anim", handDrivenGrpMaster)
+	    cmds.delete(constraint)
+	    
+	    #create fk hand match node
+	    fkHandMatchNode = cmds.duplicate(handDrivenGrpMaster, po = True, name = "hand_match_loc_" + side)[0]
+	    cmds.parent(fkHandMatchNode, handDrivenGrpMaster)
+	    
+	    #find aim axis of arm chain to determine offset value
+	    vector1 = cmds.xform("driver_lowerarm_" + side, q = True, ws = True, t = True)
+	    vector2 = cmds.xform("driver_hand_" + side, q = True, ws = True, t = True)
+	    aimAxis = self.normalizeSubVector(vector1, vector2)
+	    
+	    axis = None
+	    offset = 0
+	    for item in [ ["X", ".rx"], ["-X", ".rx"], ["Y", ".ry"], ["-Y", ".ry"], ["Z", ".rz"], ["-Z", ".rz"]]:
+		if aimAxis == item[0]:
+		    axis = item[1]
+		    
+		    if item[0].find("-") == 0:
+			offset = 90
+		    else:
+			offset = -90
+			
+
+	    
+	    
+            cmds.setAttr(fkHandMatchNode + axis, offset)
+	    
+	    cmds.parent(handDrivenGrp, handDrivenGrpMaster)
+	    
+            for joint in joints:
+                cmds.parent(joint, jointsGrp)
+                
+            for control in ctrlGroups:
+                parent = cmds.listRelatives(control, parent = True)
+                if parent == None:
+                    cmds.parent(control, fkCtrlGrp)
+		    
+	    for control in metaJoints:
+		cmds.parent(control, handDrivenGrp)
+	    
+	    for control in fkOrients:
+		cmds.parent(control, handDrivenGrp)
+		
+	    for joint in ikJoints:
+		cmds.parent(joint, handDrivenGrp)
+		
+	    
+                    
+                    
+            #constrain the master grp to the fk and ik hand joints
+            fingerSysGrp = cmds.group(empty = True, name = "finger_sys_grp_" + side)
+	    cmds.parent(fkCtrlGrp, handDrivenGrp)
+            cmds.parent([jointsGrp, ikHandlesGrp, masterPvGrp, handDrivenGrpMaster], fingerSysGrp)
+	    
+            if len(modeGrps) > 0:
+                cmds.parent(modeGrps, fingerSysGrp)
+                
+            
+            constraint = cmds.parentConstraint(["fk_wrist_" + side + "_anim", "ik_hand_" + side], handDrivenGrpMaster, mo = True)[0]
+            
+            cmds.setAttr("Rig_Settings" + "." + side + "ArmMode", 0)
+            cmds.setAttr(constraint + ".fk_wrist_" + side + "_animW0", 1)
+            cmds.setAttr(constraint + ".ik_hand_" + side + "W1", 0)
+            cmds.setDrivenKeyframe([constraint + ".fk_wrist_" + side + "_animW0", constraint + ".ik_hand_" + side + "W1"], cd = "Rig_Settings" + "." + side + "ArmMode", itt = "linear", ott = "linear")
+
+            cmds.setAttr("Rig_Settings" + "." + side + "ArmMode", 1)
+            cmds.setAttr(constraint + ".fk_wrist_" + side + "_animW0", 0)
+            cmds.setAttr(constraint + ".ik_hand_" + side + "W1", 1)
+            cmds.setDrivenKeyframe([constraint + ".fk_wrist_" + side + "_animW0", constraint + ".ik_hand_" + side + "W1"], cd = "Rig_Settings" + "." + side + "ArmMode", itt = "linear", ott = "linear")
+
+
+
+            #Constrain the driver joints to the fk and ik joints
+            for fingers in [numIndexJoints, numMiddleJoints, numRingJoints, numPinkyJoints, numThumbJoints]:
+                for i in range(int(fingers[0])):
+                    
+                    driverJoint = "driver_" + fingers[1] + "_0" + str(i + 1) + "_" + side
+                    fkJoint = "rig_fk_" + fingers[1] + "_0" + str(i + 1) + "_" + side
+                    ikJoint = "rig_ik_" + fingers[1] + "_0" + str(i + 1) + "_" + side
+                    
+                    #set driven keys on constraint
+                    if cmds.objExists(fingers[1] + "_" + side + "_ik_anim"):
+                        constraint = cmds.parentConstraint([fkJoint, ikJoint], driverJoint)[0]
+                        ikCtrl = fingers[1] + "_finger_" + side + "_mode_anim"
+                        
+                        #set driven keyframes on constraint
+                        cmds.setAttr(ikCtrl + "." + "FK_IK", 0)
+                        cmds.setAttr(constraint + "." + fkJoint + "W0", 1)
+                        cmds.setAttr(constraint + "." + ikJoint + "W1", 0)
+                        cmds.setDrivenKeyframe([constraint + "." + fkJoint + "W0", constraint + "." + ikJoint + "W1"], cd = ikCtrl + "." + "FK_IK", itt = "linear", ott = "linear")
+    
+                        cmds.setAttr(ikCtrl + "." + "FK_IK", 1)
+                        cmds.setAttr(constraint + "." + fkJoint + "W0", 0)
+                        cmds.setAttr(constraint + "." + ikJoint + "W1", 1)
+                        cmds.setDrivenKeyframe([constraint + "." + fkJoint + "W0", constraint + "." + ikJoint + "W1"], cd = ikCtrl + "." + "FK_IK", itt = "linear", ott = "linear")
+    
+                        cmds.setAttr(ikCtrl + "." + "FK_IK", 0)
+                        
+                        
+                    
+                        #setup driven keys for fk/ik control visibility
+                        ikCtrl = fingers[1] + "_finger_" + side + "_mode_anim"
+                        
+                        cmds.setAttr(ikCtrl + "." + "FK_IK", 0)
+                        cmds.setAttr(fingers[1] + "_finger_fk_ctrl_1_" + side + "_grp.v" , 1)
+                        cmds.setAttr(fingers[1] + "_" + side + "_ik_anim_grp.v", 0)
+                        cmds.setDrivenKeyframe([fingers[1] + "_finger_fk_ctrl_1_" + side + "_grp.v", fingers[1] + "_" + side + "_ik_anim_grp.v"], cd = ikCtrl + "." + "FK_IK", itt = "linear", ott = "linear")
+    
+                        cmds.setAttr(ikCtrl + "." + "FK_IK", 1)
+                        cmds.setAttr(fingers[1] + "_finger_fk_ctrl_1_" + side + "_grp.v" , 0)
+                        cmds.setAttr(fingers[1] + "_" + side + "_ik_anim_grp.v", 1)
+                        cmds.setDrivenKeyframe([fingers[1] + "_finger_fk_ctrl_1_" + side + "_grp.v", fingers[1] + "_" + side + "_ik_anim_grp.v"], cd = ikCtrl + "." + "FK_IK", itt = "linear", ott = "linear")
+    
+                        cmds.setAttr(ikCtrl + "." + "FK_IK", 0)
+                    
+                        
+
+                    else:
+                        constraint = cmds.parentConstraint([fkJoint], driverJoint)[0]
+                        
+                    
+                    
+
+                    
+                    
+            #constrain the driver metacarpals(if they exist) to the ik and fk ones
+            for metacarpal in [thumbMeta, indexMeta, middleMeta, ringMeta, pinkyMeta]:
+                if metacarpal[0] == True:
+                    driverJoint = "driver_" + metacarpal[1] + "_metacarpal_" + side
+                    fkJoint = "rig_fk_" + metacarpal[1] + "_metacarpal_" + side
+                    ikJoint = "rig_ik_" + metacarpal[1] + "_metacarpal_" + side
+
+                    if cmds.objExists(metacarpal[1] + "_" + side + "_ik_anim"):
+                        constraint = cmds.parentConstraint([fkJoint, ikJoint], driverJoint)[0]
+                        ikCtrl = metacarpal[1] + "_finger_" + side + "_mode_anim"
+                    
+                        constraint = cmds.parentConstraint([fkJoint, ikJoint], driverJoint)[0]
+                        
+                        #set driven keyframes on constraint
+                        cmds.setAttr(ikCtrl + "." + "FK_IK", 0)
+                        cmds.setAttr(constraint + "." + fkJoint + "W0", 1)
+                        cmds.setAttr(constraint + "." + ikJoint + "W1", 0)
+                        cmds.setDrivenKeyframe([constraint + "." + fkJoint + "W0", constraint + "." + ikJoint + "W1"], cd = ikCtrl + "." + "FK_IK", itt = "linear", ott = "linear")
+    
+                        cmds.setAttr(ikCtrl + "." + "FK_IK", 1)
+                        cmds.setAttr(constraint + "." + fkJoint + "W0", 0)
+                        cmds.setAttr(constraint + "." + ikJoint + "W1", 1)
+                        cmds.setDrivenKeyframe([constraint + "." + fkJoint + "W0", constraint + "." + ikJoint + "W1"], cd = ikCtrl + "." + "FK_IK", itt = "linear", ott = "linear")
+    
+                        cmds.setAttr(ikCtrl + "." + "FK_IK", 0)
+                        
+                    else:
+                        constraint = cmds.parentConstraint([fkJoint], driverJoint)[0]
+                    
+                    
+            
+            
+            
+    # # # # # # # # # # # # # # # # # # # # # # # # # # # # # # # # # # # # # # # # # # # # # # # # # # # # # # # # # # # # # # # 
+    # # # # # # # # # # # # # # # # # # # # # # # # # # # # # # # # # # # # # # # # # # # # # # # # # # # # # # # # # # # # # # # 
+    # # # # # # # # # # # # # # # # # # # # # # # # # # # # # # # # # # # # # # # # # # # # # # # # # # # # # # # # # # # # # # # 
+    # # # # # # # # # # # # # # # # # # # # # # # # # # # # # # # # # # # # # # # # # # # # # # # # # # # # # # # # # # # # # # # 
+    # # # # # # # # # # # # # # # # # # # # # # # # # # # # # # # # # # # # # # # # # # # # # # # # # # # # # # # # # # # # # # # 
+    def buildToes(self):
+        
+        #find out which toe joints need to be rigged
+        for side in ["l", "r"]:
+            
+            
+            #create a list to hold all ctrl groups that are created
+            ctrlGroups = []
+            ikGrps = []
+            joints = []
+            
+            
+            if cmds.objExists("driver_ball_" + side):
+                children = cmds.listRelatives("driver_ball_" + side, children = True, type = 'joint')
+                allToes = cmds.listRelatives("driver_ball_" + side, allDescendents = True, type = 'joint')
+                
+                #find out how many toe joints we have for each toe
+                bigToeMeta = [False, None]
+                indexMeta = [False, None]
+                middleMeta = [False, None]
+                ringMeta = [False, None]
+                pinkyMeta = [False, None]
+                numBigToes = 0
+                numIndexToes = [0, "index"]
+                numMiddleToes = [0, "middle"]
+                numRingToes = [0, "ring"]
+                numPinkyToes = [0, "pinky"]
+                
+                if allToes:
+                    for toe in allToes:
+                        #find if metatarsals exist
+                        if toe.find("meta") != -1:
+                            if toe.partition("driver_")[2].find("bigtoe") == 0:
+                                bigToeMeta = [True, "bigtoe"]
+                            if toe.partition("driver_")[2].find("index") == 0:
+                                indexMeta = [True, "index"]
+                            if toe.partition("driver_")[2].find("middle") == 0:
+                                middleMeta = [True, "middle"]
+                            if toe.partition("driver_")[2].find("ring") == 0:
+                                ringMeta = [True, "ring"]
+                            if toe.partition("driver_")[2].find("pinky") == 0:
+                                pinkyMeta = [True, "pinky"]
+                                
+                        #get num toes -meta
+                        if toe.partition("driver_")[2].find("bigtoe") == 0:
+                            numBigToes += 1
+                        if toe.partition("driver_")[2].find("index") == 0:
+                            numIndexToes[0] += 1
+                        if toe.partition("driver_")[2].find("middle") == 0:
+                            numMiddleToes[0] += 1
+                        if toe.partition("driver_")[2].find("ring") == 0:
+                            numRingToes[0] += 1
+                        if toe.partition("driver_")[2].find("pinky") == 0:
+                            numPinkyToes[0] += 1
+                    
+                
+                    #subtract metatarsals (only if they exist!)
+                    if bigToeMeta[0] == True:
+                        numBigToes -= 1
+                    if indexMeta[0] == True:
+                        numIndexToes[0] -= 1
+                    if middleMeta[0] == True:
+                        numMiddleToes[0] -= 1
+                    if ringMeta[0] == True:
+                        numRingToes[0] -= 1
+                    if pinkyMeta[0] == True:
+                        numPinkyToes[0] -= 1
+                    
+                
+  
+  
+                #duplicate the driver joints to be used as the rig joints
+                if children:
+                    for child in children:
+                        dupeChildNodes = cmds.duplicate(child, name = "temp")
+                        
+                        #parent root joint of each toe to world if not already child of world
+                        parent = cmds.listRelatives(dupeChildNodes[0], parent = True)[0]
+                        if parent != None:
+                            cmds.parent(dupeChildNodes[0], world = True)
+                            
+                        #rename duped joints
+                        for node in dupeChildNodes:
+                            if node == "temp":
+                                niceName = child.partition("driver_")[2]
+                                joint = cmds.rename(node, "rig_" + niceName)
+                                joints.append(joint)
+                                
+                            else:
+                                niceName = node.partition("driver_")[2]
+                                cmds.rename("rig_*|" + node, "rig_" + niceName)
+                            
+                
+                
+                
+                #if the metacarpal toes exist, create a control for them
+                for meta in [bigToeMeta, indexMeta, middleMeta, ringMeta, pinkyMeta]:
+                    if meta[0] == True:
+                        #create the control object for the metacarpal
+                        ctrlName = meta[1]
+                        ctrlName = ctrlName + "_metatarsal_ctrl_" + side
+                        control = self.createControl("square", 1, ctrlName)
+                        
+                        constraint = cmds.parentConstraint("rig_" + meta[1] + "_metatarsal_" + side, control)[0]
+                        cmds.delete(constraint)
+                        cmds.makeIdentity(control, t = 0, r = 1, s = 0, apply = True)
+                        cmds.setAttr(control + ".rz", -90)
+                        cmds.setAttr(control + ".sz", 15)
+                        
+                        #create the group node and parent ctrl to it
+                        ctrlGrp = cmds.group(empty = True, name = ctrlName + "_grp")
+                        ctrlGroups.append(ctrlGrp)
+                        constraint = cmds.parentConstraint("rig_" + meta[1] + "_metatarsal_" + side, ctrlGrp)[0]
+                        cmds.delete(constraint)
+                        cmds.parent(control, ctrlGrp)
+                        cmds.makeIdentity(control, t = 1, r = 1, s = 1, apply = True)
+                        cmds.setAttr(control + ".sz", 0)
+                        cmds.makeIdentity(control, t = 1, r = 1, s = 1, apply = True)
+                        
+                        #parent the rig joint to the control
+                        cmds.parentConstraint(control, "rig_" + meta[1] + "_metatarsal_" + side, mo = True)
+                        
+                        
+                        #lock attrs on control that shouldn't be animated
+                        cmds.setAttr(control + ".tx", lock = True, keyable = False)
+                        cmds.setAttr(control + ".ty", lock = True, keyable = False)
+                        cmds.setAttr(control + ".tz", lock = True, keyable = False)
+                        cmds.setAttr(control + ".sx", lock = True, keyable = False)
+                        cmds.setAttr(control + ".sy", lock = True, keyable = False)
+                        cmds.setAttr(control + ".sz", lock = True, keyable = False)
+                        cmds.setAttr(control + ".v", lock = True, keyable = False)
+                        
+                        
+                        #color the controls
+                        if side == "l":
+                            color = 5
+                        
+                        else:
+                            color = 12
+                            
+                        cmds.setAttr(control + ".overrideEnabled", 1)
+                        cmds.setAttr(control + ".overrideColor", color)
+                        
+                        
+                
+                        
+                #if the number of toes(aside from metacarpals) is 1 or 2, just create fk controls for each toe and setup hierarchy
+                toeControls = []
+                for toes in [numIndexToes, numMiddleToes, numRingToes, numPinkyToes]:
+                    if toes[0] < 3:
+                        for i in range(int(toes[0])):
+                            #create an FK control per toe
+                            ctrlName = toes[1] + "_toe_fk_ctrl_" + str(i + 1) + "_" + side
+                            control = self.createControl("circle", 3, ctrlName)
+                            ctrlGrp = cmds.group(empty = True, name = control + "_grp")
+                            
+                            metaCtrl = toes[1] + "_metatarsal_ctrl_" + side
+                            
+                            if cmds.objExists(metaCtrl) == False:
+                                if (i + 1) == 1:
+                                    ctrlGroups.append(ctrlGrp)
+                            
+                            
+                            toeControls.append(control)
+                            
+                            #position control
+                            if i == 0:
+                                constraint = cmds.parentConstraint("rig_" + toes[1] + "_proximal_phalange_" + side, control)[0]
+                                grpConstraint = cmds.parentConstraint("rig_" + toes[1] + "_proximal_phalange_" + side, ctrlGrp)[0]
+                                cmds.delete([constraint, grpConstraint])
+                                cmds.parentConstraint(control, "rig_" + toes[1] + "_proximal_phalange_" + side, mo = True)
+                                
+                            if i == 1:
+                                constraint = cmds.parentConstraint("rig_" + toes[1] + "_middle_phalange_" + side, control)[0]
+                                grpConstraint = cmds.parentConstraint("rig_" + toes[1] + "_middle_phalange_" + side, ctrlGrp)[0]
+                                cmds.delete([constraint, grpConstraint])
+                                cmds.parentConstraint(control, "rig_" + toes[1] + "_middle_phalange_" + side, mo = True)
+                                
+                            if i == 2:
+                                constraint = cmds.parentConstraint("rig_" + toes[1] + "_distal_phalange_" + side, control)[0]
+                                grpConstraint = cmds.parentConstraint("rig_" + toes[1] + "_distal_phalange_" + side, ctrlGrp)[0]
+                                cmds.delete([constraint, grpConstraint])
+                                cmds.parentConstraint(control, "rig_" + toes[1] + "_distal_phalange_" + side, mo = True)
+                                
+                                
+                            
+                            cmds.makeIdentity(control, t = 0, r = 1, s = 0, apply = True)
+                            cmds.setAttr(control + ".rz", -90)
+                            
+                            #duplicate the ctrl group to create the driven group
+                            drivenGrp = cmds.duplicate(ctrlGrp, parentOnly = True, name = control + "_driven_grp")[0]
+                            ctrlGroups.append(drivenGrp)
+                            cmds.parent(drivenGrp, ctrlGrp)
+                        
+                            #parent control to grp
+                            cmds.parent(control, drivenGrp)
+                            cmds.makeIdentity(control, t = 1, r = 1, s = 1, apply = True)
+                            
+                            #if we aren't the root of the toe chain, then parent our ctrlGrp to the previous fk control
+                            if i != 0:
+                                cmds.parent(ctrlGrp, ctrlParent)
+                                
+                            else:
+                                if cmds.objExists(toes[1] + "_metatarsal_ctrl_" + side):
+                                    cmds.parent(ctrlGrp, toes[1] + "_metatarsal_ctrl_" + side)
+                                
+                            
+                            ctrlParent = control
+                            
+                            
+                            #lock attrs on control that shouldn't be animated
+                            cmds.setAttr(control + ".tx", lock = True, keyable = False)
+                            cmds.setAttr(control + ".ty", lock = True, keyable = False)
+                            cmds.setAttr(control + ".tz", lock = True, keyable = False)
+                            cmds.setAttr(control + ".sx", lock = True, keyable = False)
+                            cmds.setAttr(control + ".sy", lock = True, keyable = False)
+                            cmds.setAttr(control + ".sz", lock = True, keyable = False)
+                            cmds.setAttr(control + ".v", lock = True, keyable = False)
+                            
+     
+                        
+                        
+                    #if the number of toes(aside from metacarpals) is 3, setup a singular rp IK chain, and a SC IK chain for toe 3 and a newly created toe tip joint
+                    else:
+                        
+                        #take the end joint and duplicate it
+                        tipJoint = cmds.duplicate("rig_" + toes[1] + "_distal_phalange_" + side, parentOnly = True, name = "rig_" + toes[1] + "_tip_" + side)[0]
+                        cmds.parent(tipJoint, "rig_" + toes[1] + "_distal_phalange_" + side)
+                        
+                        #move tip joint out a bit
+                        if side == "r":
+                            cmds.setAttr(tipJoint + ".tx", -5)
+                            
+                        else:
+                            cmds.setAttr(tipJoint + ".tx", 5)
+                        
+                        
+                        #create RP IK handle from base knuckle to distal
+                        toeRpIkNodes = cmds.ikHandle(sol = "ikRPsolver", name = toes[1] + "_RP_ikHandle_" + side, sj = "rig_" + toes[1] + "_proximal_phalange_" + side, ee = "rig_" + toes[1] + "_distal_phalange_" + side)
+                        toeScIkNodes = cmds.ikHandle(sol = "ikSCsolver", name = toes[1] + "_SC_ikHandle_" + side, sj = "rig_" + toes[1] + "_distal_phalange_" + side, ee = tipJoint)
+                        cmds.setAttr(toeRpIkNodes[0] + ".v", 0)
+                        cmds.setAttr(toeScIkNodes[0] + ".v", 0)
+                        
+                        #parent SC IK to RP IK
+                        cmds.parent(toeScIkNodes[0], toeRpIkNodes[0])
+                        
+                        #create an IK control
+                        control = self.createControl("circle", 3, toes[1] + "_ik_ctrl_" + side)
+                        ctrlGrp = cmds.group(empty = True, name = control + "_grp")
+                        ikGrps.append(ctrlGrp)
+                        toeControls.append(control)
+                            
+                        #position control
+                        constraint = cmds.parentConstraint("rig_" + toes[1] + "_distal_phalange_" + side, control)[0]
+                        grpConstraint = cmds.parentConstraint("rig_" + toes[1] + "_distal_phalange_" + side, ctrlGrp)[0]
+                        cmds.delete([constraint, grpConstraint])
+                        
+                        #create dummy group so IK controls on both sides behave the same (dummy group will have 180 offset if right side)
+                        dummyGrp = cmds.duplicate(ctrlGrp, parentOnly = True, name = ctrlGrp + "_dummy")[0]
+			spaceSwitcherFollow = cmds.duplicate(ctrlGrp, parentOnly = True, name = ctrlGrp + "_space_switcher_follow")[0]
+                        spaceSwitcher = cmds.duplicate(ctrlGrp, parentOnly = True, name = ctrlGrp + "_space_switcher")[0]
+			cmds.parent(spaceSwitcher, spaceSwitcherFollow)
+                        cmds.parent(dummyGrp, spaceSwitcher)
+                        cmds.parent(spaceSwitcherFollow, ctrlGrp)
+                        
+                        #parent ctrl to group
+                        if side == "r":
+                            cmds.setAttr(dummyGrp + ".ry", 180)
+                            
+                        cmds.parent(control, dummyGrp)
+                        cmds.setAttr(control + ".ry", -90)
+                        cmds.makeIdentity(control, t = 1, r = 1, s = 1, apply = True)
+                        cmds.parent(toeRpIkNodes[0], control)
+                        
+                        
+                        
+                        
+                        #lock attrs on control that should not be animated
+                        for attr in [".sx", ".sy", ".sz", ".v"]:
+                            cmds.setAttr(control + attr, lock = True, keyable = False)
+                        
+                        
+                        
+                        
+                        
+                
+                #need to do the bigToe separately since it will only have 3 toes max anyway, and bigToesNum == 2 will mean IK setup, and anything less than 2 == FK setup
+                
+                if numBigToes < 2:
+                    for i in range(int(numBigToes)):
+                        #create an FK control per toe
+                        ctrlName = "bigtoe_toe_fk_ctrl_" + str(i + 1) + "_" + side
+                        control = self.createControl("circle", 8, ctrlName)
+                        ctrlGrp = cmds.group(empty = True, name = control + "_grp")
+                        
+                        toeControls.append(control)
+                        metaCtrl = "bigtoe_metatarsal_ctrl_" + side
+                        
+                        if cmds.objExists(metaCtrl) == False:
+                            if (i + 1) == 1:
+                                ctrlGroups.append(ctrlGrp)
+                                    
+                        #position control
+                        if i == 0:
+                            constraint = cmds.parentConstraint("rig_bigtoe_proximal_phalange_" + side, control)[0]
+                            grpConstraint = cmds.parentConstraint("rig_bigtoe_proximal_phalange_" + side, ctrlGrp)[0]
+                            cmds.delete([constraint, grpConstraint])
+                            cmds.parentConstraint(control, "rig_bigtoe_proximal_phalange_" + side, mo = True)
+                        
+                        if i == 1:
+                            constraint = cmds.parentConstraint("rig_bigtoe_distal_phalange_" + side, control)[0]
+                            grpConstraint = cmds.parentConstraint("rig_bigtoe_distal_phalange_" + side, ctrlGrp)[0]
+                            cmds.delete([constraint, grpConstraint])
+                            cmds.parentConstraint(control, "rig_bigtoe_distal_phalange_" + side, mo = True)
+                            
+                        cmds.makeIdentity(control, t = 0, r = 1, s = 0, apply = True)
+                        cmds.setAttr(control + ".rz", -90)
+                        
+                        #duplicate the ctrl group to create the driven group
+                        drivenGrp = cmds.duplicate(ctrlGrp, parentOnly = True, name = control + "_driven_grp")[0]
+                        ctrlGroups.append(drivenGrp)
+                        cmds.parent(drivenGrp, ctrlGrp)
+                        
+                        #parent control to grp
+                        cmds.parent(control, drivenGrp)
+                        cmds.makeIdentity(control, t = 1, r = 1, s = 1, apply = True)
+                        
+                        #if we aren't the root of the toe chain, then parent our ctrlGrp to the previous fk control
+                        if i != 0:
+                            cmds.parent(ctrlGrp, ctrlParent)
+                            
+                        else:
+                            if cmds.objExists("bigtoe_metatarsal_ctrl_" + side):
+                                cmds.parent(ctrlGrp, "bigtoe_metatarsal_ctrl_" + side)
+                                
+
+                            
+                            
+                        ctrlParent = control
+                        
+                        
+                        #lock attrs on control that shouldn't be animated
+                        cmds.setAttr(control + ".tx", lock = True, keyable = False)
+                        cmds.setAttr(control + ".ty", lock = True, keyable = False)
+                        cmds.setAttr(control + ".tz", lock = True, keyable = False)
+                        cmds.setAttr(control + ".sx", lock = True, keyable = False)
+                        cmds.setAttr(control + ".sy", lock = True, keyable = False)
+                        cmds.setAttr(control + ".sz", lock = True, keyable = False)
+                        cmds.setAttr(control + ".v", lock = True, keyable = False)
+                    
+                else:
+                    #take the end joint and duplicate it
+                    tipJoint = cmds.duplicate("rig_bigtoe_distal_phalange_" + side, parentOnly = True, name = "rig_bigtoe_tip_" + side)[0]
+                    cmds.parent(tipJoint, "rig_bigtoe_distal_phalange_" + side)
+                    
+                    #since the toe has 1 less knuckle, we need another tip
+                    tipJointEnd = cmds.duplicate(tipJoint, parentOnly = True, name = "rig_bigtoe_tip_end_" + side)[0]
+                    cmds.parent(tipJointEnd, tipJoint)
+                    
+                    #move tip joint out a bit
+                    if side == "r":
+                        cmds.setAttr(tipJoint + ".tx", -5)
+                        cmds.setAttr(tipJointEnd + ".tx", -5)
+                        
+                    else:
+                        cmds.setAttr(tipJoint + ".tx", 5)
+                        cmds.setAttr(tipJointEnd + ".tx", 5)
+                        
+                    #set preferred angles on 1rst and 2nd knuckle
+                    cmds.setAttr("rig_bigtoe_proximal_phalange_" + side + ".preferredAngleZ", 45)
+                    cmds.setAttr("rig_bigtoe_distal_phalange_" + side + ".preferredAngleZ", -45)
+                    
+                    #create RP IK handle from base knuckle to distal
+                    toeRpIkNodes = cmds.ikHandle(sol = "ikRPsolver", name = "bigtoe_RP_ikHandle_" + side, sj = "rig_bigtoe_proximal_phalange_" + side, ee = tipJoint)
+                    toeScIkNodes = cmds.ikHandle(sol = "ikSCsolver", name = "bigtoe_SC_ikHandle_" + side, sj = "rig_bigtoe_distal_phalange_" + side, ee = tipJointEnd)
+                    cmds.setAttr(toeRpIkNodes[0] + ".v", 0)
+                    cmds.setAttr(toeScIkNodes[0] + ".v", 0)
+                    
+                    #parent SC IK to RP IK
+                    cmds.parent(toeScIkNodes[0], toeRpIkNodes[0])
+                    
+                    #create an IK control
+                    control = self.createControl("circle", 6, "bigtoe_ik_ctrl_" + side)
+                    ctrlGrp = cmds.group(empty = True, name = control + "_grp")
+                    ikGrps.append(ctrlGrp)
+                    toeControls.append(control)
+                    
+                    #position control
+                    constraint = cmds.parentConstraint("rig_bigtoe_distal_phalange_" + side, control)[0]
+                    grpConstraint = cmds.parentConstraint("rig_bigtoe_distal_phalange_" + side, ctrlGrp)[0]
+                    cmds.delete([constraint, grpConstraint])
+                    
+                    #create dummy group so IK controls on both sides behave the same (dummy group will have 180 offset if right side)
+                    dummyGrp = cmds.duplicate(ctrlGrp, parentOnly = True, name = ctrlGrp + "_dummy")[0]
+		    spaceSwitcherFollow = cmds.duplicate(ctrlGrp, parentOnly = True, name = ctrlGrp + "_space_switcher_follow")[0]
+                    spaceSwitcher = cmds.duplicate(ctrlGrp, parentOnly = True, name = ctrlGrp + "_space_switcher")[0]
+		    cmds.parent(spaceSwitcher, spaceSwitcherFollow)
+                    cmds.parent(dummyGrp, spaceSwitcher)
+                    cmds.parent(spaceSwitcherFollow, ctrlGrp)
+                    
+                    #parent ctrl to group
+                    if side == "r":
+                        cmds.setAttr(dummyGrp + ".ry", 180)
+                        
+                    cmds.parent(control, dummyGrp)
+                    cmds.setAttr(control + ".ry", -90)
+                    cmds.makeIdentity(control, t = 1, r = 1, s = 1, apply = True)
+                    cmds.parent(toeRpIkNodes[0], control)
+                    
+                    
+                    
+                    
+                    #lock attrs on control that should not be animated
+                    for attr in [".sx", ".sy", ".sz", ".v"]:
+                        cmds.setAttr(control + attr, lock = True, keyable = False)
+            
+            
+            
+            
+                #setup global control for toes
+                
+                #if toes < 3, rotation based(spread, curl) (grp node above FK control)
+                fkControls = []
+                ikControls = []
+                for control in toeControls:
+
+                    if control.find("fk") != -1:
+                        fkControls.append(control)
+                    
+                    if control.find("ik") != -1:
+                        ikControls.append(control)
+                       
+                       
+                       
+                        
+                        
+                
+                #IF AlL TOES ARE IK TOES
+                if len(fkControls) == 0:
+                    #create a control at the tip of the toes that will globally move the IK controls
+                    
+                    if joints:
+                        
+                        control = self.createControl("square", 1, "ik_global_ctrl_" + side)
+                        constraint = cmds.parentConstraint("ikHandle_toe_" + side, control)[0]
+                        cmds.delete(constraint)
+                        cmds.makeIdentity(control, t = 0, r = 1, s = 0, apply = True)
+                        cmds.setAttr(control + ".rz", -90)
+                        cmds.setAttr(control + ".sz", 5)
+                        
+                        #add a spread attr to the global control
+                        cmds.select(control)
+                        cmds.addAttr(longName='spread', defaultValue=0, minValue=0, maxValue=10, keyable = True)
+                        
+                            
+                        #create the group node and parent ctrl to it
+                        ctrlGrp = cmds.group(empty = True, name = "ik_global_ctrl_" + side + "_grp")
+                        ctrlGroups.append(ctrlGrp)
+                        constraint = cmds.parentConstraint("ikHandle_toe_" + side, ctrlGrp)[0]
+                        cmds.delete(constraint)
+                        
+                        #create a space switcher node
+			spaceSwitcherFollow = cmds.duplicate(ctrlGrp, parentOnly = True, name = ctrlGrp + "_space_switcher_follow")[0]
+                        spaceSwitcher = cmds.duplicate(ctrlGrp, parentOnly = True, name = ctrlGrp + "_space_switcher")[0]
+			cmds.parent(spaceSwitcher, spaceSwitcherFollow)
+                        cmds.parent(spaceSwitcherFollow, ctrlGrp)
+                        
+                        
+                        cmds.parent(control, spaceSwitcher)
+                        cmds.makeIdentity(control, t = 1, r = 1, s = 1, apply = True)
+                        cmds.setAttr(control + ".rx", 90)
+                        cmds.setAttr(control + ".scale", 2.5, 2.5, 2.5, type = 'double3')
+                        cmds.makeIdentity(control, t = 1, r = 1, s = 1, apply = True)
+                        
+                        
+                        #set the pivot to be at the base of the toes
+                        pivPos = cmds.xform("jointmover_knuckle_base_" + side, q = True, ws = True, t = True)
+                        cmds.xform(control, ws = True, piv = (pivPos[0], pivPos[1], pivPos[2]))
+                        cmds.xform(ctrlGrp, ws = True, piv = (pivPos[0], pivPos[1], pivPos[2]))
+                        cmds.xform(spaceSwitcher, ws = True, piv = (pivPos[0], pivPos[1], pivPos[2]))
+                        
+                        
+                        #lock attrs on control that are not needed to be animated
+                        cmds.setAttr(control + ".sx", lock = True, keyable = False)
+                        cmds.setAttr(control + ".sy", lock = True, keyable = False)
+                        cmds.setAttr(control + ".sz", lock = True, keyable = False)
+                        cmds.setAttr(control + ".v", lock = True, keyable = False)
+                        
+                        #parent the IK groups to the global grp and also set driven keys for toe spread
+                        for grp in ikGrps:
+                            #create a driven group
+                            group = cmds.group(empty = True, name = grp + "_driven")
+                            
+                            if grp.find("index") != -1:
+                                constraint = cmds.parentConstraint("rig_index_proximal_phalange_" + side, group)[0]
+                                cmds.delete(constraint)
+                                cmds.parent(group, "index_ik_ctrl_" + side + "_grp_space_switcher")
+                                cmds.parent("index_ik_ctrl_" + side + "_grp_dummy", group)
+                                
+                                cmds.setAttr(control + ".spread", 0)
+                                cmds.setAttr(group + ".ry", 0)
+                                cmds.setDrivenKeyframe(group + ".ry", cd = control + ".spread", itt = "linear", ott = "linear")
+                                
+                                cmds.setAttr(control + ".spread", 10)
+                                cmds.setAttr(group + ".ry", -9)
+                                cmds.setDrivenKeyframe(group + ".ry", cd = control + ".spread", itt = "linear", ott = "linear")
+                                
+                                cmds.setAttr(control + ".spread", 0)
+                                cmds.setAttr(group + ".ry", 0)
+                                
+                                
+                            if grp.find("middle") != -1:
+                                constraint = cmds.parentConstraint("rig_middle_proximal_phalange_" + side, group)[0]
+                                cmds.delete(constraint)
+                                cmds.parent(group, "middle_ik_ctrl_" + side + "_grp_space_switcher")
+                                cmds.parent("middle_ik_ctrl_" + side + "_grp_dummy", group)
+                                
+                                cmds.setAttr(control + ".spread", 0)
+                                cmds.setAttr(group + ".ry", 0)
+                                cmds.setDrivenKeyframe(group + ".ry", cd = control + ".spread", itt = "linear", ott = "linear")
+                                
+                                cmds.setAttr(control + ".spread", 10)
+                                cmds.setAttr(group + ".ry", 9.5)
+                                cmds.setDrivenKeyframe(group + ".ry", cd = control + ".spread", itt = "linear", ott = "linear")
+                                
+                                cmds.setAttr(control + ".spread", 0)
+                                cmds.setAttr(group + ".ry", 0)
+                                
+                                
+                            
+                            if grp.find("ring") != -1:
+                                constraint = cmds.parentConstraint("rig_ring_proximal_phalange_" + side, group)[0]
+                                cmds.delete(constraint)
+                                cmds.parent(group, "ring_ik_ctrl_" + side + "_grp_space_switcher")
+                                cmds.parent("ring_ik_ctrl_" + side + "_grp_dummy", group)
+                                
+                                cmds.setAttr(control + ".spread", 0)
+                                cmds.setAttr(group + ".ry", 0)
+                                cmds.setDrivenKeyframe(group + ".ry", cd = control + ".spread", itt = "linear", ott = "linear")
+                                
+                                cmds.setAttr(control + ".spread", 10)
+                                cmds.setAttr(group + ".ry", 17)
+                                cmds.setDrivenKeyframe(group + ".ry", cd = control + ".spread", itt = "linear", ott = "linear")
+                                
+                                cmds.setAttr(control + ".spread", 0)
+                                cmds.setAttr(group + ".ry", 0)
+                                
+                                
+                                
+                            if grp.find("pinky") != -1:
+                                constraint = cmds.parentConstraint("rig_pinky_proximal_phalange_" + side, group)[0]
+                                cmds.delete(constraint)
+                                cmds.parent(group, "pinky_ik_ctrl_" + side + "_grp_space_switcher")
+                                cmds.parent("pinky_ik_ctrl_" + side + "_grp_dummy", group)
+                                
+                                
+                                cmds.setAttr(control + ".spread", 0)
+                                cmds.setAttr(group + ".ry", 0)
+                                cmds.setDrivenKeyframe(group + ".ry", cd = control + ".spread", itt = "linear", ott = "linear")
+                                
+                                cmds.setAttr(control + ".spread", 10)
+                                cmds.setAttr(group + ".ry", 32)
+                                cmds.setDrivenKeyframe(group + ".ry", cd = control + ".spread", itt = "linear", ott = "linear")
+                                
+                                cmds.setAttr(control + ".spread", 0)
+                                cmds.setAttr(group + ".ry", 0)
+                                
+                                
+                            if grp.find("bigtoe") != -1:
+                                constraint = cmds.parentConstraint("rig_bigtoe_proximal_phalange_" + side, group)[0]
+                                cmds.delete(constraint)
+                                cmds.parent(group, "bigtoe_ik_ctrl_" + side + "_grp_space_switcher")
+                                cmds.parent("bigtoe_ik_ctrl_" + side + "_grp_dummy", group)
+                                
+                                
+                                cmds.setAttr(control + ".spread", 0)
+                                cmds.setAttr(group + ".ry", 0)
+                                cmds.setDrivenKeyframe(group + ".ry", cd = control + ".spread", itt = "linear", ott = "linear")
+                                
+                                cmds.setAttr(control + ".spread", 10)
+                                cmds.setAttr(group + ".ry", -15)
+                                cmds.setDrivenKeyframe(group + ".ry", cd = control + ".spread", itt = "linear", ott = "linear")
+                                
+                                cmds.setAttr(control + ".spread", 0)
+                                cmds.setAttr(group + ".ry", 0)
+                            
+                            
+                            cmds.parent(grp, control)
+                            
+                        #color the control
+                        if side == "l":
+                            color = 5
+                        
+                        else:
+                            color = 12
+                            
+                        cmds.setAttr(control + ".overrideEnabled", 1)
+                        cmds.setAttr(control + ".overrideColor", color)
+                
+                
+                
+                
+                
+                #IF ALL TOES ARE FK TOES
+                if len(ikControls) == 0:
+                    if joints:
+                        
+                        #create a control at the tip of the toes that will give the user some handy global controls, like curl, spread, etc.
+                        control = self.createControl("square", 1, "fk_global_ctrl_" + side)
+                        constraint = cmds.parentConstraint("ikHandle_toe_" + side, control)[0]
+                        cmds.delete(constraint)
+                        cmds.makeIdentity(control, t = 0, r = 1, s = 0, apply = True)
+                        cmds.setAttr(control + ".rz", -90)
+                        cmds.setAttr(control + ".sz", 5)
+                            
+                        #create the group node and parent ctrl to it
+                        ctrlGrp = cmds.group(empty = True, name = "fk_global_ctrl_" + side + "_grp")
+                        ctrlGroups.append(ctrlGrp)
+                        constraint = cmds.parentConstraint("ikHandle_toe_" + side, ctrlGrp)[0]
+                        cmds.delete(constraint)
+                        cmds.parent(control, ctrlGrp)
+                        cmds.makeIdentity(control, t = 1, r = 1, s = 1, apply = True)
+                        cmds.setAttr(control + ".rx", 90)
+                        cmds.setAttr(control + ".scale", 2.5, 2.5, 2.5, type = 'double3')
+                        cmds.makeIdentity(control, t = 1, r = 1, s = 1, apply = True)
+                        
+                        
+                        #lock attrs on control that are not needed to be animated
+                        cmds.setAttr(control + ".tx", lock = True, keyable = False)
+                        cmds.setAttr(control + ".ty", lock = True, keyable = False)
+                        cmds.setAttr(control + ".tz", lock = True, keyable = False)
+                        cmds.setAttr(control + ".sx", lock = True, keyable = False)
+                        cmds.setAttr(control + ".sy", lock = True, keyable = False)
+                        cmds.setAttr(control + ".sz", lock = True, keyable = False)
+                        cmds.setAttr(control + ".v", lock = True, keyable = False)
+                        
+                        
+                        #find all driven fk grps
+                        drivenGroups = []
+                        for group in ctrlGroups:
+                            if group.find("driven") != -1:
+                                if group.find("fk") != -1:
+                                    drivenGroups.append(group)
+                        
+                                    
+                                    
+                        #add a spread attr to the global control
+                        cmds.select(control)
+                        cmds.addAttr(longName='spread', defaultValue=0, minValue=0, maxValue=10, keyable = True)
+    
+                        #setup driven keys for the fk group nodes
+                        for group in drivenGroups:
+                            
+                            
+                            #Curl
+                            cmds.setAttr(control + ".rz", 0)
+                            cmds.setAttr(group + ".rz", 0)
+                            cmds.setDrivenKeyframe(group + ".rz", cd = control + ".rz", itt = "linear", ott = "linear")
+                            
+                            cmds.setAttr(control + ".rz", -180)
+                            cmds.setAttr(group + ".rz", -90)
+                            cmds.setDrivenKeyframe(group + ".rz", cd = control + ".rz", itt = "linear", ott = "linear")
+                            
+                            cmds.setAttr(control + ".rz", 90)
+                            cmds.setAttr(group + ".rz", 45)
+                            cmds.setDrivenKeyframe(group + ".rz", cd = control + ".rz", itt = "linear", ott = "linear")
+                            
+                            cmds.setAttr(control + ".rz", 0)
+                            cmds.setAttr(group + ".rz", 0)
+                            
+                            
+                            #Toe Lean
+                            
+                            #only the base knuckle
+                            if group.partition("ctrl_")[2].find("1") == 0:
+                                
+                                cmds.setAttr(control + ".ry", 0)
+                                cmds.setAttr(group + ".ry", 0)
+                                cmds.setDrivenKeyframe(group + ".ry", cd = control + ".ry", itt = "linear", ott = "linear")
+                                
+                                cmds.setAttr(control + ".ry", 45)
+                                cmds.setAttr(group + ".ry", 45)
+                                cmds.setDrivenKeyframe(group + ".ry", cd = control + ".ry", itt = "linear", ott = "linear")
+                                
+                                cmds.setAttr(control + ".ry", -45)
+                                cmds.setAttr(group + ".ry", -45)
+                                cmds.setDrivenKeyframe(group + ".ry", cd = control + ".ry", itt = "linear", ott = "linear")
+                                
+                                cmds.setAttr(control + ".ry", 0)
+                                cmds.setAttr(group + ".ry", 0)
+                                
+    
+                            #Toe Tilt
+                            
+                            if group.find("pinky") != -1:
+                                if group.find("1") != -1:
+                                    cmds.setAttr(control + ".rx", 0)
+                                    cmds.setAttr(group + ".rz", 0)
+                                    cmds.setDrivenKeyframe(group + ".rz", cd = control + ".rx", itt = "linear", ott = "linear")
+                                    
+                                    cmds.setAttr(control + ".rx", 60)
+                                    cmds.setAttr(group + ".rz", 65)
+                                    cmds.setDrivenKeyframe(group + ".rz", cd = control + ".rx", itt = "linear", ott = "linear")
+                                    
+                                    cmds.setAttr(control + ".rx", -60)
+                                    cmds.setAttr(group + ".rz", -65)
+                                    cmds.setDrivenKeyframe(group + ".rz", cd = control + ".rx", itt = "linear", ott = "linear")
+                                    
+                                    cmds.setAttr(control + ".rx", 0)
+                                    cmds.setAttr(group + ".rz", 0)
+                                    
+                                    
+                                if group.find("2") != -1:
+                                    cmds.setAttr(control + ".rx", 0)
+                                    cmds.setAttr(group + ".rz", 0)
+                                    cmds.setDrivenKeyframe(group + ".rz", cd = control + ".rx", itt = "linear", ott = "linear")
+                                    
+                                    cmds.setAttr(control + ".rx", 60)
+                                    cmds.setAttr(group + ".rz", -60)
+                                    cmds.setDrivenKeyframe(group + ".rz", cd = control + ".rx", itt = "linear", ott = "linear")
+                                    
+                                    cmds.setAttr(control + ".rx", -60)
+                                    cmds.setAttr(group + ".rz", 60)
+                                    cmds.setDrivenKeyframe(group + ".rz", cd = control + ".rx", itt = "linear", ott = "linear")
+                                    
+                                    cmds.setAttr(control + ".rx", 0)
+                                    cmds.setAttr(group + ".rz", 0)
+                                    
+    
+                            if group.find("ring") != -1:
+                                if group.find("1") != -1:
+                                    cmds.setAttr(control + ".rx", 0)
+                                    cmds.setAttr(group + ".rz", 0)
+                                    cmds.setDrivenKeyframe(group + ".rz", cd = control + ".rx", itt = "linear", ott = "linear")
+                                    
+                                    cmds.setAttr(control + ".rx", 60)
+                                    cmds.setAttr(group + ".rz", 45)
+                                    cmds.setDrivenKeyframe(group + ".rz", cd = control + ".rx", itt = "linear", ott = "linear")
+                                    
+                                    cmds.setAttr(control + ".rx", -60)
+                                    cmds.setAttr(group + ".rz", -45)
+                                    cmds.setDrivenKeyframe(group + ".rz", cd = control + ".rx", itt = "linear", ott = "linear")
+                                    
+                                    cmds.setAttr(control + ".rx", 0)
+                                    cmds.setAttr(group + ".rz", 0)
+                                    
+                                    
+                                if group.find("2") != -1:
+                                    cmds.setAttr(control + ".rx", 0)
+                                    cmds.setAttr(group + ".rz", 0)
+                                    cmds.setDrivenKeyframe(group + ".rz", cd = control + ".rx", itt = "linear", ott = "linear")
+                                    
+                                    cmds.setAttr(control + ".rx", 60)
+                                    cmds.setAttr(group + ".rz", -30)
+                                    cmds.setDrivenKeyframe(group + ".rz", cd = control + ".rx", itt = "linear", ott = "linear")
+                                    
+                                    cmds.setAttr(control + ".rx", -60)
+                                    cmds.setAttr(group + ".rz", 30)
+                                    cmds.setDrivenKeyframe(group + ".rz", cd = control + ".rx", itt = "linear", ott = "linear")
+                                    
+                                    cmds.setAttr(control + ".rx", 0)
+                                    cmds.setAttr(group + ".rz", 0)
+                                    
+                                    
+                            if group.find("middle") != -1:
+                                if group.find("1") != -1:
+                                    cmds.setAttr(control + ".rx", 0)
+                                    cmds.setAttr(group + ".rz", 0)
+                                    cmds.setDrivenKeyframe(group + ".rz", cd = control + ".rx", itt = "linear", ott = "linear")
+                                    
+                                    cmds.setAttr(control + ".rx", 60)
+                                    cmds.setAttr(group + ".rz", 50)
+                                    cmds.setDrivenKeyframe(group + ".rz", cd = control + ".rx", itt = "linear", ott = "linear")
+                                    
+                                    cmds.setAttr(control + ".rx", -60)
+                                    cmds.setAttr(group + ".rz", -50)
+                                    cmds.setDrivenKeyframe(group + ".rz", cd = control + ".rx", itt = "linear", ott = "linear")
+                                    
+                                    cmds.setAttr(control + ".rx", 0)
+                                    cmds.setAttr(group + ".rz", 0)
+                                    
+                                    
+                                if group.find("2") != -1:
+                                    cmds.setAttr(control + ".rx", 0)
+                                    cmds.setAttr(group + ".rz", 0)
+                                    cmds.setDrivenKeyframe(group + ".rz", cd = control + ".rx", itt = "linear", ott = "linear")
+                                    
+                                    cmds.setAttr(control + ".rx", 60)
+                                    cmds.setAttr(group + ".rz", -40)
+                                    cmds.setDrivenKeyframe(group + ".rz", cd = control + ".rx", itt = "linear", ott = "linear")
+                                    
+                                    cmds.setAttr(control + ".rx", -60)
+                                    cmds.setAttr(group + ".rz", 40)
+                                    cmds.setDrivenKeyframe(group + ".rz", cd = control + ".rx", itt = "linear", ott = "linear")
+                                    
+                                    cmds.setAttr(control + ".rx", 0)
+                                    cmds.setAttr(group + ".rz", 0)
+                                    
+                                    
+                            if group.find("index") != -1:
+                                if group.find("1") != -1:
+                                    cmds.setAttr(control + ".rx", 0)
+                                    cmds.setAttr(group + ".rz", 0)
+                                    cmds.setDrivenKeyframe(group + ".rz", cd = control + ".rx", itt = "linear", ott = "linear")
+                                    
+                                    cmds.setAttr(control + ".rx", 60)
+                                    cmds.setAttr(group + ".rz", 35)
+                                    cmds.setDrivenKeyframe(group + ".rz", cd = control + ".rx", itt = "linear", ott = "linear")
+                                    
+                                    cmds.setAttr(control + ".rx", -60)
+                                    cmds.setAttr(group + ".rz", -35)
+                                    cmds.setDrivenKeyframe(group + ".rz", cd = control + ".rx", itt = "linear", ott = "linear")
+                                    
+                                    cmds.setAttr(control + ".rx", 0)
+                                    cmds.setAttr(group + ".rz", 0)
+                                    
+                                    
+                                if group.find("2") != -1:
+                                    cmds.setAttr(control + ".rx", 0)
+                                    cmds.setAttr(group + ".rz", 0)
+                                    cmds.setDrivenKeyframe(group + ".rz", cd = control + ".rx", itt = "linear", ott = "linear")
+                                    
+                                    cmds.setAttr(control + ".rx", 60)
+                                    cmds.setAttr(group + ".rz", -25)
+                                    cmds.setDrivenKeyframe(group + ".rz", cd = control + ".rx", itt = "linear", ott = "linear")
+                                    
+                                    cmds.setAttr(control + ".rx", -60)
+                                    cmds.setAttr(group + ".rz", 25)
+                                    cmds.setDrivenKeyframe(group + ".rz", cd = control + ".rx", itt = "linear", ott = "linear")
+                                    
+                                    cmds.setAttr(control + ".rx", 0)
+                                    cmds.setAttr(group + ".rz", 0)
+                                    
+                            
+                            
+                            if group.find("bigtoe") != -1:
+                                if group.find("1") != -1:
+                                    cmds.setAttr(control + ".rx", 0)
+                                    cmds.setAttr(group + ".rz", 0)
+                                    cmds.setDrivenKeyframe(group + ".rz", cd = control + ".rx", itt = "linear", ott = "linear")
+                                    
+                                    cmds.setAttr(control + ".rx", 60)
+                                    cmds.setAttr(group + ".rz", 5)
+                                    cmds.setDrivenKeyframe(group + ".rz", cd = control + ".rx", itt = "linear", ott = "linear")
+                                    
+                                    cmds.setAttr(control + ".rx", -60)
+                                    cmds.setAttr(group + ".rz", -5)
+                                    cmds.setDrivenKeyframe(group + ".rz", cd = control + ".rx", itt = "linear", ott = "linear")
+                                    
+                                    cmds.setAttr(control + ".rx", 0)
+                                    cmds.setAttr(group + ".rz", 0)
+                                    
+                            #toe spread
+                            if group.find("bigtoe") != -1:
+                                if group.find("1") != -1:
+                                    cmds.setAttr(control + ".spread", 0)
+                                    cmds.setAttr(group + ".ry", 0)
+                                    cmds.setDrivenKeyframe(group + ".ry", cd = control + ".spread", itt = "linear", ott = "linear")
+                                    
+                                    cmds.setAttr(control + ".spread", 10)
+                                    cmds.setAttr(group + ".ry", -15)
+                                    cmds.setDrivenKeyframe(group + ".ry", cd = control + ".spread", itt = "linear", ott = "linear")
+                                    
+                                    cmds.setAttr(control + ".spread", 0)
+                                    cmds.setAttr(group + ".ry", 0)
+                                    
+                                    
+                            if group.find("index") != -1:
+                                if group.find("1") != -1:
+                                    cmds.setAttr(control + ".spread", 0)
+                                    cmds.setAttr(group + ".ry", 0)
+                                    cmds.setDrivenKeyframe(group + ".ry", cd = control + ".spread", itt = "linear", ott = "linear")
+                                    
+                                    cmds.setAttr(control + ".spread", 10)
+                                    cmds.setAttr(group + ".ry", -9)
+                                    cmds.setDrivenKeyframe(group + ".ry", cd = control + ".spread", itt = "linear", ott = "linear")
+                                    
+                                    cmds.setAttr(control + ".spread", 0)
+                                    cmds.setAttr(group + ".ry", 0)
+                                    
+                                    
+                            if group.find("middle") != -1:
+                                if group.find("1") != -1:
+                                    cmds.setAttr(control + ".spread", 0)
+                                    cmds.setAttr(group + ".ry", 0)
+                                    cmds.setDrivenKeyframe(group + ".ry", cd = control + ".spread", itt = "linear", ott = "linear")
+                                    
+                                    cmds.setAttr(control + ".spread", 10)
+                                    cmds.setAttr(group + ".ry", 9.5)
+                                    cmds.setDrivenKeyframe(group + ".ry", cd = control + ".spread", itt = "linear", ott = "linear")
+    
+                                    cmds.setAttr(control + ".spread", 0)
+                                    cmds.setAttr(group + ".ry", 0)
+                                    
+    
+                            if group.find("ring") != -1:
+                                if group.find("1") != -1:
+                                    cmds.setAttr(control + ".spread", 0)
+                                    cmds.setAttr(group + ".ry", 0)
+                                    cmds.setDrivenKeyframe(group + ".ry", cd = control + ".spread", itt = "linear", ott = "linear")
+                                    
+                                    cmds.setAttr(control + ".spread", 10)
+                                    cmds.setAttr(group + ".ry", 17)
+                                    cmds.setDrivenKeyframe(group + ".ry", cd = control + ".spread", itt = "linear", ott = "linear")
+                                    
+                                    cmds.setAttr(control + ".spread", 0)
+                                    cmds.setAttr(group + ".ry", 0)
+                                    
+                                    
+                            if group.find("pinky") != -1:
+                                if group.find("1") != -1:
+                                    cmds.setAttr(control + ".spread", 0)
+                                    cmds.setAttr(group + ".ry", 0)
+                                    cmds.setDrivenKeyframe(group + ".ry", cd = control + ".spread", itt = "linear", ott = "linear")
+                                    
+                                    cmds.setAttr(control + ".spread", 10)
+                                    cmds.setAttr(group + ".ry", 32)
+                                    cmds.setDrivenKeyframe(group + ".ry", cd = control + ".spread", itt = "linear", ott = "linear")
+    
+                                    cmds.setAttr(control + ".spread", 0)
+                                    cmds.setAttr(group + ".ry", 0)
+                        
+
+                        
+                if len(ikControls) and len(fkControls) > 0:
+                    for group in ikGrps:
+                        ctrlGroups.append(group)
+                        
+
+
+                #need to hook into foot rig, both fk and ik. To do this, we'll group up the toe controls for each side, and parent under the driver ball
+                if joints:
+                    masterGrp = cmds.group(empty = True, name = "toe_rig_" + side + "_grp")
+                    jointsGrp = cmds.group(empty = True, name = "toe_rig_joints_" + side + "_grp")
+                    cmds.parent(joints, jointsGrp)
+                    cmds.parent(jointsGrp, masterGrp)
+                    for group in ctrlGroups:
+                        if group.find("driven") == -1:
+                            cmds.parent(group, masterGrp)
+                    
+                    cmds.parentConstraint("driver_ball_" + side, masterGrp, mo = True)
+                    
+                    #parent toe groups to leg sys grp
+                    cmds.parent(masterGrp, "leg_sys_grp")
+                    
+                
+                    #Need to constrain driver joints to rig joints
+                    for toe in allToes:
+                        rigToe = toe.partition("driver_")[2]
+                        rigToe = "rig_" + rigToe
+                        
+                        cmds.parentConstraint(rigToe, toe)
+                        
+                        
+                    
+                    #color the controls
+                    for control in toeControls:
+                        if side == "l":
+                            color = 5
+                        
+                        else:
+                            color = 12
+                            
+                        cmds.setAttr(control + ".overrideEnabled", 1)
+                        cmds.setAttr(control + ".overrideColor", color)
+                        
+                        
+                    #lastly, hook up toe control visibility to foot control attribute
+                    cmds.connectAttr("ik_foot_anim_" + side + ".toeCtrlVis", masterGrp + ".v")
+                    
+                    
+            
+    # # # # # # # # # # # # # # # # # # # # # # # # # # # # # # # # # # # # # # # # # # # # # # # # # # # # # # # # # # # # # # # 
+    # # # # # # # # # # # # # # # # # # # # # # # # # # # # # # # # # # # # # # # # # # # # # # # # # # # # # # # # # # # # # # # 
+    # # # # # # # # # # # # # # # # # # # # # # # # # # # # # # # # # # # # # # # # # # # # # # # # # # # # # # # # # # # # # # # 
+    # # # # # # # # # # # # # # # # # # # # # # # # # # # # # # # # # # # # # # # # # # # # # # # # # # # # # # # # # # # # # # # 
+    # # # # # # # # # # # # # # # # # # # # # # # # # # # # # # # # # # # # # # # # # # # # # # # # # # # # # # # # # # # # # # # 
+    def finishLegs(self):
+        ball = False
+        
+
+        for side in ["l", "r"]:
+
+            #organize joints
+            legJointGrp = cmds.group(empty = True, name = "leg_joints_grp_" + side)
+            constraint = cmds.parentConstraint("driver_thigh_" + side, legJointGrp)[0]
+            cmds.delete(constraint)
+            
+            cmds.parent(["fk_leg_thigh_" + side, "ik_leg_thigh_" + side, "fk_thigh_" + side + "_orient_grp"], legJointGrp)
+            cmds.parent(legJointGrp, "leg_group_" + side)
+	    
+	    
+            
+            #create invisible legs that will drive the hips
+            filePath = self.mayaToolsDir +  "/General/ART/invis_legs.mb"
+            cmds.select("leg_group_" + side, replace = True)
+            cmds.file(filePath, es = True, type = "mayaBinary", force = True)
+	    
+	    
+            invisLegNodes = cmds.file(filePath, i = True, returnNewNodes = True, renameAll = True)
+            
+            #clean up import
+            for node in invisLegNodes:
+                if node.find("body_anim_space") != -1:
+                    if cmds.objExists(node):
+                        cmds.delete(node)
+
+
+            #constrain the no flip begin joint to the driver pelvis
+            cmds.parentConstraint("driver_pelvis", "noflip_begin_joint_" + side, mo = True)
+	    cmds.parentConstraint("hip_anim", "invis_legs_leg_group_" + side, mo = True)
+	    
+            #connect real knee vector to invis knee vector
+            cmds.connectAttr("noflip_pv_loc_" + side + ".translate", ("invis_legs_noflip_pv_loc_" + side + ".translate"))
+            cmds.delete("invis_legs_ik_knee_anim_grp_" + side)
+            if cmds.objExists("invis_legs_ik_leg_" + side + "_twistMultNode"):
+                cmds.delete("invis_legs_ik_leg_" + side + "_twistMultNode")
+                
+            if side == "r":
+                cmds.disconnectAttr("invis_legs_ik_foot_anim_" + side + ".knee_twist", "invis_legs_foot_ikHandle_" + side + ".twist")
+            cmds.connectAttr("foot_ikHandle_" + side + ".twist", "invis_legs_foot_ikHandle_" + side + ".twist")
+            
+            
+            #make sure leg orients are good
+            tempConstraint = cmds.orientConstraint("fk_thigh_" + side + "_orient_grp", "invis_legs_fk_thigh_" + side + "_orient_grp")[0]
+            cmds.delete(tempConstraint)
+            
+            #point constraint invis target loc to real foot control so invis IK goes with real foot. delete invis foot
+            cmds.parentConstraint("ik_foot_anim_" + side, ("invis_legs_ik_foot_anim_" + side))
+            
+            #drive invis fk thigh with real
+            cmds.connectAttr("fk_thigh_" + side + "_anim.rotate", "invis_legs_fk_thigh_" + side + "_anim.rotate")
+            cmds.connectAttr("fk_calf_" + side + "_anim.rotate", "invis_legs_fk_calf_" + side + "_anim.rotate")
+            
+            
+            #hide invisible legs
+            cmds.setAttr("invis_legs_leg_group_" + side + ".v", 0)
+            parent = cmds.listRelatives("invis_legs_leg_group_" + side, parent = True)
+            cmds.parent("invis_legs_leg_group_" + side, "leg_group_" + side)
+            
+            if parent:
+                cmds.delete(parent)
+                
+    
+            #create result joints
+            thighJoint = cmds.duplicate("driver_thigh_" + side, name = "result_leg_thigh_" + side, parentOnly = True)[0]
+            calfJoint = cmds.duplicate("driver_calf_" + side, name = "result_leg_calf_" + side, parentOnly = True)[0]
+            footJoint = cmds.duplicate("driver_foot_" + side, name = "result_leg_foot_" + side, parentOnly = True)[0]
+            
+            if cmds.objExists("driver_ball_" + side):
+                ball = True
+                ballJoint = cmds.duplicate("driver_ball_" + side, name = "result_leg_ball_" + side, parentOnly = True)[0]
+                
+            for joint in [thighJoint, calfJoint, footJoint]:
+                cmds.parent(joint, world = True)
+                
+            if ball:
+                cmds.parent(ballJoint, world = True)
+                
+
+            cmds.parent(footJoint, calfJoint)
+            cmds.parent(calfJoint, thighJoint)
+            
+            if ball:
+                cmds.parent(ballJoint, footJoint)
+                
+            cmds.makeIdentity(thighJoint, t = 0, r = 1, s = 0, apply = True)
+            
+            
+            
+            
+            #create IK fix joints so that all the orients of the leg systems match up
+            ikFixThighJoint = cmds.duplicate("driver_thigh_" + side, name = "ikFix_leg_thigh_" + side, parentOnly = True)[0]
+            ikFixCalfJoint = cmds.duplicate("driver_calf_" + side, name = "ikFix_leg_calf_" + side, parentOnly = True)[0]
+            ikFixFootJoint = cmds.duplicate("driver_foot_" + side, name = "ikFix_leg_foot_" + side, parentOnly = True)[0]
+            
+            if cmds.objExists("driver_ball_" + side):
+                ball = True
+                ikFixBallJoint = cmds.duplicate("driver_ball_" + side, name = "ikFix_leg_ball_" + side, parentOnly = True)[0]
+                
+            for joint in [ikFixThighJoint, ikFixCalfJoint, ikFixFootJoint]:
+                cmds.parent(joint, world = True)
+                
+            if ball:
+                cmds.parent(ikFixBallJoint, world = True)
+                
+
+            cmds.parent(ikFixFootJoint, ikFixCalfJoint)
+            cmds.parent(ikFixCalfJoint, ikFixThighJoint)
+            
+            if ball:
+                cmds.parent(ikFixBallJoint, ikFixFootJoint)
+                
+            cmds.makeIdentity(ikFixThighJoint, t = 0, r = 1, s = 0, apply = True)
+            
+
+            cmds.parentConstraint("ik_leg_thigh_" + side, ikFixThighJoint, mo = True)
+            cmds.parentConstraint("ik_leg_calf_" + side, ikFixCalfJoint, mo = True)
+            cmds.parentConstraint("ik_leg_foot_" + side, ikFixFootJoint, mo = True)
+            
+            
+            if ball:
+                cmds.parentConstraint("ik_leg_ball_" + side, ikFixBallJoint, mo = True)
+                
+            
+            
+
+            
+            #constrain result joints to fk and ik joints
+            thighConstraint = cmds.parentConstraint(["fk_leg_thigh_" + side, ikFixThighJoint], thighJoint, mo = True)[0]
+            calfConstraint = cmds.parentConstraint(["fk_leg_calf_" + side, ikFixCalfJoint], calfJoint, mo = True)[0]
+            footConstraint = cmds.parentConstraint(["fk_leg_foot_" + side, ikFixFootJoint], footJoint, mo = True)[0]
+            
+            if ball:
+                ballConstraint = cmds.parentConstraint(["fk_leg_ball_" + side, ikFixBallJoint], ballJoint, mo = True)[0]
+                
+            #add switch attr
+            cmds.select("Rig_Settings")
+            cmds.addAttr(longName=(side + "LegMode"), at = 'enum', en = "fk:ik:", keyable = True)
+            
+            #connect up attr to constraints
+            constraints =[thighConstraint, calfConstraint, footConstraint]
+            if ball:
+                constraints.append(ballConstraint)
+                
+            reverseNode = cmds.shadingNode("reverse", asUtility = True, name = "legSwitcher_reverse_node_" + side)
+            cmds.connectAttr("Rig_Settings" + "." + side + "LegMode", reverseNode + ".inputX")
+            
+            
+            for constraint in constraints:
+                targets = cmds.parentConstraint(constraint, q = True, weightAliasList = True)
+                cmds.connectAttr("Rig_Settings" + "." + side + "LegMode", constraint + "." + targets[1])
+                cmds.connectAttr(reverseNode + ".outputX", constraint + "." + targets[0])
+                
+                
+            #connect up visibility of controls to leg mode
+            cmds.connectAttr("Rig_Settings" + "." + side + "LegMode", "ik_leg_grp_" + side + ".v")
+            cmds.connectAttr(reverseNode + ".outputX", "fk_thigh_" + side + "_anim_grp.v")
+            
+                        
+                        
+            #set default to IK
+            cmds.setAttr("Rig_Settings" + "." + side + "LegMode", 1)
+            
+            
+            #constrain driver legs to result legs
+            cmds.parentConstraint(thighJoint, "driver_thigh_" + side, mo = True)
+            cmds.parentConstraint(calfJoint, "driver_calf_" + side, mo = True)
+            cmds.parentConstraint(footJoint, "driver_foot_" + side, mo = True)
+            
+            if ball:
+                cmds.parentConstraint(ballJoint, "driver_ball_" + side, mo = True)
+                
+                
+	    #create blend nodes for the scale
+	    scaleBlendColors_UpLeg = cmds.shadingNode("blendColors", asUtility = True, name = side + "_up_leg_scale_blend")
+	    cmds.connectAttr("ik_leg_thigh_" + side + ".scale", scaleBlendColors_UpLeg + ".color1")
+	    cmds.connectAttr("fk_thigh_" + side + "_anim.scale", scaleBlendColors_UpLeg + ".color2")
+	    cmds.connectAttr(scaleBlendColors_UpLeg + ".output", "driver_thigh_" + side + ".scale")
+	    
+	    
+	    scaleBlendColors_LoLeg = cmds.shadingNode("blendColors", asUtility = True, name = side + "_lo_leg_scale_blend")
+	    cmds.connectAttr("ik_leg_calf_" + side + ".scale", scaleBlendColors_LoLeg + ".color1")
+	    cmds.connectAttr("fk_calf_" + side + "_anim.scale", scaleBlendColors_LoLeg + ".color2")
+	    cmds.connectAttr(scaleBlendColors_LoLeg + ".output", "driver_calf_" + side + ".scale")
+	    
+	    scaleBlendColors_Foot = cmds.shadingNode("blendColors", asUtility = True, name = side + "_foot_scale_blend")
+	    cmds.connectAttr("ik_leg_foot_" + side + ".scale", scaleBlendColors_Foot + ".color1")
+	    cmds.connectAttr("fk_foot_" + side + "_anim.scale", scaleBlendColors_Foot + ".color2")
+	    cmds.connectAttr(scaleBlendColors_Foot + ".output", "driver_foot_" + side + ".scale")
+	    
+	    
+	    
+	    
+	    #set limits
+	    cmds.select("driver_thigh_" + side)
+	    cmds.transformLimits(sy = (.05, 2.0), sz = (.05, 2.0), esy = [False, True], esz = [False, True])
+	    cmds.select("driver_calf_" + side)
+	    cmds.transformLimits(sy = (.05, 2.0), sz = (.05, 2.0), esy = [False, True], esz = [False, True])	    
+                
+                
+                
+            #clean up legs hiearchy
+            cmds.parent([thighJoint,  ikFixThighJoint], legJointGrp)
+            cmds.setAttr("leg_ctrl_grp_" + side + ".v", 0)
+            
+            
+            #hide stuff
+            cmds.setAttr(ikFixThighJoint + ".v", 0)
+            cmds.setAttr("fk_leg_thigh_" + side + ".v", 0)
+            
+            
+            #Setup Twist Joints if the user selected them
+            if side == "l":
+                if cmds.getAttr("Skeleton_Settings.leftUpperLegTwist") > 0:
+                    self.buildThighTwist("l")
+                    
+                if cmds.getAttr("Skeleton_Settings.leftLowerLegTwist") > 0:
+                    self.buildCalfTwist("l")
+                    
+
+                    
+            if side == "r":
+                if cmds.getAttr("Skeleton_Settings.rightUpperLegTwist") > 0:
+                    self.buildThighTwist("r")
+                    
+                if cmds.getAttr("Skeleton_Settings.rightLowerLegTwist") > 0:
+                    self.buildCalfTwist("r")
+		    
+        #clean up the leg hierarchy. group all leg systems under 1 group
+        legMasterGrp = cmds.group(empty = True, name = "leg_sys_grp")
+        cmds.parent(["leg_group_l", "leg_group_r"], legMasterGrp)
+                
+
+            
+	
+            
+    # # # # # # # # # # # # # # # # # # # # # # # # # # # # # # # # # # # # # # # # # # # # # # # # # # # # # # # # # # # # # # # 
+    # # # # # # # # # # # # # # # # # # # # # # # # # # # # # # # # # # # # # # # # # # # # # # # # # # # # # # # # # # # # # # # 
+    # # # # # # # # # # # # # # # # # # # # # # # # # # # # # # # # # # # # # # # # # # # # # # # # # # # # # # # # # # # # # # # 
+    # # # # # # # # # # # # # # # # # # # # # # # # # # # # # # # # # # # # # # # # # # # # # # # # # # # # # # # # # # # # # # # 
+    # # # # # # # # # # # # # # # # # # # # # # # # # # # # # # # # # # # # # # # # # # # # # # # # # # # # # # # # # # # # # # # 
+    def buildCalfTwist(self, side):
+	
+	
+        if side == "l":
+            color = 5
+        
+        else:
+            color = 12
+            
+
+        #create our roll group
+        rollGrp = cmds.group(empty = True, name = "calf_" + side + "_roll_grp")
+        cmds.parentConstraint("driver_calf_" + side, rollGrp)
+        
+        #create our twist joint and twist mod joint
+        cmds.select(clear = True)
+        twistJoint = cmds.joint(name = "calf_" + side + "_twist_joint")
+        cmds.select(clear = True)
+        
+        constraint = cmds.parentConstraint("driver_calf_twist_01_" + side, twistJoint)[0]
+        cmds.delete(constraint)
+        
+        cmds.parent(twistJoint, rollGrp)
+        cmds.makeIdentity(twistJoint, t = 0, r = 1, s = 0, apply = True)
+	
+        #twist mod joint
+        twistMod = cmds.duplicate(twistJoint, po = True, name = "calf_" + side + "_twist_mod")[0]
+        cmds.parent(twistMod, twistJoint)
+        
+        
+        #create the manual twist control
+        twistCtrl = self.createControl("circle", 15, "calf_" + side + "_twist_anim")
+        cmds.setAttr(twistCtrl + ".ry", -90)
+        cmds.makeIdentity(twistCtrl, r = 1, apply =True)
+        
+        constraint = cmds.parentConstraint(twistMod, twistCtrl)[0]
+        cmds.delete(constraint)
+        
+        twistCtrlGrp = cmds.group(empty = True, name = "calf_" + side + "_twist_anim_grp")
+        constraint = cmds.parentConstraint(twistMod, twistCtrlGrp)[0]
+        cmds.delete(constraint)
+        
+        cmds.parent(twistCtrl, twistCtrlGrp)
+        cmds.parent(twistCtrlGrp, twistMod)
+        cmds.makeIdentity(twistCtrl, t = 1, r = 1, s = 1, apply = True)
+	
+	cmds.setAttr(twistCtrl + ".overrideEnabled", 1)
+        cmds.setAttr(twistCtrl + ".overrideColor", color)
+	for attr in [".sx", ".sy", ".sz"]:
+	    cmds.setAttr(twistCtrl + attr, lock = True, keyable = False)
+	
+	cmds.setAttr(twistCtrl + ".v", keyable = False)
+	
+        #add attr on clavicle anim for manual twist control visibility
+        cmds.select("hip_anim")
+        cmds.addAttr(longName=(side + "CalfTwistCtrlVis"), at = 'bool', dv = 0, keyable = True)
+        cmds.connectAttr("hip_anim." + side + "CalfTwistCtrlVis", twistCtrl + ".v")
+        cmds.connectAttr("hip_anim." + side + "CalfTwistCtrlVis", twistMod + ".v")
+	cmds.connectAttr("hip_anim." + side + "CalfTwistCtrlVis", twistJoint + ".v")
+	cmds.setAttr(twistMod + ".radius", .01)
+	cmds.setAttr(twistJoint + ".radius", .01)
+	
+	
+        #setup a simple relationship of foot rotateX value into mult node. input2X is driven by an attr on rig settings for twist amt(default is .5). Output into twist joint
+        twistMultNode = cmds.shadingNode("multiplyDivide", asUtility = True, name = "calf_twist_" + side + "_mult_node")
+        
+        #add attr to rig settings
+        cmds.select("Rig_Settings")
+        cmds.addAttr(longName= ( side + "CalfTwistAmount" ), defaultValue=.5, minValue=0, maxValue=1, keyable = True)
+                     
+                     
+        #connect output of driver hand into input1x
+        cmds.connectAttr("driver_foot_" + side + ".rx", twistMultNode + ".input1X")
+        
+        #connect attr into input2x
+        cmds.connectAttr("Rig_Settings." + side + "CalfTwistAmount", twistMultNode + ".input2X")
+        
+        #connect output into driver calf twist
+        cmds.connectAttr(twistMultNode + ".outputX", twistJoint + ".rx")
+	
+	#constrain driver joint to twist joint
+	cmds.parentConstraint(twistCtrl, "driver_calf_twist_01_" + side, mo = True)
+	
+	
+	#if there is more than 1 roll bone, set those up now:
+	if side == "l":
+	    sideName = "left"
+	else:
+	    sideName = "right"
+	    
+	data = cmds.getAttr("SkeletonSettings_Cache." + sideName + "LegOptions_numCalfTwistBones")
+	numRolls = ast.literal_eval(data)[0]
+	
+	
+	if numRolls > 1:
+	    for i in range(int(numRolls)):
+		
+		if i == 1:
+		    
+		    cmds.setAttr("Rig_Settings." + side + "CalfTwistAmount", .75)
+		    cmds.select("Rig_Settings")
+		    cmds.addAttr(longName= ( side + "CalfTwist2Amount" ), defaultValue=.5, minValue=0, maxValue=1, keyable = True)
+		    
+		    #create the manual twist control setup
+		    twistMod = cmds.duplicate("driver_calf_twist_0" + str(i + 1) + "_" + side , po = True, name = "calf_" + side + "_twist2_mod")[0]
+		    cmds.parent(twistMod, rollGrp)
+        
+        
+		    #create the manual twist control
+		    twistCtrl = self.createControl("circle", 15, "calf_" + side + "_twist2_anim")
+		    cmds.setAttr(twistCtrl + ".ry", -90)
+		    cmds.makeIdentity(twistCtrl, r = 1, apply =True)
+		    
+		    constraint = cmds.parentConstraint(twistMod, twistCtrl)[0]
+		    cmds.delete(constraint)
+		    
+		    twistCtrlGrp = cmds.group(empty = True, name = "calf_" + side + "_twist2_anim_grp")
+		    constraint = cmds.parentConstraint(twistMod, twistCtrlGrp)[0]
+		    cmds.delete(constraint)
+		    
+		    cmds.parent(twistCtrl, twistCtrlGrp)
+		    cmds.parent(twistCtrlGrp, twistMod)
+		    cmds.makeIdentity(twistCtrl, t = 1, r = 1, s = 1, apply = True)
+    
+		    cmds.connectAttr("hip_anim." + side + "CalfTwistCtrlVis", twistCtrl + ".v")
+		    cmds.connectAttr("hip_anim." + side + "CalfTwistCtrlVis", twistMod + ".v")
+		    for attr in [".sx", ".sy", ".sz"]:
+			cmds.setAttr(twistCtrl + attr, lock = True, keyable = False)
+		    
+		    cmds.setAttr(twistCtrl + ".v", keyable = False)
+		    cmds.setAttr(twistCtrl + ".overrideEnabled", 1)
+		    cmds.setAttr(twistCtrl + ".overrideColor", color)
+	
+		    #drive the twist joint
+		    twistMultNode = cmds.shadingNode("multiplyDivide", asUtility = True, name = "calf_twist_2_" + side + "_mult_node")
+		    cmds.connectAttr("driver_calf_twist_01_" + side + ".rx", twistMultNode + ".input1X")
+		    cmds.connectAttr("Rig_Settings." + side + "CalfTwist2Amount", twistMultNode + ".input2X")
+		    cmds.connectAttr(twistMultNode + ".outputX", twistCtrlGrp + ".rx")
+		    cmds.parentConstraint(twistCtrl, "driver_calf_twist_0" + str(i + 1) + "_" + side, mo = True)
+	
+	
+
+		    
+		if i == 2:
+		    
+		    cmds.select("Rig_Settings")
+		    cmds.addAttr(longName= ( side + "CalfTwist3Amount" ), defaultValue=.25, minValue=0, maxValue=1, keyable = True)
+		    
+		    #create the manual twist control setup
+		    twistMod = cmds.duplicate("driver_calf_twist_0" + str(i + 1) + "_" + side , po = True, name = "calf_" + side + "_twist3_mod")[0]
+		    cmds.parent(twistMod, rollGrp)
+        
+        
+		    #create the manual twist control
+		    twistCtrl = self.createControl("circle", 15, "calf_" + side + "_twist3_anim")
+		    cmds.setAttr(twistCtrl + ".ry", -90)
+		    cmds.makeIdentity(twistCtrl, r = 1, apply =True)
+		    
+		    constraint = cmds.parentConstraint(twistMod, twistCtrl)[0]
+		    cmds.delete(constraint)
+		    
+		    twistCtrlGrp = cmds.group(empty = True, name = "calf_" + side + "_twist3_anim_grp")
+		    constraint = cmds.parentConstraint(twistMod, twistCtrlGrp)[0]
+		    cmds.delete(constraint)
+		    
+		    cmds.parent(twistCtrl, twistCtrlGrp)
+		    cmds.parent(twistCtrlGrp, twistMod)
+		    cmds.makeIdentity(twistCtrl, t = 1, r = 1, s = 1, apply = True)
+    
+		    cmds.connectAttr("hip_anim." + side + "CalfTwistCtrlVis", twistCtrl + ".v")
+		    cmds.connectAttr("hip_anim." + side + "CalfTwistCtrlVis", twistMod + ".v")
+		    for attr in [".sx", ".sy", ".sz"]:
+			cmds.setAttr(twistCtrl + attr, lock = True, keyable = False)
+		    
+		    cmds.setAttr(twistCtrl + ".v", keyable = False)
+		    cmds.setAttr(twistCtrl + ".overrideEnabled", 1)
+		    cmds.setAttr(twistCtrl + ".overrideColor", color)
+	
+		    #drive the twist joint
+		    twistMultNode = cmds.shadingNode("multiplyDivide", asUtility = True, name = "calf_twist_3_" + side + "_mult_node")
+		    cmds.connectAttr("driver_calf_twist_01_" + side + ".rx", twistMultNode + ".input1X")
+		    cmds.connectAttr("Rig_Settings." + side + "CalfTwist3Amount", twistMultNode + ".input2X")
+		    cmds.connectAttr(twistMultNode + ".outputX", twistCtrlGrp + ".rx")
+		    cmds.parentConstraint(twistCtrl, "driver_calf_twist_0" + str(i + 1) + "_" + side, mo = True)
+		    
+
+	
+        #clean up hierarchy
+        cmds.parent(rollGrp, "leg_group_" + side)
+        
+    # # # # # # # # # # # # # # # # # # # # # # # # # # # # # # # # # # # # # # # # # # # # # # # # # # # # # # # # # # # # # # # 
+    # # # # # # # # # # # # # # # # # # # # # # # # # # # # # # # # # # # # # # # # # # # # # # # # # # # # # # # # # # # # # # # 
+    # # # # # # # # # # # # # # # # # # # # # # # # # # # # # # # # # # # # # # # # # # # # # # # # # # # # # # # # # # # # # # # 
+    # # # # # # # # # # # # # # # # # # # # # # # # # # # # # # # # # # # # # # # # # # # # # # # # # # # # # # # # # # # # # # # 
+    # # # # # # # # # # # # # # # # # # # # # # # # # # # # # # # # # # # # # # # # # # # # # # # # # # # # # # # # # # # # # # # 
+    def buildForearmTwist(self, side):
+        
+        if side == "l":
+            color = 5
+        
+        else:
+            color = 12
+            
+
+        #create our roll group
+        rollGrp = cmds.group(empty = True, name = "lowerarm_" + side + "_roll_grp")
+        cmds.parentConstraint("driver_lowerarm_" + side, rollGrp)
+        
+        #create our twist joint and twist mod joint
+        cmds.select(clear = True)
+        twistJoint = cmds.joint(name = "lowerarm_" + side + "_twist_joint")
+        cmds.select(clear = True)
+        
+        constraint = cmds.parentConstraint("driver_lowerarm_twist_01_" + side, twistJoint)[0]
+        cmds.delete(constraint)
+        
+        cmds.parent(twistJoint, rollGrp)
+        
+        #twist mod joint
+        twistMod = cmds.duplicate(twistJoint, po = True, name = "lowerarm_" + side + "_twist_mod")[0]
+        cmds.parent(twistMod, twistJoint)
+        
+        
+        #create the manual twist control
+        twistCtrl = self.createControl("circle", 15, "lowerarm_" + side + "_twist_anim")
+        cmds.setAttr(twistCtrl + ".ry", -90)
+        cmds.makeIdentity(twistCtrl, r = 1, apply =True)
+        
+        constraint = cmds.parentConstraint(twistMod, twistCtrl)[0]
+        cmds.delete(constraint)
+        
+        twistCtrlGrp = cmds.group(empty = True, name = "lowerarm_" + side + "_twist_anim_grp")
+        constraint = cmds.parentConstraint(twistMod, twistCtrlGrp)[0]
+        cmds.delete(constraint)
+        
+        cmds.parent(twistCtrl, twistCtrlGrp)
+        cmds.parent(twistCtrlGrp, twistMod)
+        cmds.makeIdentity(twistCtrl, t = 1, r = 1, s = 1, apply = True)
+	
+	cmds.setAttr(twistCtrl + ".overrideEnabled", 1)
+        cmds.setAttr(twistCtrl + ".overrideColor", color)
+	for attr in [".sx", ".sy", ".sz"]:
+	    cmds.setAttr(twistCtrl + attr, lock = True, keyable = False)
+	
+	cmds.setAttr(twistCtrl + ".v", keyable = False)
+	
+        #add attr on clavicle anim for manual twist control visibility
+        cmds.select("clavicle_" + side + "_anim")
+        cmds.addAttr(longName=("twistCtrlVisLower"), at = 'bool', dv = 0, keyable = True)
+        cmds.connectAttr("clavicle_" + side + "_anim.twistCtrlVisLower", twistCtrl + ".v")
+        cmds.connectAttr("clavicle_" + side + "_anim.twistCtrlVisLower", twistMod + ".v")
+	cmds.connectAttr("clavicle_" + side + "_anim.twistCtrlVisLower", twistJoint + ".v")
+	cmds.setAttr(twistMod + ".radius", .01)
+	cmds.setAttr(twistJoint + ".radius", .01)
+	
+	
+        #setup a simple relationship of foot rotateX value into mult node. input2X is driven by an attr on rig settings for twist amt(default is .5). Output into twist joint
+        twistMultNode = cmds.shadingNode("multiplyDivide", asUtility = True, name = "forearm_twist_" + side + "_mult_node")
+        
+        #add attr to rig settings
+        cmds.select("Rig_Settings")
+        cmds.addAttr(longName= ( side + "ForearmTwistAmount" ), defaultValue=.5, minValue=0, maxValue=1, keyable = True)
+                     
+                     
+        #connect output of driver hand into input1x
+        cmds.connectAttr("driver_hand_" + side + ".rx", twistMultNode + ".input1X")
+        
+        #connect attr into input2x
+        cmds.connectAttr("Rig_Settings." + side + "ForearmTwistAmount", twistMultNode + ".input2X")
+        
+        #connect output into driver calf twist
+        cmds.connectAttr(twistMultNode + ".outputX", twistJoint + ".rx")
+	
+	#constrain driver joint to twist joint
+	cmds.parentConstraint(twistCtrl, "driver_lowerarm_twist_01_" + side, mo = True)
+	
+	
+	#if there is more than 1 roll bone, set those up now:
+	if side == "l":
+	    sideName = "left"
+	else:
+	    sideName = "right"
+	    
+	data = cmds.getAttr("SkeletonSettings_Cache." + sideName + "ArmOptions_numLowArmTwistBones")
+	numRolls = ast.literal_eval(data)[0]
+	
+	
+	if numRolls > 1:
+	    for i in range(int(numRolls)):
+		
+		if i == 1:
+		    
+		    cmds.setAttr("Rig_Settings." + side + "ForearmTwistAmount", .75)
+		    cmds.select("Rig_Settings")
+		    cmds.addAttr(longName= ( side + "ForearmTwist2Amount" ), defaultValue=.5, minValue=0, maxValue=1, keyable = True)
+		    
+		    #create the manual twist control setup
+		    twistMod = cmds.duplicate("driver_lowerarm_twist_0" + str(i + 1) + "_" + side , po = True, name = "lowerarm_" + side + "_twist2_mod")[0]
+		    cmds.parent(twistMod, rollGrp)
+        
+        
+		    #create the manual twist control
+		    twistCtrl = self.createControl("circle", 15, "lowerarm_" + side + "_twist2_anim")
+		    cmds.setAttr(twistCtrl + ".ry", -90)
+		    cmds.makeIdentity(twistCtrl, r = 1, apply =True)
+		    
+		    constraint = cmds.parentConstraint(twistMod, twistCtrl)[0]
+		    cmds.delete(constraint)
+		    
+		    twistCtrlGrp = cmds.group(empty = True, name = "lowerarm_" + side + "_twist2_anim_grp")
+		    constraint = cmds.parentConstraint(twistMod, twistCtrlGrp)[0]
+		    cmds.delete(constraint)
+		    
+		    cmds.parent(twistCtrl, twistCtrlGrp)
+		    cmds.parent(twistCtrlGrp, twistMod)
+		    cmds.makeIdentity(twistCtrl, t = 1, r = 1, s = 1, apply = True)
+		    
+		    cmds.connectAttr("clavicle_" + side + "_anim.twistCtrlVisLower", twistCtrl + ".v")
+		    cmds.connectAttr("clavicle_" + side + "_anim.twistCtrlVisLower", twistMod + ".v")
+		    for attr in [".sx", ".sy", ".sz"]:
+			cmds.setAttr(twistCtrl + attr, lock = True, keyable = False)
+		    
+		    cmds.setAttr(twistCtrl + ".v", keyable = False)
+		    cmds.setAttr(twistCtrl + ".overrideEnabled", 1)
+		    cmds.setAttr(twistCtrl + ".overrideColor", color)
+	
+		    #drive the twist joint
+		    twistMultNode = cmds.shadingNode("multiplyDivide", asUtility = True, name = "forearm_twist_2_" + side + "_mult_node")
+		    cmds.connectAttr("driver_lowerarm_twist_01_" + side + ".rx", twistMultNode + ".input1X")
+		    cmds.connectAttr("Rig_Settings." + side + "ForearmTwist2Amount", twistMultNode + ".input2X")
+		    cmds.connectAttr(twistMultNode + ".outputX", twistCtrlGrp + ".rx")
+		    cmds.parentConstraint(twistCtrl, "driver_lowerarm_twist_0" + str(i + 1) + "_" + side, mo = True)
+	
+	
+
+		    
+		if i == 2:
+		    
+		    cmds.select("Rig_Settings")
+		    cmds.addAttr(longName= ( side + "ForearmTwist3Amount" ), defaultValue=.25, minValue=0, maxValue=1, keyable = True)
+		    
+		    #create the manual twist control setup
+		    twistMod = cmds.duplicate("driver_lowerarm_twist_0" + str(i + 1) + "_" + side , po = True, name = "lowerarm_" + side + "_twist3_mod")[0]
+		    cmds.parent(twistMod, rollGrp)
+        
+        
+		    #create the manual twist control
+		    twistCtrl = self.createControl("circle", 15, "lowerarm_" + side + "_twist3_anim")
+		    cmds.setAttr(twistCtrl + ".ry", -90)
+		    cmds.makeIdentity(twistCtrl, r = 1, apply =True)
+		    
+		    constraint = cmds.parentConstraint(twistMod, twistCtrl)[0]
+		    cmds.delete(constraint)
+		    
+		    twistCtrlGrp = cmds.group(empty = True, name = "lowerarm_" + side + "_twist3_anim_grp")
+		    constraint = cmds.parentConstraint(twistMod, twistCtrlGrp)[0]
+		    cmds.delete(constraint)
+		    
+		    cmds.parent(twistCtrl, twistCtrlGrp)
+		    cmds.parent(twistCtrlGrp, twistMod)
+		    cmds.makeIdentity(twistCtrl, t = 1, r = 1, s = 1, apply = True)
+		    
+		    cmds.connectAttr("clavicle_" + side + "_anim.twistCtrlVisLower", twistCtrl + ".v")
+		    cmds.connectAttr("clavicle_" + side + "_anim.twistCtrlVisLower", twistMod + ".v")
+		    for attr in [".sx", ".sy", ".sz"]:
+			cmds.setAttr(twistCtrl + attr, lock = True, keyable = False)
+		    
+		    cmds.setAttr(twistCtrl + ".v", keyable = False)
+		    cmds.setAttr(twistCtrl + ".overrideEnabled", 1)
+		    cmds.setAttr(twistCtrl + ".overrideColor", color)
+	
+		    #drive the twist joint
+		    twistMultNode = cmds.shadingNode("multiplyDivide", asUtility = True, name = "forearm_twist_3_" + side + "_mult_node")
+		    cmds.connectAttr("driver_lowerarm_twist_01_" + side + ".rx", twistMultNode + ".input1X")
+		    cmds.connectAttr("Rig_Settings." + side + "ForearmTwist3Amount", twistMultNode + ".input2X")
+		    cmds.connectAttr(twistMultNode + ".outputX", twistCtrlGrp + ".rx")
+		    cmds.parentConstraint(twistCtrl, "driver_lowerarm_twist_0" + str(i + 1) + "_" + side, mo = True)
+		    
+
+		    
+
+    # # # # # # # # # # # # # # # # # # # # # # # # # # # # # # # # # # # # # # # # # # # # # # # # # # # # # # # # # # # # # # # 
+    # # # # # # # # # # # # # # # # # # # # # # # # # # # # # # # # # # # # # # # # # # # # # # # # # # # # # # # # # # # # # # # 
+    # # # # # # # # # # # # # # # # # # # # # # # # # # # # # # # # # # # # # # # # # # # # # # # # # # # # # # # # # # # # # # # 
+    # # # # # # # # # # # # # # # # # # # # # # # # # # # # # # # # # # # # # # # # # # # # # # # # # # # # # # # # # # # # # # # 
+    # # # # # # # # # # # # # # # # # # # # # # # # # # # # # # # # # # # # # # # # # # # # # # # # # # # # # # # # # # # # # # # 
+    def buildArmRoll(self, side):
+        
+        if side == "l":
+            color = 5
+	    sideName = "left"
+        
+        else:
+            color = 12
+	    sideName = "right"
+	    
+	
+	#get number of roll bones
+	data = cmds.getAttr("SkeletonSettings_Cache." + sideName + "ArmOptions_numUpArmTwistBones")
+	numRolls = ast.literal_eval(data)[0]	
+	
+	
+	#create a nurbs plane for our ribbon
+	ribbon = cmds.nurbsPlane(ax = [0,0,1], lr = numRolls, width = 10, d = 3, u = 1, v = numRolls, ch = True, name = "upperarm_twist_ribbon_" + side)[0]
+	
+	#rebuild the ribbon with 1 U span
+	ribbon = cmds.rebuildSurface(ribbon, su = 1, du = 1, sv = numRolls, dv = 1, ch = 1)[0]
+	cmds.setAttr(ribbon + ".rz", -90)
+	cmds.makeIdentity(ribbon, apply = True, t = 1, r = 1, s = 1)
+	
+	#create 2 temporary skin joints
+	moveVal = 0
+	for i in range(numRolls - 1):
+	    moveVal += 10
+	
+	cmds.select(clear = True)
+	topSkinJoint = cmds.joint(name = "top_skinJoint_temp")
+	cmds.move(moveVal, 0, 0, r = True, os = True, wd = True)
+	cmds.select(clear = True)
+	    
+	cmds.select(clear = True)
+	bottomSkinJoint = cmds.joint(name = "bottom_skinJoint_temp")
+	cmds.move(moveVal * -1, 0, 0, r = True, os = True, wd = True)
+	cmds.select(clear = True)
+	
+
+	#skin ribbon
+	cmds.select([ribbon, topSkinJoint, bottomSkinJoint])
+	skin = cmds.skinCluster(tsb = True, mi = 2, omi = True, dr = 5, bm = 0, sm = 0)
+	
+	
+	#position the joints, thus moving the ribbon
+	constraint = cmds.parentConstraint("driver_upperarm_" + side, topSkinJoint)[0]
+	cmds.delete(constraint)
+	
+	constraint = cmds.parentConstraint("driver_lowerarm_" + side, bottomSkinJoint)[0]
+	cmds.delete(constraint)
+	
+	
+	#delete ribbon history and skin joints
+	cmds.delete(ribbon, ch = True)
+	cmds.delete([bottomSkinJoint, topSkinJoint])
+	
+	
+	#create hair system on ribbon
+	cmds.select(ribbon)
+	mel.eval("createHair 1 3 10 0 0 0 0 5 0 2 1 1;")
+	
+	
+	#figure out which follicles created represent which areas on the ribbon
+	hairs = cmds.ls(type = "hairSystem")
+	
+	if len(hairs) > 0:
+	    hairSys = hairs[0]
+	    parent = cmds.listRelatives(hairs[0], parent = True)[0]
+	    hairSys = cmds.rename(parent, "upperarm_twist_" + side + "_hairSys")
+	    
+	follicles = cmds.listConnections(hairSys + "Shape", type = "follicle")
+	follicles = set(follicles)
+	hairFollicles = follicles
+	
+	
+
+	#delete outputCurves
+	cmds.delete(parent + "OutputCurves")	
+
+
+	#create a joint per follicle
+	for follicle in hairFollicles:
+	    cmds.select(clear = True)
+	    joint = cmds.joint(name = follicle + "_joint")
+	    cmds.select(clear = True)
+	    constraint = cmds.parentConstraint(follicle, joint)[0]
+	    cmds.delete(constraint)
+	    
+	    cmds.parent(joint, follicle)
+	    cmds.makeIdentity(joint, apply = True, t = 0, r = 1, s = 0)
+	    
+	    
+
+	#create the skin joints (final)
+	skinJoints = []
+	for i in range(numRolls + 1):
+	    cmds.select(clear = True)
+	    skinJoint = cmds.joint(name = "skin_upperarm_twist_joint_" + side + str(i))
+	    cmds.select(clear = True)
+	    skinJoints.append(skinJoint)
+	    
+	for i in range(numRolls):
+	    constraint = cmds.parentConstraint("driver_upperarm_twist_0" + str(i + 1) + "_" + side, "skin_upperarm_twist_joint_" + side + str(i))[0]
+	    cmds.delete(constraint)
+	
+	constraint = cmds.parentConstraint("driver_lowerarm_" + side, skinJoint)[0]
+	cmds.delete(constraint)
+	
+
+	
+	#create our manual control curves
+	x = 1
+	groups = []
+	for joint in skinJoints:
+	    
+	    if joint != skinJoints[-1]:
+		if joint == skinJoints[0]:
+		    name = "upperarm_" + side + "_twist_anim_grp"
+		else:
+		    name = "upperarm_" + side + "_twist" + str(x) + "_anim_grp"
+		    
+		group = cmds.group(empty = True, name = name)
+		groups.append(group)
+		constraint = cmds.parentConstraint(joint, group)[0]
+		cmds.delete(constraint)
+		x = x + 1
+		
+	for i in range(int(len(groups))):
+	    name = groups[i].partition("_grp")[0]
+	    twistCtrl = self.createControl("circle", 20, name)
+	    cmds.setAttr(twistCtrl + ".ry", -90)
+	    cmds.makeIdentity(twistCtrl, r = 1, apply =True)
+	    
+	    constraint = cmds.parentConstraint(groups[i], twistCtrl)[0]
+	    cmds.delete(constraint)
+	    
+	    cmds.parent(twistCtrl, groups[i])
+	    cmds.parent(skinJoints[i], twistCtrl)
+	    cmds.makeIdentity(skinJoints[i], apply = True, t = 0, r = 1, s = 0)
+	    
+	    #clean up control
+	    cmds.setAttr(twistCtrl + ".v", keyable = False)
+	    cmds.setAttr(twistCtrl + ".overrideEnabled", 1)
+	    cmds.setAttr(twistCtrl + ".overrideColor", color)	    
+	    
+	    
+	#organize groups
+	masterGrp = cmds.group(empty = True, name = "upperarm_twist_master_grp_" + side)
+	constraint = cmds.parentConstraint("rig_clavicle_" + side, masterGrp)[0]
+	cmds.delete(constraint)
+	
+	rollGrp = cmds.duplicate(masterGrp, name = "upperarm_twist_roll_grp_" + side)[0]
+	constraint = cmds.parentConstraint("driver_upperarm_" + side, rollGrp)[0]
+	cmds.delete(constraint)	
+	cmds.parent(rollGrp, masterGrp)
+	
+	#set rotate order on roll grp (xzy)
+	cmds.setAttr(rollGrp + ".rotateOrder", 3)
+	
+	
+	for group in groups:
+	    cmds.parent(group, rollGrp)
+	    
+	    
+	#skin ribbon to skin joints
+	cmds.select(ribbon)
+	for joint in skinJoints:
+	    cmds.select(joint, add = True)
+	    
+	skin = cmds.skinCluster(tsb = True, mi = 2, omi = True, dr = 5, bm = 0, sm = 0)
+	
+	
+	#orient roll grp to both fk/ik arm joints and set driven keys between them
+	upArmConstOrient = cmds.orientConstraint(["fk_upperarm_" + side, "ik_upperarm_" + side], rollGrp, mo = True, skip = "x")[0]
+
+	cmds.setAttr("Rig_Settings." + side + "ArmMode", 0)
+	cmds.setAttr(upArmConstOrient + "." + "fk_upperarm_" + side + "W0", 1)
+	cmds.setAttr(upArmConstOrient + "." + "ik_upperarm_" + side + "W1", 0)
+	cmds.setDrivenKeyframe([upArmConstOrient + "." + "fk_upperarm_" + side + "W0", upArmConstOrient + "." + "ik_upperarm_" + side + "W1", ], cd = "Rig_Settings." + side + "ArmMode", itt = "linear", ott = "linear")
+
+	cmds.setAttr("Rig_Settings." + side + "ArmMode", 1)
+	cmds.setAttr(upArmConstOrient + "." + "fk_upperarm_" + side + "W0", 0)
+	cmds.setAttr(upArmConstOrient + "." + "ik_upperarm_" + side + "W1", 1)
+	cmds.setDrivenKeyframe([upArmConstOrient + "." + "fk_upperarm_" + side + "W0", upArmConstOrient + "." + "ik_upperarm_" + side + "W1", ], cd = "Rig_Settings." + side + "ArmMode", itt = "linear", ott = "linear")
+    
+	cmds.setAttr("Rig_Settings." + side + "ArmMode", 0)
+	
+	
+	#parent end skin joint to masterGrp and orientConstrain twist to driver upper arm
+	cmds.parent(skinJoints[-1], rollGrp)
+	cmds.orientConstraint("driver_upperarm_" + side, skinJoints[-1], skip = ["y", "z"])
+	
+	
+	#parentConstraint master roll grp to driver clavicle
+	cmds.parentConstraint("driver_clavicle_" + side, masterGrp, mo = True)
+	
+	#hook up driver joints
+	hairFollicles = sorted(hairFollicles)
+	hairFollicles = hairFollicles[::-1]
+	num = 1
+	
+	for i in range(len(hairFollicles)):	    
+	    if cmds.objExists("driver_upperarm_twist_0" + str(num) + "_" + side):
+		    cmds.orientConstraint(skinJoints[i], "driver_upperarm_twist_0" + str(num) + "_" + side)
+		    cmds.pointConstraint(skinJoints[i], "driver_upperarm_twist_0" + str(num) + "_" + side, mo = True)
+		    cmds.scaleConstraint(skinJoints[i], "driver_upperarm_twist_0" + str(num) + "_" + side)
+		    num = num + 1
+
+
+
+
+	#add attr on clavicle anim for manual twist control visibility
+	cmds.select("clavicle_" + side + "_anim")
+	cmds.addAttr(longName=("twistCtrlVis"), at = 'bool', dv = 0, keyable = True)
+	cmds.connectAttr("clavicle_" + side + "_anim.twistCtrlVis", rollGrp + ".v")
+
+	
+	#hook up multiply nodes so that twistAmount values from rig settings affect the ribbon twist
+        cmds.select("Rig_Settings")
+        cmds.addAttr(longName= ( side + "UpperarmTwistAmount" ), defaultValue= .9, minValue= 0 , maxValue= 1, keyable = True)
+	
+	
+	#take twist ammount attr, multiply by -1, and feed into upperarm twist joint 1
+        multNodeA = cmds.shadingNode("multiplyDivide", asUtility = True, name = "upperarm_twist_" + side + "_multNodeA")
+        cmds.connectAttr("Rig_Settings." + side + "UpperarmTwistAmount", multNodeA+ ".input1X")
+	cmds.setAttr(multNodeA + ".input2X", -1)
+	
+        multNodeB = cmds.shadingNode("multiplyDivide", asUtility = True, name = "upperarm_twist_" + side + "_multNodeB")
+        cmds.connectAttr(rollGrp + ".rx", multNodeB+ ".input1X")
+	cmds.connectAttr(multNodeA + ".outputX", multNodeB + ".input2X")
+        cmds.connectAttr(multNodeB + ".outputX", groups[0] + ".rx")	
+	
+	
+	#any twist joints over the initial, setup simply mult nodes for carry down values
+	if numRolls > 1:
+	    for i in range(int(numRolls)):
+		if i == 1:
+		    cmds.select("Rig_Settings")
+		    cmds.addAttr(longName= ( side + "UpperarmTwist2Amount" ), defaultValue=.5, minValue=0, maxValue=1, keyable = True)
+		    
+		    #hook up multiply nodes so that twistAmount values from rig settings affect the ribbon twist
+		    multNode = cmds.shadingNode("multiplyDivide", asUtility = True, name = "upperarm_twist2_" + side + "_multNode")
+		    blendNode = cmds.shadingNode("blendColors", asUtility = True, name = "upperarm_twist2_" + side + "_multNode")
+		    
+		    #hook up blendnode to take in fk and ik upperarm rx values
+		    cmds.connectAttr( "ik_upperarm_" + side + ".rx", blendNode + ".color1R")
+		    cmds.connectAttr( "fk_upperarm_" + side + ".rx", blendNode + ".color2R")
+		    
+		    #take output of that and plug into multNode. multiply by the twist ammount attribute value
+		    cmds.connectAttr(blendNode + ".outputR", multNode + ".input1X")
+		    cmds.connectAttr("Rig_Settings." + side + "UpperarmTwist2Amount" , multNode + ".input2X")
+		    cmds.connectAttr(multNode + ".outputX", groups[i] + ".rx")
+		    
+		    #connect blendNode.blender to rig settings arm mode
+		    cmds.connectAttr("Rig_Settings." + side + "ArmMode", blendNode + ".blender")
+
+		    
+		if i == 2:
+		    cmds.select("Rig_Settings")
+		    cmds.addAttr(longName= ( side + "UpperarmTwist3Amount" ), defaultValue=.5, minValue=0, maxValue=1, keyable = True)
+		    
+		    #hook up multiply nodes so that twistAmount values from rig settings affect the ribbon twist
+		    multNode = cmds.shadingNode("multiplyDivide", asUtility = True, name = "upperarm_twist3_" + side + "_multNode")
+		    blendNode = cmds.shadingNode("blendColors", asUtility = True, name = "upperarm_twist3_" + side + "_multNode")
+		    
+		    #hook up blendnode to take in fk and ik upperarm rx values
+		    cmds.connectAttr( "ik_upperarm_" + side + ".rx", blendNode + ".color1R")
+		    cmds.connectAttr( "fk_upperarm_" + side + ".rx", blendNode + ".color2R")
+		    
+		    #take output of that and plug into multNode. multiply by the twist ammount attribute value
+		    cmds.connectAttr(blendNode + ".outputR", multNode + ".input1X")
+		    cmds.connectAttr("Rig_Settings." + side + "UpperarmTwist3Amount" , multNode + ".input2X")
+		    cmds.connectAttr(multNode + ".outputX", groups[i] + ".rx")
+		    
+		    #connect blendNode.blender to rig settings arm mode
+		    cmds.connectAttr("Rig_Settings." + side + "ArmMode", blendNode + ".blender")
+	    
+	    
+	#Group up and parent into rig
+	twistGrp = cmds.group(empty = True, name = "upperarm_twist_grp_" + side)
+	cmds.parent([ribbon, hairSys, masterGrp], twistGrp)
+	
+	#find follicles grp
+	for follicle in hairFollicles:
+		folliclesGrp = cmds.listRelatives(follicle, parent = True)	
+	cmds.parent(folliclesGrp[0], twistGrp)
+	
+	if cmds.objExists("nucleus1"):
+	    cmds.parent("nucleus1", twistGrp)
+	
+	#turn inherits transforms off
+	cmds.setAttr(folliclesGrp[0] + ".inheritsTransform", 0)
+	cmds.setAttr(ribbon + ".inheritsTransform", 0)
+	
+	#hide nodes
+	for node in [folliclesGrp[0], ribbon, hairSys, skinJoints[0]]:
+	    cmds.setAttr(node + ".v", 0)
+	
+
+    # # # # # # # # # # # # # # # # # # # # # # # # # # # # # # # # # # # # # # # # # # # # # # # # # # # # # # # # # # # # # # # 
+    # # # # # # # # # # # # # # # # # # # # # # # # # # # # # # # # # # # # # # # # # # # # # # # # # # # # # # # # # # # # # # # 
+    # # # # # # # # # # # # # # # # # # # # # # # # # # # # # # # # # # # # # # # # # # # # # # # # # # # # # # # # # # # # # # # 
+    # # # # # # # # # # # # # # # # # # # # # # # # # # # # # # # # # # # # # # # # # # # # # # # # # # # # # # # # # # # # # # # 
+    # # # # # # # # # # # # # # # # # # # # # # # # # # # # # # # # # # # # # # # # # # # # # # # # # # # # # # # # # # # # # # # 
+    def buildThighTwist(self, side):
+        
+        
+        if side == "l":
+            color = 5
+        
+        else:
+            color = 12
+	    
+	    
+        #create a 2 bone joint chain and position the joints
+        cmds.select(clear = True)
+        driverJointStart = cmds.joint(name = "thigh_twist_" + side + "_driver")
+        cmds.select(clear = True)
+        driverJointEnd = cmds.joint(name = "thigh_twist_" + side + "_driver_end")
+        cmds.select(clear = True)
+        
+        constraint = cmds.parentConstraint("driver_pelvis", driverJointStart)[0]
+        cmds.delete(constraint)
+        
+        constraint = cmds.parentConstraint("driver_thigh_" + side, driverJointEnd)[0]
+        cmds.delete(constraint)
+        
+        cmds.parent(driverJointEnd, driverJointStart)
+        
+        #position the chain back a bit(+Y)
+        scaleFactor = self.getScaleFactor()
+        cmds.setAttr(driverJointStart + ".ty", (35 * scaleFactor))
+        
+        
+        #create our driver rig grp
+        driverGrp = cmds.group(empty = True, name = "thigh_twist_" + side + "_driver_grp")
+        constraint = cmds.parentConstraint("driver_pelvis", driverGrp)
+        
+        #create RP IK on driver joint chain
+        ikNodes = cmds.ikHandle(sol = "ikRPsolver", name = "thigh_twist_driver_" + side + "_ikHandle", sj = driverJointStart, ee = driverJointEnd)
+        cmds.setAttr(ikNodes[0] + ".v", 0)
+        
+        #point constrain IK handle to calf
+        cmds.pointConstraint("driver_calf_" + side, ikNodes[0])
+        
+        #create our pole vector locator and get it positioned
+        poleVector = cmds.spaceLocator(name = "thigh_twist_driver_" + side + "_PV")[0]
+        cmds.setAttr(poleVector + ".v", 0)
+        cmds.parent(poleVector, "driver_thigh_" + side)
+        
+        for attr in [".tx", ".ty", ".tz", ".rx", ".ry", ".rz"]:
+            cmds.setAttr(poleVector + attr, 0)
+            
+        #move it back in +Y a bit
+        if side == "l":
+            cmds.setAttr(poleVector + ".ty", (35 * scaleFactor))
+            
+        else:
+            cmds.setAttr(poleVector + ".ty", (-35 * scaleFactor))
+        
+        #create a pole Vector constraint between IK handle and the locator
+        cmds.poleVectorConstraint(poleVector, ikNodes[0])
+        
+        #parent nodes to driver grp
+        cmds.parent([driverJointStart, poleVector, ikNodes[0]], driverGrp)
+        
+        
+        #create our roll group
+        rollGrp = cmds.group(empty = True, name = "thigh_" + side + "_roll_grp")
+        cmds.parentConstraint("driver_thigh_" + side, rollGrp)
+        
+        #create our twist joint and twist mod joint
+        cmds.select(clear = True)
+        twistJoint = cmds.joint(name = "thigh_" + side + "_twist_joint")
+        cmds.select(clear = True)
+        
+        constraint = cmds.parentConstraint("driver_thigh_twist_01_" + side, twistJoint)[0]
+        cmds.delete(constraint)
+        
+        cmds.parent(twistJoint, rollGrp)
+        
+        #twist mod joint
+        twistMod = cmds.duplicate(twistJoint, po = True, name = "thigh_" + side + "_twist_mod")[0]
+        cmds.parent(twistMod, twistJoint)
+        
+
+        #create the manual twist control
+        twistCtrl = self.createControl("circle", 30, "thigh_" + side + "_twist_anim")
+        cmds.setAttr(twistCtrl + ".ry", -90)
+        cmds.makeIdentity(twistCtrl, r = 1, apply =True)
+        
+        constraint = cmds.parentConstraint(twistMod, twistCtrl)[0]
+        cmds.delete(constraint)
+        
+        twistCtrlGrp = cmds.group(empty = True, name = "thigh_" + side + "_twist_anim_grp")
+        constraint = cmds.parentConstraint(twistMod, twistCtrl)[0]
+        cmds.delete(constraint)
+        
+        cmds.parent(twistCtrl, twistCtrlGrp)
+        cmds.parent(twistCtrlGrp, twistMod)
+        cmds.makeIdentity(twistCtrl, t = 1, r = 1, s = 1, apply = True)
+        #lock attrs on twist ctrl and color
+        for attr in [".sx", ".sy", ".sz"]:
+            cmds.setAttr(twistCtrl + attr, lock = True, keyable = False)
+        
+        cmds.setAttr(twistCtrl + ".v", keyable = False)
+        cmds.setAttr(twistCtrl + ".overrideEnabled", 1)
+        cmds.setAttr(twistCtrl + ".overrideColor", color)
+        
+        
+        #add attr on clavicle anim for manual twist control visibility
+        cmds.select("hip_anim")
+        cmds.addAttr(longName=(side + "TwistCtrlVis"), at = 'bool', dv = 0, keyable = True)
+        cmds.connectAttr("hip_anim." + side + "TwistCtrlVis", twistCtrl + ".v")
+        cmds.connectAttr("hip_anim." + side + "TwistCtrlVis", twistMod + ".v")
+	cmds.connectAttr("hip_anim." + side + "TwistCtrlVis", twistJoint + ".v")
+	cmds.setAttr(twistMod + ".radius", .01)
+	cmds.setAttr(twistJoint + ".radius", .01)
+        
+        #create our aim up locator
+        aimUp = cmds.spaceLocator(name = "thigh_twist_driver_" + side + "_aimUp")[0]
+        cmds.setAttr(aimUp + ".v", 0)
+        cmds.parent(aimUp, "driver_thigh_" + side)
+        
+        for attr in [".tx", ".ty", ".tz", ".rx", ".ry", ".rz"]:
+            cmds.setAttr(aimUp + attr, 0)
+            
+        #move it back in +Y a bit
+        if side == "l":
+            cmds.setAttr(aimUp + ".ty", (35 * scaleFactor))
+            cmds.setAttr(aimUp + ".tx", (-70 * scaleFactor))
+            
+        if side == "r":
+            cmds.setAttr(aimUp + ".ty", (-35 * scaleFactor))
+            cmds.setAttr(aimUp + ".tx", (70 * scaleFactor))
+        
+        #parent aim up to thigh twist driver
+        cmds.parent(aimUp, driverJointStart)
+        
+        
+        #aim constrain the twist joint roll bone to the driver calf
+        cmds.makeIdentity(twistJoint, t = 0, r = 1, s = 0, apply = True)
+        cmds.aimConstraint("driver_calf_" + side, twistJoint, weight = 1, aimVector = [-1, 0, 0], upVector = [0, 1, 0], worldUpType = "object", worldUpObject = aimUp)
+        
+        #create a multiply node that will push some twist up to the twist mode joint but not all
+        cmds.select("Rig_Settings")
+        cmds.addAttr(longName= ( side + "ThighTwistAmount" ), defaultValue= .9, minValue= 0 , maxValue= 1, keyable = True)
+	
+        multNode = cmds.shadingNode("multiplyDivide", asUtility = True, name = "thigh_twist_" + side + "_multNode")
+	subtractNode = cmds.shadingNode("plusMinusAverage", asUtility = True, name = "thigh_twist_corrective_" + side + "_subtractNode")
+	cmds.setAttr(subtractNode + ".operation", 2)
+	cmds.connectAttr("Rig_Settings." + side + "ThighTwistAmount", subtractNode + ".input1D[0]")
+	cmds.setAttr(subtractNode + ".input1D[1]", 1)
+	
+	
+        cmds.connectAttr(twistJoint + ".rx", multNode+ ".input1X")
+	cmds.connectAttr(subtractNode + ".output1D", multNode + ".input2X")
+        cmds.connectAttr(multNode + ".outputX", twistMod + ".rx")
+	
+        
+        #constrain the driver twist joint to the twist mod
+        cmds.parentConstraint(twistCtrl, "driver_thigh_twist_01_" + side, mo = True)
+        
+        #constrain the driver twist joint to the twist mod
+	#get the number of roll bones
+	if side == "l":
+	    sideName = "left"
+	else:
+	    sideName = "right"
+	    
+	data = cmds.getAttr("SkeletonSettings_Cache." + sideName + "LegOptions_numThighTwistBones")
+	numRolls = ast.literal_eval(data)[0]
+	
+	
+	if numRolls > 1:
+	    for i in range(int(numRolls)):
+		
+		
+		if i == 1:
+		    cmds.select("Rig_Settings")
+		    cmds.addAttr(longName= ( side + "ThighTwist2Amount" ), defaultValue=.5, minValue=0, maxValue=1, keyable = True)
+		    
+		    #create the manual twist control setup
+		    twistMod = cmds.duplicate("driver_thigh_twist_0" + str(i + 1) + "_" + side , po = True, name = "thigh_" + side + "_twist2_mod")[0]
+		    cmds.parent(twistMod, rollGrp)
+        
+        
+		    #create the manual twist control
+		    twistCtrl = self.createControl("circle", 30, "thigh_" + side + "_twist2_anim")
+		    cmds.setAttr(twistCtrl + ".ry", -90)
+		    cmds.makeIdentity(twistCtrl, r = 1, apply =True)
+		    
+		    constraint = cmds.parentConstraint(twistMod, twistCtrl)[0]
+		    cmds.delete(constraint)
+		    
+		    twistCtrlGrp = cmds.group(empty = True, name = "thigh_" + side + "_twist2_anim_grp")
+		    constraint = cmds.parentConstraint(twistMod, twistCtrlGrp)[0]
+		    cmds.delete(constraint)
+		    
+		    cmds.parent(twistCtrl, twistCtrlGrp)
+		    cmds.parent(twistCtrlGrp, twistMod)
+		    cmds.makeIdentity(twistCtrl, t = 1, r = 1, s = 1, apply = True)
+		    
+		    cmds.connectAttr("hip_anim." + side + "TwistCtrlVis", twistCtrl + ".v")
+		    cmds.connectAttr("hip_anim." + side + "TwistCtrlVis", twistMod + ".v")
+		    
+		    #lock attrs on twist ctrl and color
+		    for attr in [".sx", ".sy", ".sz"]:
+			cmds.setAttr(twistCtrl + attr, lock = True, keyable = False)
+		    
+		    cmds.setAttr(twistCtrl + ".v", keyable = False)
+		    cmds.setAttr(twistCtrl + ".overrideEnabled", 1)
+		    cmds.setAttr(twistCtrl + ".overrideColor", color)
+	
+		    #drive the twist joint
+		    twistMultNode = cmds.shadingNode("multiplyDivide", asUtility = True, name = "thigh_twist_2_" + side + "_mult_node")
+		    cmds.connectAttr("driver_thigh_twist_01_" + side + ".rx", twistMultNode + ".input1X")
+		    cmds.connectAttr("Rig_Settings." + side + "ThighTwist2Amount", twistMultNode + ".input2X")
+		    cmds.connectAttr(twistMultNode + ".outputX", twistCtrlGrp + ".rx")
+		    cmds.parentConstraint(twistCtrl, "driver_thigh_twist_0" + str(i + 1) + "_" + side, mo = True)
+		    
+		    
+		    
+		if i == 2:
+		    cmds.select("Rig_Settings")
+		    cmds.addAttr(longName= ( side + "ThighTwist3Amount" ), defaultValue=.25, minValue=0, maxValue=1, keyable = True)
+		    
+		    #create the manual twist control setup
+		    twistMod = cmds.duplicate("driver_thigh_twist_0" + str(i + 1) + "_" + side , po = True, name = "thigh_" + side + "_twist3_mod")[0]
+		    cmds.parent(twistMod, rollGrp)
+        
+        
+		    #create the manual twist control
+		    twistCtrl = self.createControl("circle", 30, "thigh_" + side + "_twist3_anim")
+		    cmds.setAttr(twistCtrl + ".ry", -90)
+		    cmds.makeIdentity(twistCtrl, r = 1, apply =True)
+		    
+		    constraint = cmds.parentConstraint(twistMod, twistCtrl)[0]
+		    cmds.delete(constraint)
+		    
+		    twistCtrlGrp = cmds.group(empty = True, name = "thigh_" + side + "_twist3_anim_grp")
+		    constraint = cmds.parentConstraint(twistMod, twistCtrlGrp)[0]
+		    cmds.delete(constraint)
+		    
+		    cmds.parent(twistCtrl, twistCtrlGrp)
+		    cmds.parent(twistCtrlGrp, twistMod)
+		    cmds.makeIdentity(twistCtrl, t = 1, r = 1, s = 1, apply = True)
+		    
+		    cmds.connectAttr("hip_anim." + side + "TwistCtrlVis", twistCtrl + ".v")
+		    cmds.connectAttr("hip_anim." + side + "TwistCtrlVis", twistMod + ".v")
+		    
+		    #lock attrs on twist ctrl and color
+		    for attr in [".sx", ".sy", ".sz"]:
+			cmds.setAttr(twistCtrl + attr, lock = True, keyable = False)
+		    
+		    cmds.setAttr(twistCtrl + ".v", keyable = False)
+		    cmds.setAttr(twistCtrl + ".overrideEnabled", 1)
+		    cmds.setAttr(twistCtrl + ".overrideColor", color)
+	
+		    #drive the twist joint
+		    twistMultNode = cmds.shadingNode("multiplyDivide", asUtility = True, name = "thigh_twist_3_" + side + "_mult_node")
+		    cmds.connectAttr("driver_thigh_twist_01_" + side + ".rx", twistMultNode + ".input1X")
+		    cmds.connectAttr("Rig_Settings." + side + "ThighTwist3Amount", twistMultNode + ".input2X")
+		    cmds.connectAttr(twistMultNode + ".outputX", twistCtrlGrp + ".rx")
+		    cmds.parentConstraint(twistCtrl, "driver_thigh_twist_0" + str(i + 1) + "_" + side, mo = True)
+        
+        #clean up hierarchy
+        cmds.parent([rollGrp, driverGrp], "leg_group_" + side)
+        
+
+    # # # # # # # # # # # # # # # # # # # # # # # # # # # # # # # # # # # # # # # # # # # # # # # # # # # # # # # # # # # # # # # 
+    # # # # # # # # # # # # # # # # # # # # # # # # # # # # # # # # # # # # # # # # # # # # # # # # # # # # # # # # # # # # # # # 
+    # # # # # # # # # # # # # # # # # # # # # # # # # # # # # # # # # # # # # # # # # # # # # # # # # # # # # # # # # # # # # # # 
+    # # # # # # # # # # # # # # # # # # # # # # # # # # # # # # # # # # # # # # # # # # # # # # # # # # # # # # # # # # # # # # # 
+    # # # # # # # # # # # # # # # # # # # # # # # # # # # # # # # # # # # # # # # # # # # # # # # # # # # # # # # # # # # # # # # 
+    def buildAutoHips(self):
+        
+        hipWorld = cmds.group(empty = True, name = "auto_hip_world")
+        yzRot = cmds.group(empty = True, name = "auto_hip_yz_rot_solver")
+        xRot = cmds.group(empty = True, name = "auto_hip_x_rot_solver")
+        legSys = cmds.group(empty = True, name = "auto_hip_legs_system")
+        switchNode = cmds.spaceLocator(name = "auto_hip_on_off")[0]
+        hipRedirect = cmds.spaceLocator(name = "hip_ctrl_redirect")[0]
+        
+        #hide locators
+        for node in [switchNode, hipRedirect]:
+	    shape = cmds.listRelatives(node, children = True)[0]
+	    try:
+		cmds.setAttr(shape + ".v", 0)
+		cmds.setAttr(shape + ".v", lock = True)
+	    except:
+		pass
+        
+        for node in [hipWorld, yzRot, xRot, legSys, switchNode, hipRedirect]:
+            constraint = cmds.parentConstraint("driver_pelvis", node)[0]
+            cmds.delete(constraint)
+            
+        #setup hierarchy
+        cmds.parent("hip_anim", hipWorld)
+        cmds.parent([yzRot, xRot, legSys], "hip_anim")
+        cmds.parent(switchNode, yzRot)
+        cmds.parent(hipRedirect, switchNode)
+        
+        
+        #create the fk knee space locators
+        pelvisPos = cmds.xform("driver_pelvis", q = True, t = True, ws = True)
+        height = pelvisPos[2]
+        
+        
+        leftKneeLoc = cmds.spaceLocator(name = "auto_hips_knee_loc_l")[0]
+        rightKneeLoc = cmds.spaceLocator(name = "auto_hips_knee_loc_r")[0]
+        
+        constraint = cmds.pointConstraint("ik_leg_calf_l", leftKneeLoc)[0]
+        cmds.delete(constraint)
+        
+        constraint = cmds.pointConstraint("ik_leg_calf_r", rightKneeLoc)[0]
+        cmds.delete(constraint)
+        
+        cmds.makeIdentity([leftKneeLoc, rightKneeLoc], t = 1, r = 1, s = 1, apply = True)
+        
+        leftKneeLocConstraint = cmds.pointConstraint(["invis_legs_ik_leg_calf_l", "invis_legs_fk_calf_l_anim"], leftKneeLoc)[0]
+        rightKneeLocConstraint = cmds.pointConstraint(["invis_legs_ik_leg_calf_r","invis_legs_fk_calf_r_anim"], rightKneeLoc)[0]
+        
+        leftTargets = cmds.pointConstraint(leftKneeLocConstraint, q = True, weightAliasList = True)
+        rightTargets = cmds.pointConstraint(rightKneeLocConstraint, q = True, weightAliasList = True)
+        
+        cmds.setAttr(leftKneeLocConstraint + "." + leftTargets[1], 0)
+        cmds.setAttr(rightKneeLocConstraint + "." + rightTargets[1], 0)
+        
+        
+        #create the thigh locators for solving x (twist)
+        leftThighLoc = cmds.spaceLocator(name = "auto_hips_x_rot_solv_l")[0]
+        rightThighLoc = cmds.spaceLocator(name = "auto_hips_x_rot_solv_r")[0]
+        
+        constraint = cmds.pointConstraint("driver_thigh_l", leftThighLoc)[0]
+        cmds.delete(constraint)
+        
+        constraint = cmds.pointConstraint("driver_thigh_r", rightThighLoc)[0]
+        cmds.delete(constraint)
+        
+        cmds.setAttr(leftThighLoc + ".tz", height)
+        cmds.setAttr(rightThighLoc + ".tz", height)
+        
+        
+        #create the x rotation solver joint chain
+        cmds.select(clear = True)
+        xRotJointStart = cmds.joint(name = "auto_hips_x_rot_solv_joint_start")
+        cmds.select(clear = True)
+                    
+        constraint = cmds.pointConstraint(leftThighLoc, xRotJointStart)[0]
+        cmds.delete(constraint)
+        
+        constraint = cmds.orientConstraint("driver_pelvis", xRotJointStart)[0]
+        cmds.delete(constraint)
+        
+        xRotJointEnd = cmds.duplicate(xRotJointStart, name = "auto_hips_x_rot_solv_joint_end")[0]
+        cmds.parent(xRotJointEnd, xRotJointStart)
+        cmds.xform(xRotJointEnd, ws = True, t = (pelvisPos[0], pelvisPos[1], pelvisPos[2]))
+        
+        cmds.makeIdentity(xRotJointStart, r = 1, apply = True)
+        cmds.pointConstraint(leftThighLoc, xRotJointStart)
+        
+        xRotIKNodes = cmds.ikHandle(sol = "ikSCsolver", name = "auto_hips_x_rot_solver_ik", sj = xRotJointStart, ee = xRotJointEnd)
+        cmds.pointConstraint(rightThighLoc, xRotIKNodes[0])
+        
+        
+        #hookup x rot solver
+        upAxis = self.getUpAxis("hip_anim")
+        cmds.connectAttr(xRotJointStart + ".rotate" + upAxis, legSys + ".rotate" + upAxis)
+        
+        #hookup motion of thigh locators
+        rotXikHipL = cmds.spaceLocator(name = "auto_hips_ik_x_rot_solv_loc_l")[0]
+        rotXikHipR = cmds.spaceLocator(name = "auto_hips_ik_x_rot_solv_loc_r")[0]
+        
+        constraint = cmds.pointConstraint("noflip_pv_loc_l", rotXikHipL)[0]
+        cmds.delete(constraint)
+        
+        constraint = cmds.pointConstraint("noflip_pv_loc_r", rotXikHipR)[0]
+        cmds.delete(constraint)
+        
+        cmds.pointConstraint(leftKneeLoc, leftThighLoc, mo = True, skip = ["x", "z"])
+        cmds.pointConstraint(rightKneeLoc, rightThighLoc, mo = True, skip = ["x", "z"])
+        
+        
+        #create the multiply node for our x rot solver node to halve the results
+        xRotMult = cmds.shadingNode("multiplyDivide", name = "auto_hips_x_rot_mult_node", asUtility = True)
+        cmds.connectAttr(legSys + ".rotate" + upAxis, xRotMult + ".input1X")
+        cmds.setAttr(xRotMult + ".input2X", .5)
+        cmds.connectAttr(xRotMult + ".outputX", xRot + ".rotate" + upAxis)
+        
+        
+        #create the joints for the yz rotations
+        cmds.select(clear = True)
+        yzRotStartJoint = cmds.joint(name = "auto_hips_yz_rot_solv_joint_start")
+        
+        constraint = cmds.parentConstraint("driver_pelvis", yzRotStartJoint)[0]
+        cmds.delete(constraint)
+        yzRotEndJoint = cmds.duplicate(yzRotStartJoint, name = "auto_hips_yz_rot_solv_joint_end")[0]
+        cmds.parent(yzRotEndJoint, yzRotStartJoint)
+        cmds.setAttr(yzRotEndJoint + ".tx", (height/2)* -1)
+        
+        cmds.makeIdentity(yzRotStartJoint, r = 1, apply = True)
+    
+        yzRotikNodes = cmds.ikHandle(sol = "ikRPsolver", name = "auto_hips_yz_rot_solv_ik", sj = yzRotStartJoint, ee = yzRotEndJoint)
+        yzRotTargetLoc = cmds.spaceLocator(name = "auto_hips_yz_target_loc")[0]
+        constraint = cmds.pointConstraint([leftKneeLoc, rightKneeLoc], yzRotTargetLoc)[0]
+        cmds.delete(constraint)
+        
+        cmds.makeIdentity(yzRotTargetLoc, t = 1, apply = True)
+        cmds.parent(yzRotikNodes[0], yzRotTargetLoc)
+        
+        
+        #setup motion for yz solver
+        cmds.pointConstraint([leftKneeLoc, rightKneeLoc], yzRotTargetLoc)
+        yzSolvConst = cmds.orientConstraint(yzRotStartJoint, yzRot)[0]
+        
+        
+        #setup distance tools for reducing popping when foot gets close to pelvis
+        cmds.select(clear = True)
+        yzRotFlipCtrlNodesL = cmds.duplicate("noflip_begin_joint_l", name = "auto_hips_dist_ctrl_begin_joint_l", rc = True)
+        yzRotFlipCtrlBeginL = yzRotFlipCtrlNodesL[0]
+        yzRotFlipCtrlEndL = cmds.rename(yzRotFlipCtrlNodesL[1], "auto_hips_dist_ctrl_end_joint_l")
+        
+
+        #now for the right side
+        cmds.select(clear = True)
+        yzRotFlipCtrlNodesR = cmds.duplicate("noflip_begin_joint_r", name = "auto_hips_dist_ctrl_begin_joint_r", rc = True)
+        yzRotFlipCtrlBeginR = yzRotFlipCtrlNodesR[0]
+        yzRotFlipCtrlEndR = cmds.rename(yzRotFlipCtrlNodesR[1], "auto_hips_dist_ctrl_end_joint_r")
+        
+
+        #setup the distance mover joint
+        distMoverL = cmds.duplicate(yzRotFlipCtrlEndL, name = "distance_mover_joint_l")[0]
+        distMoverR = cmds.duplicate(yzRotFlipCtrlEndR, name = "distance_mover_joint_r")[0]
+        
+        cmds.pointConstraint("ik_foot_anim_l", distMoverL, mo = True)
+        cmds.pointConstraint("ik_foot_anim_r", distMoverR, mo = True)
+
+        
+        #setup system for fixing flipping when foot gets too close to pelvis
+        originalLen = cmds.getAttr(yzRotFlipCtrlEndL + ".tz")
+        conditionNodeL = cmds.shadingNode("condition", asUtility = True, name = "autoHips_flipFix_condition_l")
+        conditionNodeR = cmds.shadingNode("condition", asUtility = True, name = "autoHips_flipFix_condition_r")
+        
+        cmds.setAttr(conditionNodeL + ".secondTerm", originalLen)
+        cmds.setAttr(conditionNodeR + ".secondTerm", originalLen)
+        
+        cmds.connectAttr(distMoverL + ".tz", conditionNodeL + ".firstTerm")
+        cmds.connectAttr(distMoverR + ".tz", conditionNodeR + ".firstTerm")
+        
+        scaleFactor = self.getScaleFactor()
+        
+        cmds.setAttr(conditionNodeL + ".operation", 2)
+        cmds.setAttr(conditionNodeR + ".operation", 2)
+        cmds.setAttr(conditionNodeL + ".colorIfFalseR", 0)
+        cmds.setAttr(conditionNodeR + ".colorIfFalseR", 0)
+        cmds.setAttr(conditionNodeL + ".colorIfTrueR", -60 * scaleFactor)
+        cmds.setAttr(conditionNodeR + ".colorIfTrueR", -60 * scaleFactor)
+        
+        cmds.connectAttr(conditionNodeL + ".outColorR", "noflip_aim_soft_grp_l.tz")
+        cmds.connectAttr(conditionNodeR + ".outColorR", "noflip_aim_soft_grp_r.tz")
+        
+        
+        
+        #hookup the hips to use the x and yz rot solver data
+        cmds.connectAttr(xRot + ".rotate" + upAxis, yzSolvConst + ".offset" + upAxis)
+        cmds.connectAttr("hip_anim.rotate", hipRedirect + ".rotate")
+        
+        #orient constrain the on/off node to the world control and the yz rot solver
+        onOffConstraint = cmds.orientConstraint([hipWorld, yzSolvConst], switchNode)[0]
+        
+        #add auto on/off attr to hip control. hookup connections to constraint
+        cmds.select("hip_anim")
+        cmds.addAttr(longName='autoHips', defaultValue=0, minValue=0, maxValue=1, keyable = True)
+        
+        
+        orientTargets = cmds.orientConstraint(onOffConstraint, q = True, weightAliasList = True)
+        
+        cmds.connectAttr("hip_anim.autoHips", onOffConstraint + "." + orientTargets[1])
+        reverseNode = cmds.shadingNode("reverse", asUtility = True, name = "autoHips_reverse_node_onOff")
+        cmds.connectAttr("hip_anim.autoHips", reverseNode + ".inputX")
+        cmds.connectAttr(reverseNode + ".outputX", onOffConstraint + "." + orientTargets[0])
+        
+        
+        
+        #plug the body anim's rotates into the onOffConstraint's offsets
+        
+        bodyMultNode = cmds.shadingNode("multiplyDivide", asUtility = True, name = "autoHips_body_rotation_fix")
+        cmds.connectAttr("body_anim.rotateX", bodyMultNode + ".input1X")
+        cmds.connectAttr("body_anim.rotateY", bodyMultNode + ".input1Y")
+        cmds.connectAttr("body_anim.rotateZ", bodyMultNode + ".input1Z")
+        
+        cmds.connectAttr("hip_anim.autoHips", bodyMultNode + ".input2X")
+        cmds.connectAttr("hip_anim.autoHips", bodyMultNode + ".input2Y")
+        cmds.connectAttr("hip_anim.autoHips", bodyMultNode + ".input2Z")
+        
+        #cmds.connectAttr(bodyMultNode + ".outputX", onOffConstraint + ".offsetX")
+        #cmds.connectAttr(bodyMultNode + ".outputY", onOffConstraint + ".offsetY")
+        #cmds.connectAttr(bodyMultNode + ".outputZ", onOffConstraint + ".offsetZ")
+        
+        #update invisible leg IK joints
+        cmds.setToolTo('moveSuperContext')
+        cmds.select("invis_legs_ik_leg_thigh_l", hi = True)
+        cmds.setToolTo('RotateSuperContext')
+        cmds.select(clear = True)
+        
+        
+        #constrain driver joint to control
+        cmds.parentConstraint(hipRedirect, "driver_pelvis", mo = True)
+        
+        #parent hip world to body anim
+        cmds.parent(hipWorld, "body_anim")
+        
+        #clean up hierarchy
+        autoHipsMasterGrp = cmds.group(empty = True, name = "autoHips_sys_grp")
+        cmds.parent([leftKneeLoc, rightKneeLoc, leftThighLoc, rightThighLoc, xRotJointStart, rotXikHipL, rotXikHipR, yzRotStartJoint,xRotIKNodes[0], yzRotTargetLoc], autoHipsMasterGrp)
+        
+        #cosntrain ik leg bones(thighs) to the hip ctrl redirect
+        cmds.parentConstraint(hipRedirect, "leg_joints_grp_l", mo = True)
+        cmds.parentConstraint(hipRedirect, "leg_joints_grp_r", mo = True)
+        
+        #parent the bottom of the spline ik spine bone to the hip redirect
+	if cmds.objExists("spine_splineIK_bottom_joint"):
+	    cmds.parent("spine_splineIK_bottom_joint", hipRedirect)
+        
+        #hide stuff
+        cmds.setAttr("autoHips_sys_grp.v", 0)
+        
+        
+        #hook up autohips to leg mode
+        
+        reverseNodeL = "legSwitcher_reverse_node_l"
+        reverseNodeR = "legSwitcher_reverse_node_r"
+        
+        targets = cmds.pointConstraint(leftKneeLocConstraint, q = True, weightAliasList = True)
+        cmds.connectAttr("Rig_Settings" + ".lLegMode", leftKneeLocConstraint + "." + targets[0])
+        cmds.connectAttr(reverseNodeL + ".outputX", leftKneeLocConstraint + "." + targets[1])
+        
+        
+        targets = cmds.pointConstraint(rightKneeLocConstraint, q = True, weightAliasList = True)
+        cmds.connectAttr("Rig_Settings" + ".rLegMode", rightKneeLocConstraint + "." + targets[0])
+        cmds.connectAttr(reverseNodeR + ".outputX", rightKneeLocConstraint + "." + targets[1])
+        
+    
+        
+        
+        
+
+    # # # # # # # # # # # # # # # # # # # # # # # # # # # # # # # # # # # # # # # # # # # # # # # # # # # # # # # # # # # # # # # 
+    # # # # # # # # # # # # # # # # # # # # # # # # # # # # # # # # # # # # # # # # # # # # # # # # # # # # # # # # # # # # # # # 
+    # # # # # # # # # # # # # # # # # # # # # # # # # # # # # # # # # # # # # # # # # # # # # # # # # # # # # # # # # # # # # # # 
+    # # # # # # # # # # # # # # # # # # # # # # # # # # # # # # # # # # # # # # # # # # # # # # # # # # # # # # # # # # # # # # # 
+    # # # # # # # # # # # # # # # # # # # # # # # # # # # # # # # # # # # # # # # # # # # # # # # # # # # # # # # # # # # # # # # 
+    def autoSpine(self):
+        
+        numSpineBones = cmds.getAttr("Skeleton_Settings.numSpineBones")
+	
+	if numSpineBones > 2:
+	    
+	    #drive the mid IK spine control based on what the upper spine control is doing
+	    
+	    
+	    #drives the ik spine controls based on auto hip attr and hip motion
+	    yzRotSolver = "auto_hip_yz_rot_solver"
+	    midDriver = "mid_ik_anim_driver_grp"
+	    midDriverTop = "mid_ik_anim_translate_driver_grp"
+	    
+	    
+	    #add auto on/off attr to hip control. hookup connections to constraint
+	    cmds.select("chest_ik_anim")
+	    cmds.addAttr(longName='autoSpine', defaultValue=0, minValue=0, maxValue=1, keyable = True)
+	    cmds.addAttr(longName='rotationInfluence', defaultValue=.25, minValue=0, maxValue=1, keyable = True)
+	    
+	    topCtrlMultRY = cmds.shadingNode("multiplyDivide", asUtility = True, name = "autoSpine_top_driver_mult_ry")
+	    topCtrlMultRZ = cmds.shadingNode("multiplyDivide", asUtility = True, name = "autoSpine_top_driver_mult_rz")
+	    topCtrlMultSwitchRY = cmds.shadingNode("multiplyDivide", asUtility = True, name = "autoSpine_top_mult_switch_ry")
+	    topCtrlMultSwitchRZ = cmds.shadingNode("multiplyDivide", asUtility = True, name = "autoSpine_top_mult_switch_rz")
+	    
+	    
+	    #create a node that will track all world space translations and rotations on the chest IK anim
+	    chestMasterTrackNode = cmds.spaceLocator(name = "chest_ik_track_parent")[0]
+	    constraint = cmds.parentConstraint("chest_ik_anim", chestMasterTrackNode)[0]
+	    cmds.delete(constraint)
+	    
+	    chestTrackNode = cmds.spaceLocator(name = "chest_ik_tracker")[0]
+	    constraint = cmds.parentConstraint("chest_ik_anim", chestTrackNode)[0]
+	    cmds.delete(constraint)
+	    cmds.parent(chestTrackNode, chestMasterTrackNode)
+	    cmds.parentConstraint("chest_ik_anim", chestTrackNode)
+	    
+	    cmds.parent(chestMasterTrackNode, "body_anim")
+	    
+	    #hide locator
+	    cmds.setAttr(chestMasterTrackNode + ".v", 0)
+	    
+	    
+	    
+
+	
+	    #Rotate Y
+	    cmds.connectAttr(chestTrackNode + ".ry", topCtrlMultRY + ".input1X")
+	    cmds.connectAttr("chest_ik_anim.rotationInfluence", topCtrlMultRY + ".input2X")
+	    
+	    cmds.connectAttr(topCtrlMultRY + ".outputX", topCtrlMultSwitchRY + ".input1X")
+	    cmds.connectAttr("chest_ik_anim.autoSpine", topCtrlMultSwitchRY + ".input2X")
+	    cmds.connectAttr(topCtrlMultSwitchRY + ".outputX", midDriver + ".tz")	    
+	    
+	    
+	    #Rotate Z
+	    multInverse = cmds.shadingNode("multiplyDivide", asUtility = True, name = "autoSpine_mult_rz_inverse")
+	    cmds.connectAttr("chest_ik_anim.rotationInfluence", multInverse + ".input1X")
+	    cmds.setAttr(multInverse + ".input2X", -1)
+	    
+	    
+	    cmds.connectAttr(chestTrackNode + ".rz", topCtrlMultRZ + ".input1X")
+	    cmds.connectAttr(multInverse + ".outputX", topCtrlMultRZ + ".input2X")
+	    
+	    cmds.connectAttr(topCtrlMultRZ + ".outputX", topCtrlMultSwitchRZ + ".input1X")
+	    cmds.connectAttr("chest_ik_anim.autoSpine", topCtrlMultSwitchRZ + ".input2X")
+	    cmds.connectAttr(topCtrlMultSwitchRZ + ".outputX", midDriver + ".ty")
+	    
+	    
+	    #Translate X
+	    #Chest Control Translate X + Hip Control Translate X / 2 * autpSpine
+	    autoSpineTXNode = cmds.shadingNode("plusMinusAverage", asUtility = True, name = midDriverTop + "_TX_Avg")
+	    cmds.setAttr(autoSpineTXNode + ".operation", 3)
+	    autoSpineTX_MultNode = cmds.shadingNode("multiplyDivide", asUtility = True, name = midDriverTop + "_TX_Mult")
+	    
+	    cmds.connectAttr("chest_ik_anim.translateX", autoSpineTXNode + ".input1D[0]")
+	    cmds.connectAttr("hip_anim.translateX", autoSpineTXNode + ".input1D[1]")
+	    cmds.connectAttr(autoSpineTXNode + ".output1D", autoSpineTX_MultNode + ".input1X")
+	    cmds.connectAttr("chest_ik_anim.autoSpine", autoSpineTX_MultNode + ".input2X")
+	    cmds.connectAttr(autoSpineTX_MultNode + ".outputX", midDriverTop + ".translateX")
+	    
+	    
+	    
+	    #Translate Y
+	    autoSpineTYNode = cmds.shadingNode("plusMinusAverage", asUtility = True, name = midDriverTop + "_TY_Avg")
+	    cmds.setAttr(autoSpineTYNode + ".operation", 3)
+	    autoSpineTY_MultNode = cmds.shadingNode("multiplyDivide", asUtility = True, name = midDriverTop + "_TY_Mult")
+	    
+	    cmds.connectAttr(chestTrackNode + ".translateY", autoSpineTYNode + ".input1D[0]")
+	    cmds.connectAttr("hip_anim.translateY", autoSpineTYNode + ".input1D[1]")
+	    cmds.connectAttr(autoSpineTYNode + ".output1D", autoSpineTY_MultNode + ".input1X")
+	    cmds.connectAttr("chest_ik_anim.autoSpine", autoSpineTY_MultNode + ".input2X")
+	    cmds.connectAttr(autoSpineTY_MultNode + ".outputX", midDriverTop + ".translateY")
+	    
+	    
+	    
+	    #Translate Z
+	    autoSpineTZNode = cmds.shadingNode("plusMinusAverage", asUtility = True, name = midDriverTop + "_TZ_Avg")
+	    cmds.setAttr(autoSpineTZNode + ".operation", 3)
+	    autoSpineTZ_MultNode = cmds.shadingNode("multiplyDivide", asUtility = True, name = midDriverTop + "_TZ_Mult")
+	    
+	    cmds.connectAttr(chestTrackNode + ".translateZ", autoSpineTZNode + ".input1D[0]")
+	    cmds.connectAttr("hip_anim.translateZ", autoSpineTZNode + ".input1D[1]")
+	    cmds.connectAttr(autoSpineTZNode + ".output1D", autoSpineTZ_MultNode + ".input1X")
+	    cmds.connectAttr("chest_ik_anim.autoSpine", autoSpineTZ_MultNode + ".input2X")
+	    cmds.connectAttr(autoSpineTZ_MultNode + ".outputX", midDriverTop + ".translateZ")	    
+	    
+
+
+        
+    # # # # # # # # # # # # # # # # # # # # # # # # # # # # # # # # # # # # # # # # # # # # # # # # # # # # # # # # # # # # # # # 
+    # # # # # # # # # # # # # # # # # # # # # # # # # # # # # # # # # # # # # # # # # # # # # # # # # # # # # # # # # # # # # # # 
+    # # # # # # # # # # # # # # # # # # # # # # # # # # # # # # # # # # # # # # # # # # # # # # # # # # # # # # # # # # # # # # # 
+    # # # # # # # # # # # # # # # # # # # # # # # # # # # # # # # # # # # # # # # # # # # # # # # # # # # # # # # # # # # # # # # 
+    # # # # # # # # # # # # # # # # # # # # # # # # # # # # # # # # # # # # # # # # # # # # # # # # # # # # # # # # # # # # # # # 
+    def hookupSpine(self, ikControls, fkControls):
+        
+        #constrain the driver joints to the FK and IK controls
+        i = 0
+	upAxis = "X"
+	
+        for joint in ikControls:
+            
+            driverJoint = joint.partition("splineIK_")[2]
+            driverJoint = "driver_" + driverJoint
+	    
+            
+            if joint == ikControls[0]:
+		upAxis = self.getUpAxis(driverJoint)
+		
+		if upAxis == "X":
+		    axisB = "Y"
+		    axisC = "Z"
+		
+		if upAxis == "Y":
+		    axisB = "X"
+		    axisC = "Z"
+		    
+		if upAxis == "Z":
+		    axisB = "X"
+		    axisC = "Y"
+		    
+		    
+		twistJoint = "twist_" + joint
+		cmds.parentConstraint([twistJoint, fkControls[i]], driverJoint, mo = True)
+		cmds.scaleConstraint([twistJoint, fkControls[i]], driverJoint, mo = True)
+		scaleDriverJoint = driverJoint
+
+		    
+            if joint != ikControls[len(ikControls) - 1] and joint != ikControls[0]:
+                twistJoint = "twist_" + joint
+                cmds.parentConstraint([twistJoint, fkControls[i]], driverJoint, mo = True)
+		
+		#create blendColors nodes for scale
+		blenderNodeScale = cmds.shadingNode("blendColors", asUtility = True, name = driverJoint + "_blenderNodeScale")
+		cmds.connectAttr(twistJoint + ".scale" + axisB, blenderNodeScale + ".color1R")		
+		cmds.connectAttr(fkControls[i] + ".scale" + axisB, blenderNodeScale + ".color2R")	
+		cmds.connectAttr(twistJoint + ".scale" + axisC, blenderNodeScale + ".color1G")		
+		cmds.connectAttr(fkControls[i] + ".scale" + axisC, blenderNodeScale + ".color2G")
+		cmds.connectAttr(blenderNodeScale + ".outputR", driverJoint + ".scale" + axisB)
+		cmds.connectAttr(blenderNodeScale + ".outputG", driverJoint + ".scale" + axisC)
+		
+		
+            if joint == ikControls[len(ikControls) - 1]:
+                topJoint = "spine_splineIK_top_joint"
+                twistJoint = "twist_" + joint
+		
+                cmds.orientConstraint([topJoint, fkControls[i]], driverJoint, mo = True)
+                cmds.pointConstraint([twistJoint, fkControls[i]], driverJoint, mo = True)
+		
+            i = i + 1
+            
+            
+        
+        #add attributes to the Rig_Settings node
+        cmds.select("Rig_Settings")
+        cmds.addAttr(longName='spine_ik', defaultValue=0, minValue=0, maxValue=1, keyable = True)
+        cmds.addAttr(longName='spine_fk', defaultValue=0, minValue=0, maxValue=1, keyable = True)
+        
+
+        #hookup Rig_Settings attrs to the parentConstraint weights
+        
+        #find connection targets
+        for ctrl in ikControls:
+            
+            driverJoint = ctrl.partition("splineIK_")[2]
+            driverJoint = "driver_" + driverJoint
+            constraint = driverJoint + "_parentConstraint1"
+	    
+	    #hook up blendColors scale node
+	    try:
+		blenderNodeScale = driverJoint + "_blenderNodeScale"
+		cmds.connectAttr("Rig_Settings" + ".spine_ik", blenderNodeScale + ".blender")
+	    except:
+		pass
+	    
+	    
+	    
+            if cmds.objExists(constraint):
+                targets = cmds.parentConstraint(constraint, q = True, weightAliasList = True)
+                
+                for target in targets:
+                    if target.find("IK") != -1:
+                        cmds.connectAttr("Rig_Settings" + ".spine_ik", constraint + "." + target)
+			
+		    else:
+			cmds.connectAttr("Rig_Settings" + ".spine_fk", constraint + "." + target)
+			
+
+		    
+                        
+                    
+                
+                
+            else:
+                pointConstraint = driverJoint + "_pointConstraint1"
+                orientConstraint = driverJoint + "_orientConstraint1"
+                
+                pointTargets = cmds.pointConstraint(pointConstraint, q = True, weightAliasList = True)
+                orientTargets = cmds.orientConstraint(orientConstraint, q = True, weightAliasList = True)
+                
+                for target in pointTargets:
+                    if target.find("IK") != -1:
+                        cmds.connectAttr("Rig_Settings" + ".spine_ik", pointConstraint + "." + target)
+			
+		    else:
+			cmds.connectAttr("Rig_Settings" + ".spine_fk", pointConstraint + "." + target)
+                        
+                        
+                for target in orientTargets:
+                    if target.find("IK") != -1:
+                        cmds.connectAttr("Rig_Settings" + ".spine_ik", orientConstraint + "." + target)
+
+    
+		    else:
+			cmds.connectAttr("Rig_Settings" + ".spine_fk", orientConstraint + "." + target)
+                
+                
+        #hook up spine control vis to the rig settings
+        cmds.connectAttr("Rig_Settings" + ".spine_fk", "spine_01_anim_grp.v")
+        cmds.connectAttr("Rig_Settings" + ".spine_ik", "mid_ik_anim_grp.v")
+        cmds.connectAttr("Rig_Settings" + ".spine_ik", "chest_ik_anim_grp.v")
+	
+        
+        #set spine rig to be IK by default
+        cmds.setAttr("Rig_Settings" + ".spine_ik", 1)
+        
+            
+    # # # # # # # # # # # # # # # # # # # # # # # # # # # # # # # # # # # # # # # # # # # # # # # # # # # # # # # # # # # # # # # 
+    # # # # # # # # # # # # # # # # # # # # # # # # # # # # # # # # # # # # # # # # # # # # # # # # # # # # # # # # # # # # # # # 
+    # # # # # # # # # # # # # # # # # # # # # # # # # # # # # # # # # # # # # # # # # # # # # # # # # # # # # # # # # # # # # # # 
+    # # # # # # # # # # # # # # # # # # # # # # # # # # # # # # # # # # # # # # # # # # # # # # # # # # # # # # # # # # # # # # # 
+    # # # # # # # # # # # # # # # # # # # # # # # # # # # # # # # # # # # # # # # # # # # # # # # # # # # # # # # # # # # # # # # 
+    def buildFKArms(self):
+        
+        #duplicate our driver joints to create our FK arm skeleton
+        for side in ["l", "r"]:
+            
+            #if the side has an arm, create the fk joints
+            if cmds.objExists("driver_upperarm_" + side):
+                autoClavJointStart = cmds.duplicate("driver_clavicle_" + side, po = True, name = "auto_clavicle_" + side)[0]
+                clavJoint = cmds.duplicate("driver_clavicle_" + side, po = True, name = "rig_clavicle_" + side)[0]
+                fkArmJoint = cmds.duplicate("driver_upperarm_" + side, po = True, name = "fk_upperarm_" + side)[0]
+                fkElbowJoint = cmds.duplicate("driver_lowerarm_" + side, po = True, name = "fk_lowerarm_" + side)[0]
+                fkWristJoint = cmds.duplicate("driver_hand_" + side, po = True, name = "fk_hand_" + side)[0]
+                
+                #parent the fk upperarm to the world
+                parent = cmds.listRelatives(clavJoint, parent = True)
+                if parent != None:
+                    cmds.parent(clavJoint, world = True)
+                    
+                #recreate the fk arm hierarchy
+                cmds.parent(fkArmJoint, clavJoint)
+                cmds.parent(fkElbowJoint, fkArmJoint)
+                cmds.parent(fkWristJoint, fkElbowJoint)
+		cmds.makeIdentity(fkArmJoint, t = 0, r = 1, s = 0, apply = True)
+                
+                #set rotation order on fk arm joint
+		cmds.setAttr(fkArmJoint + ".rotateOrder", 3)
+		
+                #create the shoulder hierarchy
+                parent = cmds.listRelatives(autoClavJointStart, parent = True)
+                if parent != None:
+                    cmds.parent(autoClavJointStart, world = True)
+                    
+                autoClavEndJoint = cmds.duplicate(fkArmJoint, parentOnly = True, name = "auto_clavicle_end_" + side)[0]
+                pos = cmds.xform(autoClavEndJoint, q = True, ws = True, t = True)
+                zPos = cmds.xform(autoClavJointStart, q = True, ws = True, t = True)[2]
+                cmds.xform(autoClavEndJoint, ws = True, t = [pos[0], pos[1], zPos])
+                cmds.parent(autoClavEndJoint, autoClavJointStart)
+                
+                
+                #create the IK for the clavicle
+                ikNodes = cmds.ikHandle(sj = autoClavJointStart, ee = autoClavEndJoint, sol = "ikSCsolver", name = "auto_clav_to_elbow_ikHandle_" + side)[0]
+                
+                #position the IK handle at the elbow joint
+                constraint = cmds.pointConstraint(fkElbowJoint, ikNodes)[0]
+                cmds.delete(constraint)
+                
+                #create our autoClav world grp
+                autoClavWorld =  cmds.group(empty = True, name = "auto_clav_world_grp_" + side)
+                constraint = cmds.pointConstraint(autoClavEndJoint, autoClavWorld)[0]
+                cmds.delete(constraint)
+                
+                cmds.makeIdentity(autoClavWorld, t = 1, r = 1, s = 1, apply = True)
+                
+                #duplicate the FK arm to create our invisible arm
+                invisUpArm = cmds.duplicate(fkArmJoint, po = True, name = "invis_upperarm_" + side)[0]
+                invisLowArm = cmds.duplicate(fkElbowJoint, po = True, name = "invis_lowerarm_" + side)[0]
+                invisHand = cmds.duplicate(fkWristJoint, po = True, name = "invis_hand_" + side)[0]
+                
+                cmds.parent(invisHand, invisLowArm)
+                cmds.parent(invisLowArm, invisUpArm)
+                cmds.parent(invisUpArm, autoClavWorld)
+                
+                
+                #create our upperarm orient locator
+                invisArmOrient = cmds.spaceLocator(name = "invis_arm_orient_loc_" + side)[0]
+                invisArmOrientGrp = cmds.group(empty = True, name = "invis_arm_orient_loc_grp_" + side)
+                
+                constraint = cmds.parentConstraint(fkArmJoint, invisArmOrient)[0]
+                cmds.delete(constraint)
+                
+                constraint = cmds.parentConstraint(fkArmJoint, invisArmOrientGrp)[0]
+                cmds.delete(constraint)
+                
+                cmds.parent(invisArmOrient, invisArmOrientGrp)
+                
+                
+                
+                #create the invis arm fk control (to derive autoClav info)
+                invisFkArmCtrl = self.createControl("circle", 20, "invis_fk_arm_" + side + "_anim")
+                cmds.setAttr(invisFkArmCtrl + ".ry", -90)
+                cmds.makeIdentity(invisFkArmCtrl, r = 1, apply =True)
+                
+                constraint = cmds.parentConstraint(fkArmJoint, invisFkArmCtrl)[0]
+                cmds.delete(constraint)
+                
+                invisFkArmCtrlGrp = cmds.group(empty = True, name = "invis_fk_arm_" + side + "_grp")
+                constraint = cmds.parentConstraint(fkArmJoint, invisFkArmCtrlGrp)[0]
+                cmds.delete(constraint)
+                cmds.parent(invisFkArmCtrl, invisFkArmCtrlGrp)
+                
+                
+                
+                #position the arm FK control so that it is about halfway down the arm length
+                dist = (cmds.getAttr(fkElbowJoint + ".tx")) / 2
+                cmds.setAttr(invisFkArmCtrl + ".translateX", dist)
+
+
+                #set the pivot of the arm control back to the arm joint
+                piv = cmds.xform(fkArmJoint, q = True, ws = True, rotatePivot = True)
+                cmds.xform(invisFkArmCtrl, ws = True, piv = [piv[0], piv[1], piv[2]])
+                
+                #freeze transforms on the control
+                cmds.makeIdentity(invisFkArmCtrl, t = 1, r = 1, s = 1, apply = True)
+                
+                
+                
+                #parent the orient arm grp to the fk ctrl
+                cmds.parent(invisArmOrientGrp, invisFkArmCtrl)
+                
+                #duplicate the invis arm fk control setup for the real fk control(upper arm)
+                dupeNodes = cmds.duplicate(invisFkArmCtrlGrp, rc = True)
+                
+                for node in dupeNodes:
+                    name = node.partition("invis_")[2]
+                    if name.find("1") != -1:
+                        name = name.partition("1")[0]
+                        
+                    cmds.rename(node, name)
+                    
+                
+                #orient constrain the invis fk up arm to the invis up arm orient loc. Also do this for the real arm
+                cmds.orientConstraint(invisArmOrient, invisUpArm)
+                cmds.orientConstraint("arm_orient_loc_" + side, fkArmJoint)
+		
+
+                #connect invis arm ctrl rotates to be driven by real arm control rotates
+                cmds.connectAttr("fk_arm_" + side + "_anim.rotate", invisFkArmCtrl + ".rotate")
+                
+                
+                #the following section of code will essentially give us a vector from the clav joint to the elbow. This will help to drive the rotations of the auto clav
+                #create our locators to determine elbow's position in space (will drive the ik handle on the auto clav)
+                autoTransLoc = cmds.spaceLocator(name = "elbow_auto_trans_loc_" + side)[0]
+                constraint = cmds.pointConstraint(fkElbowJoint, autoTransLoc)[0]
+                cmds.delete(constraint)
+                
+                #duplicate the locator to create a parent loc
+                autoTransLocParent = cmds.duplicate(autoTransLoc, po = True, name = "elbow_auto_trans_loc_parent_" + side)[0]
+                cmds.pointConstraint(autoTransLoc, ikNodes)
+                cmds.parent(autoTransLoc, autoTransLocParent)
+                cmds.parent(autoTransLocParent, autoClavWorld)
+                
+                cmds.makeIdentity(autoTransLocParent, t = 1, r = 1, s = 1, apply = True)
+                
+                #constrain the parent trans loc(elbow) to the invis elbow joint
+                cmds.pointConstraint(invisLowArm, autoTransLocParent, mo = True)
+                
+                
+                #create our locator that will handle switching between auto clav and manual clav. position at end joint (autoClavEndJoint)
+                autoClavSwitchLoc = cmds.spaceLocator(name = "auto_clav_switch_loc_" + side)[0]
+                
+                constraint = cmds.pointConstraint(autoClavEndJoint, autoClavSwitchLoc)[0]
+                cmds.delete(constraint)
+                cmds.parent(autoClavSwitchLoc, autoClavWorld)
+                cmds.makeIdentity(autoClavSwitchLoc, t = 1, r = 1, s = 1, apply = True)
+                cmds.parent(autoClavJointStart, autoClavWorld)
+                
+                
+                #setup constraint for switching between auto/manual
+                autoClavSwitchConstraint = cmds.pointConstraint([autoClavEndJoint, autoClavWorld], autoClavSwitchLoc, mo = True)[0]
+                cmds.setAttr(autoClavSwitchConstraint + "." + autoClavWorld + "W1", 0)
+                
+                
+                #create our IK for the auto clav to move
+                autoClavIK = cmds.ikHandle(sj = clavJoint, ee = fkArmJoint, sol = "ikSCsolver", name = "auto_clav_ikHandle_" + side)[0]
+                autoClavIKGrp = cmds.group(empty = True, name = "auto_clav_ikHandle_grp_" + side)
+                constraint = cmds.pointConstraint(autoClavIK, autoClavIKGrp)[0]
+                cmds.delete(constraint)
+		autoClavIKGrpMaster = cmds.duplicate(po = True, name = "auto_clav_ikHandle_grp_master_" + side)[0]
+                
+		
+                cmds.parent(autoClavIKGrpMaster, autoClavSwitchLoc)
+		cmds.parent(autoClavIKGrp, autoClavIKGrpMaster)
+                cmds.parent(autoClavIK, autoClavIKGrp)
+                #cmds.makeIdentity(autoClavIKGrp, t = 1, r = 1, s = 1, apply = True)
+                
+                
+                
+                #create the shoulder control
+                shoulderCtrl = self.createControl("pin", 1.5, "clavicle_" + side + "_anim")
+                cmds.setAttr(shoulderCtrl + ".ry", 90)
+                cmds.setAttr(shoulderCtrl + ".rx", 90)
+            
+                constraint = cmds.pointConstraint(fkArmJoint, shoulderCtrl)[0]
+                cmds.delete(constraint)
+                
+                shoulderCtrlGrp = cmds.group(empty = True, name = "clavicle_" + side + "_anim_grp")
+                constraint = cmds.pointConstraint(fkArmJoint, shoulderCtrl)[0]
+                cmds.delete(constraint)
+                
+                cmds.parent(shoulderCtrl, shoulderCtrlGrp)
+                cmds.parent(shoulderCtrlGrp, autoClavWorld)
+                
+                cmds.makeIdentity(shoulderCtrl, t = 1, r = 1, s = 1, apply = True)
+                cmds.setAttr(shoulderCtrl + ".sx", .65)
+                cmds.setAttr(shoulderCtrl + ".sy", 1.2)
+                cmds.setAttr(shoulderCtrl + ".sz", 1.2)
+                cmds.makeIdentity(shoulderCtrl, t = 1, r = 1, s = 1, apply = True)
+                
+                
+
+                
+                
+                #connect shoulder ctrl translate to the autoClavIKGrp translate
+                cmds.connectAttr(shoulderCtrl + ".translate", autoClavIKGrp + ".translate")
+                cmds.connectAttr(autoClavSwitchLoc + ".translate", shoulderCtrl + ".rotatePivotTranslate")
+                
+                
+                #set limits on shoulder control
+		cmds.select(shoulderCtrl)
+		
+		if side == "l":
+		    cmds.transformLimits(tx = (-1,0), etx = (False, True))
+		else:
+		    cmds.transformLimits(tx = (0,1), etx = (True, False))
+                
+                
+                #create FK elbow control
+                fkElbowCtrl = self.createControl("circle", 18, "fk_elbow_" + side + "_anim")
+                cmds.setAttr(fkElbowCtrl + ".ry", -90)
+                cmds.makeIdentity(fkElbowCtrl, r = 1, apply =True)
+                
+                constraint = cmds.parentConstraint(fkElbowJoint, fkElbowCtrl)[0]
+                cmds.delete(constraint)
+                
+                fkElbowCtrlGrp = cmds.group(empty = True, name = "fk_elbow_" + side + "_anim_grp")
+                constraint = cmds.parentConstraint(fkElbowJoint, fkElbowCtrlGrp)[0]
+                cmds.delete(constraint)
+                
+                cmds.parent(fkElbowCtrl, fkElbowCtrlGrp)
+                cmds.makeIdentity(fkElbowCtrl, t = 1, r = 1, s = 1, apply = True)
+                cmds.parent(fkElbowCtrlGrp, "fk_arm_" + side + "_anim")
+                
+                #constrain elbow joint to ctrl
+                cmds.parentConstraint(fkElbowCtrl, fkElbowJoint)
+                
+                
+                #create FK wrist control
+                fkWristCtrl = self.createControl("circle", 15, "fk_wrist_" + side + "_anim")
+                cmds.setAttr(fkWristCtrl + ".ry", -90)
+                cmds.makeIdentity(fkWristCtrl, r = 1, apply =True)
+                
+                constraint = cmds.parentConstraint(fkWristJoint, fkWristCtrl)[0]
+                cmds.delete(constraint)
+                
+                fkWristCtrlGrp = cmds.group(empty = True, name = "fk_wrist_" + side + "_anim_grp")
+                constraint = cmds.parentConstraint(fkWristJoint, fkWristCtrlGrp)[0]
+                cmds.delete(constraint)
+                
+                cmds.parent(fkWristCtrl, fkWristCtrlGrp)
+                cmds.makeIdentity(fkWristCtrl, t = 1, r = 1, s = 1, apply = True)
+                cmds.parent(fkWristCtrlGrp, fkElbowCtrl)
+                
+                #constrain wrist joint to ctrl
+                cmds.parentConstraint(fkWristCtrl, fkWristJoint)
+                
+                
+                #point constrain the fk arm grp to the fk upper arm joint
+                cmds.pointConstraint(fkArmJoint, "fk_arm_" + side + "_grp")
+                
+
+                #clean up FK rig in scene
+                cmds.parent(invisFkArmCtrlGrp, autoClavWorld)
+                
+                
+                #find children under autoClavWorld
+                children = cmds.listRelatives(autoClavWorld, children = True)
+                dntGrp = cmds.group(empty = True, name = "auto_clav_sys_grp_" + side)
+                cmds.parent(dntGrp, autoClavWorld)
+                
+                for child in children:
+                    cmds.parent(child, dntGrp)
+                    
+                cmds.parent(shoulderCtrlGrp, autoClavWorld)
+                
+                
+                #group up the groups!
+                jointsGrp = cmds.group(empty = True, name = "joints_" + side + "_grp")
+                cmds.parent(clavJoint, jointsGrp)
+                
+                masterGrp = cmds.group(empty = True, name = "arm_rig_master_grp_" + side)
+		constraint = cmds.pointConstraint(fkArmJoint, masterGrp)[0]
+		cmds.delete(constraint)
+		
+                cmds.parent([ikNodes, jointsGrp, autoClavWorld], masterGrp)
+                
+                cmds.setAttr(dntGrp + ".v", 0)
+                cmds.setAttr(ikNodes + ".v", 0)
+                
+                
+                #add fk orientation options(normal, body, world)
+                fkOrient = cmds.spaceLocator(name = "fk_orient_master_loc_" + side)[0]
+                shape = cmds.listRelatives(fkOrient, shapes = True)[0]
+                cmds.setAttr(shape + ".v", 0)
+                
+                constraint = cmds.parentConstraint(autoClavWorld, fkOrient)[0]
+                cmds.delete(constraint)
+                
+                fkNormalOrient = cmds.duplicate(fkOrient, po = True, name = "fk_orient_normal_loc_" + side)[0]
+                fkBodyOrient = cmds.duplicate(fkOrient, po = True, name = "fk_orient_body_loc_" + side)[0]
+                fkWorldOrient = cmds.duplicate(fkOrient, po = True, name = "fk_orient_world_loc_" + side)[0]
+                
+                fkOrientConstraint = cmds.orientConstraint([fkNormalOrient, fkBodyOrient, fkWorldOrient], fkOrient)[0]
+                
+                cmds.parent(fkBodyOrient, "body_anim")
+                cmds.parent(fkNormalOrient, shoulderCtrl)
+                cmds.parent(fkOrient, masterGrp)
+		
+                #parent FK arm ctrl grp to shoulder ctrl
+                cmds.parent("fk_arm_" + side + "_grp", fkOrient)
+		
+		
+                #put mode into default fk operation
+                cmds.setAttr(fkOrientConstraint + "." + fkBodyOrient + "W1", 0)
+                cmds.setAttr(fkOrientConstraint + "." + fkWorldOrient + "W2", 0)
+                
+                
+                #get the number of spine bones and constrain the master grp to the last spine joint
+                numSpineBones = self.getSpineJoints()
+                cmds.parentConstraint("driver_spine_0" + str(len(numSpineBones)), masterGrp, mo = True)
+                
+                
+                
+                
+                
+                #setup autoShoulder attr
+                cmds.select(shoulderCtrl)
+                cmds.addAttr(longName='autoShoulders', defaultValue=0, minValue=0, maxValue=1, keyable = True)
+                
+                cmds.setAttr(shoulderCtrl + ".autoShoulders", 0)
+                cmds.setAttr(autoClavSwitchConstraint + "." + autoClavEndJoint + "W0", 0)
+                cmds.setAttr(autoClavSwitchConstraint + "." + autoClavWorld + "W1", 1)
+                cmds.setDrivenKeyframe([autoClavSwitchConstraint + "." + autoClavEndJoint + "W0", autoClavSwitchConstraint + "." + autoClavWorld + "W1"], cd = shoulderCtrl + ".autoShoulders", itt = "linear", ott = "linear")
+                
+                
+                cmds.setAttr(shoulderCtrl + ".autoShoulders", 1)
+                cmds.setAttr(autoClavSwitchConstraint + "." + autoClavEndJoint + "W0", 1)
+                cmds.setAttr(autoClavSwitchConstraint + "." + autoClavWorld + "W1", 0)
+                cmds.setDrivenKeyframe([autoClavSwitchConstraint + "." + autoClavEndJoint + "W0", autoClavSwitchConstraint + "." + autoClavWorld + "W1"], cd = shoulderCtrl + ".autoShoulders", itt = "linear", ott = "linear")
+                
+                cmds.setAttr(shoulderCtrl + ".autoShoulders", 0)
+                
+                
+                
+                
+                
+                #setup FK arm orient attr
+                cmds.select("Rig_Settings")
+                cmds.addAttr(longName= side + "FkArmOrient", at = 'enum', en = "fk:body:world:", keyable = True)
+                
+                
+                cmds.setAttr("Rig_Settings." + side + "FkArmOrient", 0)
+                cmds.setAttr(fkOrientConstraint + "." + fkNormalOrient + "W0", 1)
+                cmds.setAttr(fkOrientConstraint + "." + fkBodyOrient + "W1", 0)
+                cmds.setAttr(fkOrientConstraint + "." + fkWorldOrient + "W2", 0)
+                cmds.setDrivenKeyframe([fkOrientConstraint + "." + fkNormalOrient + "W0", fkOrientConstraint + "." + fkBodyOrient + "W1", fkOrientConstraint + "." + fkWorldOrient + "W2"], cd = "Rig_Settings." + side + "FkArmOrient", itt = "linear", ott = "linear")
+                
+                
+                cmds.setAttr("Rig_Settings." + side + "FkArmOrient", 1)
+                cmds.setAttr(fkOrientConstraint + "." + fkNormalOrient + "W0", 0)
+                cmds.setAttr(fkOrientConstraint + "." + fkBodyOrient + "W1", 1)
+                cmds.setAttr(fkOrientConstraint + "." + fkWorldOrient + "W2", 0)
+                cmds.setDrivenKeyframe([fkOrientConstraint + "." + fkNormalOrient + "W0", fkOrientConstraint + "." + fkBodyOrient + "W1", fkOrientConstraint + "." + fkWorldOrient + "W2"], cd = "Rig_Settings." + side + "FkArmOrient", itt = "linear", ott = "linear")
+                
+                cmds.setAttr("Rig_Settings." + side + "FkArmOrient", 2)
+                cmds.setAttr(fkOrientConstraint + "." + fkNormalOrient + "W0", 0)
+                cmds.setAttr(fkOrientConstraint + "." + fkBodyOrient + "W1", 0)
+                cmds.setAttr(fkOrientConstraint + "." + fkWorldOrient + "W2", 1)
+                cmds.setDrivenKeyframe([fkOrientConstraint + "." + fkNormalOrient + "W0", fkOrientConstraint + "." + fkBodyOrient + "W1", fkOrientConstraint + "." + fkWorldOrient + "W2"], cd = "Rig_Settings." + side + "FkArmOrient", itt = "linear", ott = "linear")
+                
+                cmds.setAttr("Rig_Settings." + side + "FkArmOrient", 0)
+                
+                
+                #setup limits on auto clavicle
+		
+		
+                cmds.setAttr(shoulderCtrl + ".autoShoulders", 1)
+                #cmds.setAttr("fk_arm_" + side + "_anim.ry", -60)
+                #grrr. need to force update since maya is not getting the info correctly
+                cmds.select("fk_arm_" + side + "_anim")
+                cmds.setToolTo( 'moveSuperContext' )
+                cmds.refresh(force = True)
+                cmds.select(clear = True)
+                #limitInfo = cmds.xform(autoClavSwitchLoc, q = True, t = True)
+                #cmds.setAttr("fk_arm_" + side + "_anim.ry", 0)
+                cmds.setAttr(shoulderCtrl + ".autoShoulders", 0)
+                cmds.select("fk_arm_" + side + "_anim")
+                cmds.setToolTo( 'moveSuperContext' )
+                cmds.refresh(force = True)
+                cmds.select(clear = True)
+                
+                
+
+                
+
+                #lock attrs that should not be animated and colorize controls
+                for control in [shoulderCtrl, "fk_arm_" + side + "_anim", fkElbowCtrl, fkWristCtrl]:
+                    
+                    if control == shoulderCtrl:
+                        
+                        for attr in [".rx", ".ry", ".rz", ".sx", ".sy", ".sz", ".v"]:
+                            cmds.setAttr(control + attr, lock = True, keyable = False)
+                    
+                    else:
+                        for attr in [".sx", ".sy", ".sz", ".v"]:
+                            cmds.setAttr(control + attr, lock = True, keyable = False)
+
+                            
+                    if side == "l":
+                        color = 6
+                        
+                    else:
+                        color = 13
+
+                    cmds.setAttr(control + ".overrideEnabled", 1)
+                    cmds.setAttr(control + ".overrideColor", color)
+		    
+		#parent fkWorldOrient to master anim
+		cmds.parent(fkWorldOrient, "master_anim")
+                    
+                    
+		
+    # # # # # # # # # # # # # # # # # # # # # # # # # # # # # # # # # # # # # # # # # # # # # # # # # # # # # # # # # # # # # # # 
+    # # # # # # # # # # # # # # # # # # # # # # # # # # # # # # # # # # # # # # # # # # # # # # # # # # # # # # # # # # # # # # # 
+    # # # # # # # # # # # # # # # # # # # # # # # # # # # # # # # # # # # # # # # # # # # # # # # # # # # # # # # # # # # # # # # 
+    # # # # # # # # # # # # # # # # # # # # # # # # # # # # # # # # # # # # # # # # # # # # # # # # # # # # # # # # # # # # # # # 
+    # # # # # # # # # # # # # # # # # # # # # # # # # # # # # # # # # # # # # # # # # # # # # # # # # # # # # # # # # # # # # # # 
+    def buildIkArms(self):
+        
+        #duplicate the fk arm joints and create our ik arm chain
+        for side in ["l", "r"]:
+            if cmds.objExists("fk_upperarm_" + side):
+                ikUpArmJoint = cmds.duplicate("fk_upperarm_" + side, po = True, name = "ik_upperarm_" + side)[0]
+                ikLowArmJoint = cmds.duplicate("fk_lowerarm_" + side, po = True, name = "ik_lowerarm_" + side)[0]
+                ikWristJoint = cmds.duplicate("fk_hand_" + side, po = True, name = "ik_hand_" + side)[0]
+                ikWristEndJoint = cmds.duplicate("fk_hand_" + side, po = True, name = "ik_hand_end_" + side)[0]
+                
+                cmds.parent([ikWristEndJoint], ikWristJoint)
+                cmds.parent(ikWristJoint, ikLowArmJoint)
+                cmds.parent(ikLowArmJoint, ikUpArmJoint)
+                
+                
+                #create fk matching joints
+		fkMatchUpArm = cmds.duplicate(ikUpArmJoint, po = True, name = "ik_upperarm_fk_matcher_" + side)[0]
+		fkMatchLowArm = cmds.duplicate(ikLowArmJoint, po = True, name = "ik_lowerarm_fk_matcher_" + side)[0]
+		fkMatchWrist = cmds.duplicate(ikWristJoint, po = True, name = "ik_wrist_fk_matcher_" + side)[0]
+		
+                cmds.parent(fkMatchWrist, fkMatchLowArm)
+                cmds.parent(fkMatchLowArm, fkMatchUpArm)
+		
+		cmds.parentConstraint(ikUpArmJoint, fkMatchUpArm, mo = True)
+		cmds.parentConstraint(ikLowArmJoint, fkMatchLowArm, mo = True)
+		cmds.parentConstraint(ikWristJoint, fkMatchWrist, mo = True)
+		
+                #move wrist end joint out a bit
+                scaleFactor = self.getScaleFactor()
+                
+                if side == "l":
+                    cmds.setAttr(ikWristEndJoint + ".tx", 15 * scaleFactor)
+                    
+                else:
+                    cmds.setAttr(ikWristEndJoint + ".tx", -15 * scaleFactor)
+                
+                
+                cmds.makeIdentity(ikUpArmJoint, t = 0, r = 1, s = 0, apply = True)
+		
+		#set rotate order on ikUpArm
+		cmds.setAttr(ikUpArmJoint + ".rotateOrder", 3)
+		
+		
+                #set preferred angle on arm
+		cmds.setAttr(ikLowArmJoint + ".preferredAngleZ", -90)
+                
+		
+                #create ik control
+                ikCtrl = self.createControl("circle", 15, "ik_wrist_" + side + "_anim")
+                cmds.setAttr(ikCtrl + ".ry", -90)
+                cmds.makeIdentity(ikCtrl, r = 1, apply =True)
+                
+		constraint = cmds.pointConstraint(ikWristJoint, ikCtrl)[0]
+                cmds.delete(constraint)
+                
+                ikCtrlGrp = cmds.group(empty = True, name = "ik_wrist_" + side + "_anim_grp")
+		constraint = cmds.pointConstraint(ikWristJoint, ikCtrlGrp)[0]
+		cmds.delete(constraint)
+
+                
+                ikCtrlSpaceSwitchFollow = cmds.duplicate(ikCtrlGrp, po = True, n = "ik_wrist_" + side + "_anim_space_switcher_follow")[0]
+                ikCtrlSpaceSwitch = cmds.duplicate(ikCtrlGrp, po = True, n = "ik_wrist_" + side + "_anim_space_switcher")[0]
+                
+		cmds.parent(ikCtrlSpaceSwitch, ikCtrlSpaceSwitchFollow)
+                cmds.parent(ikCtrlGrp, ikCtrlSpaceSwitch)
+	    
+		cmds.parent(ikCtrl, ikCtrlGrp)
+		cmds.makeIdentity(ikCtrlGrp, t = 1, r = 1, s = 1, apply = True)
+		
+		
+                #create RP IK on arm and SC ik from wrist to wrist end
+                rpIkHandle = cmds.ikHandle(name = "arm_ikHandle_" + side, solver = "ikRPsolver", sj = ikUpArmJoint, ee = ikWristJoint)[0]
+                scIkHandle = cmds.ikHandle(name = "hand_ikHandle_" + side, solver = "ikSCsolver", sj = ikWristJoint, ee = ikWristEndJoint)[0]
+                
+                cmds.parent(scIkHandle, rpIkHandle)
+                cmds.setAttr(rpIkHandle + ".v", 0)
+                cmds.setAttr(scIkHandle + ".v", 0)
+                
+                
+                #parent IK to ik control
+                cmds.parent(rpIkHandle, ikCtrl)
+                
+                
+                #create a pole vector control
+                ikPvCtrl = self.createControl("sphere", 6, "ik_elbow_" + side + "_anim")
+                constraint = cmds.pointConstraint(ikLowArmJoint, ikPvCtrl)[0]
+                cmds.delete(constraint)
+                cmds.makeIdentity(ikPvCtrl, t = 1, r = 1, s = 1, apply = True)
+                
+                #move out a bit
+		cmds.setAttr(ikPvCtrl + ".ty", (30 * scaleFactor))
+		cmds.makeIdentity(ikPvCtrl, t = 1, r = 1, s = 1, apply = True)
+		    
+
+		    
+                
+                
+                #create group for control
+                ikPvCtrlGrp = cmds.group(empty = True, name = "ik_elbow_" + side + "_anim_grp")
+                constraint = cmds.parentConstraint(ikPvCtrl, ikPvCtrlGrp)[0]
+                cmds.delete(constraint)
+                
+                ikPvSpaceSwitchFollow = cmds.duplicate(ikPvCtrlGrp, po = True, name = "ik_elbow_" + side + "_anim_space_switcher_follow")[0]
+                ikPvSpaceSwitch = cmds.duplicate(ikPvCtrlGrp, po = True, name = "ik_elbow_" + side + "_anim_space_switcher")[0]
+		
+		cmds.parent(ikPvSpaceSwitch, ikPvSpaceSwitchFollow)
+                cmds.parent(ikPvCtrlGrp, ikPvSpaceSwitch)
+                cmds.parent(ikPvCtrl, ikPvCtrlGrp)
+                cmds.parent(ikPvSpaceSwitchFollow, "offset_anim")
+                cmds.makeIdentity(ikPvCtrl, t = 1, r = 1, s = 1, apply = True)
+                
+                #setup pole vector constraint
+                cmds.poleVectorConstraint(ikPvCtrl, rpIkHandle)
+                
+                
+                #create IK for invisible arm
+                invisRpIkHandle = cmds.ikHandle(name = "invis_arm_ikHandle_" + side, solver = "ikRPsolver", sj = "invis_upperarm_" + side, ee = "invis_hand_" + side)[0]
+                cmds.parent(invisRpIkHandle, ikCtrl)
+                cmds.poleVectorConstraint(ikPvCtrl, invisRpIkHandle)
+                cmds.setAttr(invisRpIkHandle + ".v", 0)
+                
+                
+                
+                #constrain driver joints to both fk and ik chains
+                cmds.parentConstraint("rig_clavicle_" + side, "driver_clavicle_" + side, mo = True)
+                upArmConstPoint = cmds.pointConstraint(["fk_arm_" + side + "_anim", "ik_upperarm_" + side], "driver_upperarm_" + side)[0]
+		upArmConstOrient = cmds.orientConstraint(["fk_upperarm_" + side, "ik_upperarm_" + side], "driver_upperarm_" + side)[0]
+                lowArmConst = cmds.parentConstraint(["fk_lowerarm_" + side, "ik_lowerarm_" + side], "driver_lowerarm_" + side)[0]
+                handConst = cmds.parentConstraint(["fk_hand_" + side, "ik_hand_" + side], "driver_hand_" + side, mo = True)[0]
+		
+		
+		#create blend nodes for the scale
+		scaleBlendColors_UpArm = cmds.shadingNode("blendColors", asUtility = True, name = side + "_up_arm_scale_blend")
+		cmds.connectAttr(ikUpArmJoint + ".scale", scaleBlendColors_UpArm + ".color1")
+		cmds.connectAttr("fk_arm_" + side + "_anim.scale", scaleBlendColors_UpArm + ".color2")
+		cmds.connectAttr(scaleBlendColors_UpArm + ".output", "driver_upperarm_" + side + ".scale")
+		
+		
+		scaleBlendColors_LoArm = cmds.shadingNode("blendColors", asUtility = True, name = side + "_lo_arm_scale_blend")
+		cmds.connectAttr(ikLowArmJoint + ".scale", scaleBlendColors_LoArm + ".color1")
+		cmds.connectAttr("fk_elbow_" + side + "_anim.scale", scaleBlendColors_LoArm + ".color2")
+		cmds.connectAttr(scaleBlendColors_LoArm + ".output", "driver_lowerarm_" + side + ".scale")
+		
+		scaleBlendColors_Wrist = cmds.shadingNode("blendColors", asUtility = True, name = side + "_wrist_scale_blend")
+		cmds.connectAttr(ikWristJoint + ".scale", scaleBlendColors_Wrist + ".color1")
+		cmds.connectAttr("fk_wrist_" + side + "_anim.scale", scaleBlendColors_Wrist + ".color2")
+		cmds.connectAttr(scaleBlendColors_Wrist + ".output", "driver_hand_" + side + ".scale")
+		
+		
+		
+
+		#set limits
+		cmds.select("driver_upperarm_" + side)
+		cmds.transformLimits(sy = (.05, 1.25), sz = (.05, 1.25), esy = [False, True], esz = [False, True])
+		cmds.select("driver_lowerarm_" + side)
+		cmds.transformLimits(sy = (.05, 1.25), sz = (.05, 1.25), esy = [False, True], esz = [False, True])		
+                
+                
+                
+                #create the IK/FK switch
+                cmds.select("Rig_Settings")
+                cmds.addAttr(longName= side + "ArmMode", at = 'enum', en = "fk:ik:", keyable = True)
+                
+                cmds.setAttr("Rig_Settings." + side + "ArmMode", 0)
+                
+                cmds.setAttr(upArmConstPoint + "." + "fk_arm_" + side + "_anim" + "W0", 1)
+                cmds.setAttr(upArmConstPoint + "." + "ik_upperarm_" + side + "W1", 0)
+                cmds.setAttr(upArmConstOrient + "." + "fk_upperarm_" + side + "W0", 1)
+                cmds.setAttr(upArmConstOrient + "." + "ik_upperarm_" + side + "W1", 0)
+		cmds.setAttr(scaleBlendColors_UpArm + "." + "blender", 0)	
+                
+                cmds.setAttr(lowArmConst + "." + "fk_lowerarm_" + side + "W0", 1)
+                cmds.setAttr(lowArmConst + "." + "ik_lowerarm_" + side + "W1", 0)
+                cmds.setAttr(scaleBlendColors_LoArm + "." + "blender", 0)
+		
+                cmds.setAttr(handConst + "." + "fk_hand_" + side + "W0", 1)
+                cmds.setAttr(handConst + "." + "ik_hand_" + side + "W1", 0)
+                cmds.setAttr(scaleBlendColors_Wrist + "." + "blender", 0)
+		
+                cmds.setAttr(invisRpIkHandle + ".ikBlend", 0)
+                cmds.setAttr("fk_arm_" + side + "_grp.v", 1)
+                cmds.setAttr("ik_wrist_" + side + "_anim_space_switcher.v", 0)
+                cmds.setAttr("ik_elbow_" + side + "_anim_space_switcher.v", 0)
+                
+
+                cmds.setDrivenKeyframe([scaleBlendColors_UpArm + "." + "blender", scaleBlendColors_LoArm + "." + "blender", scaleBlendColors_Wrist + "." + "blender", upArmConstPoint + "." + "fk_arm_" + side + "_anim" + "W0", upArmConstPoint + "." + "ik_upperarm_" + side + "W1", ], cd = "Rig_Settings." + side + "ArmMode", itt = "linear", ott = "linear")
+		cmds.setDrivenKeyframe([upArmConstOrient + "." + "fk_upperarm_" + side + "W0", upArmConstOrient + "." + "ik_upperarm_" + side + "W1", ], cd = "Rig_Settings." + side + "ArmMode", itt = "linear", ott = "linear")
+                cmds.setDrivenKeyframe([lowArmConst + "." + "fk_lowerarm_" + side + "W0", lowArmConst + "." + "ik_lowerarm_" + side + "W1", ], cd = "Rig_Settings." + side + "ArmMode", itt = "linear", ott = "linear")
+                cmds.setDrivenKeyframe([handConst + "." + "fk_hand_" + side + "W0", handConst + "." + "ik_hand_" + side + "W1", ], cd = "Rig_Settings." + side + "ArmMode", itt = "linear", ott = "linear")
+                cmds.setDrivenKeyframe([invisRpIkHandle + ".ikBlend", "fk_arm_" + side + "_grp.v"], cd = "Rig_Settings." + side + "ArmMode", itt = "linear", ott = "linear")
+                cmds.setDrivenKeyframe(["ik_wrist_" + side + "_anim_space_switcher.v", "ik_elbow_" + side + "_anim_space_switcher.v"], cd = "Rig_Settings." + side + "ArmMode", itt = "linear", ott = "linear")
+
+                cmds.setAttr("Rig_Settings." + side + "ArmMode", 1)
+                
+                cmds.setAttr(upArmConstPoint + "." + "fk_arm_" + side + "_anim" "W0", 0)
+                cmds.setAttr(upArmConstPoint + "." + "ik_upperarm_" + side + "W1", 1)
+                cmds.setAttr(upArmConstOrient + "." + "fk_upperarm_" + side + "W0", 0)
+                cmds.setAttr(upArmConstOrient + "." + "ik_upperarm_" + side + "W1", 1)		
+                cmds.setAttr(scaleBlendColors_UpArm + "." + "blender", 1)	
+		
+                cmds.setAttr(lowArmConst + "." + "fk_lowerarm_" + side + "W0", 0)
+                cmds.setAttr(lowArmConst + "." + "ik_lowerarm_" + side + "W1", 1)
+                cmds.setAttr(scaleBlendColors_LoArm + "." + "blender", 1)
+		
+                cmds.setAttr(handConst + "." + "fk_hand_" + side + "W0", 0)
+                cmds.setAttr(handConst + "." + "ik_hand_" + side + "W1", 1)
+                cmds.setAttr(scaleBlendColors_Wrist + "." + "blender", 1)
+		
+                cmds.setAttr(invisRpIkHandle + ".ikBlend", 1)
+                cmds.setAttr("fk_arm_" + side + "_grp.v", 0)
+                cmds.setAttr("ik_wrist_" + side + "_anim_space_switcher.v", 1)
+                cmds.setAttr("ik_elbow_" + side + "_anim_space_switcher.v", 1)
+
+                cmds.setDrivenKeyframe([scaleBlendColors_UpArm + "." + "blender", scaleBlendColors_LoArm + "." + "blender", scaleBlendColors_Wrist + "." + "blender", upArmConstPoint + "." + "fk_arm_" + side + "_anim" + "W0", upArmConstPoint + "." + "ik_upperarm_" + side + "W1", ], cd = "Rig_Settings." + side + "ArmMode", itt = "linear", ott = "linear")
+		cmds.setDrivenKeyframe([upArmConstOrient + "." + "fk_upperarm_" + side + "W0", upArmConstOrient + "." + "ik_upperarm_" + side + "W1", ], cd = "Rig_Settings." + side + "ArmMode", itt = "linear", ott = "linear")                
+                cmds.setDrivenKeyframe([lowArmConst + "." + "fk_lowerarm_" + side + "W0", lowArmConst + "." + "ik_lowerarm_" + side + "W1", ], cd = "Rig_Settings." + side + "ArmMode", itt = "linear", ott = "linear")
+                cmds.setDrivenKeyframe([handConst + "." + "fk_hand_" + side + "W0", handConst + "." + "ik_hand_" + side + "W1", ], cd = "Rig_Settings." + side + "ArmMode", itt = "linear", ott = "linear")
+                cmds.setDrivenKeyframe([invisRpIkHandle + ".ikBlend", "fk_arm_" + side + "_grp.v"], cd = "Rig_Settings." + side + "ArmMode", itt = "linear", ott = "linear")
+                cmds.setDrivenKeyframe(["ik_wrist_" + side + "_anim_space_switcher.v", "ik_elbow_" + side + "_anim_space_switcher.v"], cd = "Rig_Settings." + side + "ArmMode", itt = "linear", ott = "linear")
+            
+                cmds.setAttr("Rig_Settings." + side + "ArmMode", 0)
+                
+                
+                
+                #setup stretch on IK
+                cmds.select(ikCtrl)
+                cmds.addAttr(longName=("stretch"), at = 'double',min = 0, max = 1, dv = 0, keyable = True)
+		cmds.addAttr(longName=("squash"), at = 'double',min = 0, max = 1, dv = 0, keyable = True)
+                stretchMultNode = cmds.shadingNode("multiplyDivide", asUtility = True, name = "ikHand_stretchToggleMultNode_" + side)
+    
+                #need to get the total length of the arm chain
+		totalDist = abs(cmds.getAttr(ikLowArmJoint + ".tx" ) + cmds.getAttr(ikWristJoint + ".tx"))
+		
+		    
+		#create a distanceBetween node
+		distBetween = cmds.shadingNode("distanceBetween", asUtility = True, name = side + "_ik_arm_distBetween")
+		
+		#get world positions of upper arm and ik
+		baseGrp = cmds.group(empty = True, name = "ik_arm_base_grp_" + side)
+		endGrp = cmds.group(empty = True, name = "ik_arm_end_grp_" + side)
+		cmds.pointConstraint(ikUpArmJoint, baseGrp)
+		cmds.pointConstraint(ikCtrl, endGrp)
+		
+		#hook in group translates into distanceBetween node inputs
+		cmds.connectAttr(baseGrp + ".translate", distBetween + ".point1")
+		cmds.connectAttr(endGrp + ".translate", distBetween + ".point2")
+		
+		#create a condition node that will compare original length to current length
+		#if second term is greater than, or equal to the first term, the chain needs to stretch
+		ikArmCondition = cmds.shadingNode("condition", asUtility = True, name = side + "_ik_arm_stretch_condition")
+		cmds.setAttr(ikArmCondition + ".operation", 3)
+		cmds.connectAttr(distBetween + ".distance", ikArmCondition + ".secondTerm")
+		cmds.setAttr(ikArmCondition + ".firstTerm", totalDist)
+		
+		#hook up the condition node's return colors
+		cmds.setAttr(ikArmCondition + ".colorIfTrueR", totalDist)
+		cmds.connectAttr(distBetween + ".distance", ikArmCondition + ".colorIfFalseR")
+		
+
+                #create the mult/divide node(set to divide) that will take the original creation length as a static value in input2x, and the connected length into 1x.
+                armDistMultNode = cmds.shadingNode("multiplyDivide", asUtility = True, name = "arm_dist_multNode_" + side)
+                cmds.setAttr(armDistMultNode + ".operation", 2) #divide
+                cmds.setAttr(armDistMultNode + ".input2X", totalDist)
+                cmds.connectAttr(ikArmCondition + ".outColorR", armDistMultNode + ".input1X")
+                
+                #create a stretch toggle mult node that multiplies the stretch factor by the bool of the stretch attr. (0 or 1), this way our condition reads
+                #if this result is greater than the original length(impossible if stretch bool is off, since result will be 0), than take this result and plug it
+                #into the scale of our IK arm joints
+		stretchToggleCondition = cmds.shadingNode("condition", asUtility = True, name = "arm_stretch_toggle_condition_" + side)
+		cmds.setAttr(stretchToggleCondition + ".operation", 0)
+		cmds.connectAttr(ikCtrl + ".stretch", stretchToggleCondition + ".firstTerm")
+		cmds.setAttr(stretchToggleCondition + ".secondTerm", 1)
+		cmds.connectAttr(armDistMultNode + ".outputX", stretchToggleCondition + ".colorIfTrueR")
+		cmds.setAttr(stretchToggleCondition + ".colorIfFalseR", 1)
+		
+	    
+		#set up the squash nodes
+		squashMultNode = cmds.shadingNode("multiplyDivide", asUtility = True, name = side + "_ik_arm_squash_mult")
+		cmds.setAttr(squashMultNode + ".operation", 2)
+		cmds.setAttr(squashMultNode + ".input1X", totalDist)
+		cmds.connectAttr(ikArmCondition + ".outColorR", squashMultNode + ".input2X")
+
+
+		#create a stretch toggle mult node that multiplies the stretch factor by the bool of the stretch attr. (0 or 1), this way our condition reads
+		#if this result is greater than the original length(impossible if stretch bool is off, since result will be 0), than take this result and plug it
+		#into the scale of our IK arm joints
+		squashToggleCondition = cmds.shadingNode("condition", asUtility = True, name = "arm_squash_toggle_condition_" + side)
+		cmds.setAttr(squashToggleCondition + ".operation", 0)
+		cmds.connectAttr(ikCtrl + ".squash", squashToggleCondition + ".firstTerm")
+		cmds.setAttr(squashToggleCondition + ".secondTerm", 1)
+		cmds.connectAttr(squashMultNode + ".outputX", squashToggleCondition + ".colorIfTrueR")
+		cmds.setAttr(squashToggleCondition + ".colorIfFalseR", 1)
+
+
+		
+		#connect to arm scale
+                cmds.connectAttr(stretchToggleCondition + ".outColorR", ikUpArmJoint + ".sx")
+		cmds.connectAttr(stretchToggleCondition + ".outColorR", ikLowArmJoint + ".sx")
+                
+                cmds.connectAttr(squashToggleCondition + ".outColorR", ikLowArmJoint + ".sy")
+		cmds.connectAttr(squashToggleCondition + ".outColorR", ikLowArmJoint + ".sz")
+		
+		cmds.connectAttr(squashToggleCondition + ".outColorR", ikUpArmJoint + ".sy")
+		cmds.connectAttr(squashToggleCondition + ".outColorR", ikUpArmJoint + ".sz")
+
+
+		#add base and end groups to arm grp
+                cmds.parent([baseGrp, endGrp], "arm_rig_master_grp_" + side)
+		
+		
+		
+                # # # # # # # # # # # # # # # # # # # # # # # # # # # # # # # # # # # # # # # # # # # #
+                
+                #setup roll bones if obj Exists
+                if cmds.objExists("driver_upperarm_twist_01_" + side):
+                    self.buildArmRoll(side)
+                    
+                if cmds.objExists("driver_lowerarm_twist_01_" + side):
+                    self.buildForearmTwist(side)
+                    
+                #colorize controls, cleanup attrs, and cleanup hierarchy
+                for attr in [".sx", ".sy", ".sz", ".v"]:
+                    cmds.setAttr(ikCtrl + attr, lock = True, keyable = False)
+                    
+                for attr in [".sx", ".sy", ".sz", ".rx", ".ry", ".rz", ".v"]:
+                    cmds.setAttr(ikPvCtrl + attr, lock = True, keyable = False)
+                    
+                    
+                for control in [ikCtrl, ikPvCtrl]:
+                    if side == "l":
+                        color = 6
+                        
+                    else:
+                        color = 13
+    
+                    cmds.setAttr(control + ".overrideEnabled", 1)
+                    cmds.setAttr(control + ".overrideColor", color)
+
+
+                
+    # # # # # # # # # # # # # # # # # # # # # # # # # # # # # # # # # # # # # # # # # # # # # # # # # # # # # # # # # # # # # # # 
+    # # # # # # # # # # # # # # # # # # # # # # # # # # # # # # # # # # # # # # # # # # # # # # # # # # # # # # # # # # # # # # # 
+    # # # # # # # # # # # # # # # # # # # # # # # # # # # # # # # # # # # # # # # # # # # # # # # # # # # # # # # # # # # # # # # 
+    # # # # # # # # # # # # # # # # # # # # # # # # # # # # # # # # # # # # # # # # # # # # # # # # # # # # # # # # # # # # # # # 
+    # # # # # # # # # # # # # # # # # # # # # # # # # # # # # # # # # # # # # # # # # # # # # # # # # # # # # # # # # # # # # # # 
+    def buildNeckAndHead(self):
+        
+        numNeckBones = cmds.getAttr("Skeleton_Settings.numNeckBones")
+        
+        if numNeckBones == 1:            
+            #create the FK control for the neck
+            neckControl = self.createControl("circle", 25, "neck_01_fk_anim")
+            constraint = cmds.parentConstraint("driver_neck_01", neckControl)[0]
+            cmds.delete(constraint)
+            
+            neckControlGrp = cmds.group(empty = True, name = "neck_01_fk_anim_grp")
+            constraint = cmds.parentConstraint("driver_neck_01", neckControlGrp)[0]
+            cmds.delete(constraint)
+            
+            cmds.parent(neckControl, neckControlGrp)
+            cmds.setAttr(neckControl + ".ry", -90)
+            cmds.makeIdentity(neckControl, t = 1, r = 1, s = 1, apply = True)
+            
+            
+            #create the FK control for the head
+            headControl = self.createControl("circle", 30, "head_fk_anim")
+            constraint = cmds.parentConstraint("driver_head", headControl)[0]
+            cmds.delete(constraint)
+            
+            headControlGrp = cmds.group(empty = True, name = "head_fk_anim_grp")
+            constraint = cmds.parentConstraint("driver_head", headControlGrp)[0]
+            cmds.delete(constraint)
+            
+            cmds.parent(headControl, headControlGrp)
+            cmds.setAttr(headControl + ".ry", 90)
+	    cmds.setAttr(headControl + ".rz", -35)
+            cmds.makeIdentity(headControl, t = 1, r = 1, s = 1, apply = True)
+            
+            
+            #setup head orientation
+            orientNodes = self.setupHeadOrientation(neckControl, headControl)
+	    neckOrientNodes = self.setupNeckOrientation(neckControl)
+            
+            #hook into spine
+            cmds.parent(headControlGrp, orientNodes[0])
+            cmds.parent(orientNodes[0], neckControl)
+            numSpineBones = self.getSpineJoints()
+	    cmds.parent(neckOrientNodes[0], neckControlGrp)
+	    cmds.parent(neckControl, neckOrientNodes[0])
+	    cmds.parentConstraint("driver_spine_0" + str(len(numSpineBones)), neckControlGrp, mo = True)
+            
+            #constrain driver joints to controls
+            cmds.parentConstraint(neckControl, "driver_neck_01", mo = True)
+            cmds.parentConstraint(headControl, "driver_head", mo = True)
+	    cmds.connectAttr(neckControl + ".scale", "driver_neck_01.scale")
+	    cmds.connectAttr(headControl + ".scale", "driver_head.scale")
+            
+            
+            #lock attrs, color controls, and clean up hierarchy
+            for control in [neckControl, headControl]:
+                for attr in [".v"]:
+                    cmds.setAttr(control + attr, lock = True, keyable = False)
+                    
+            cmds.setAttr(neckControl + ".overrideEnabled", 1)
+            cmds.setAttr(neckControl + ".overrideColor", 18)
+            cmds.setAttr(headControl + ".overrideEnabled", 1)
+            cmds.setAttr(headControl + ".overrideColor", 17)
+            
+            masterGrp = cmds.group(empty = True, name = "head_sys_grp")
+	    cmds.parent(orientNodes[4], masterGrp)
+	    cmds.parent(neckOrientNodes[3], masterGrp)
+	    #cmds.parent(neckOrientNodes[0], masterGrp)
+            
+            
+        else:
+            neckControlMid = ""
+            if numNeckBones == 2:
+                #create the FK controls for the neck
+                neckControl = self.createControl("circle", 25, "neck_01_fk_anim")
+                constraint = cmds.parentConstraint("driver_neck_01", neckControl)[0]
+                cmds.delete(constraint)
+                
+                neckControlGrp = cmds.group(empty = True, name = "neck_01_fk_anim_grp")
+                constraint = cmds.parentConstraint("driver_neck_01", neckControlGrp)[0]
+                cmds.delete(constraint)
+                
+                cmds.parent(neckControl, neckControlGrp)
+                cmds.setAttr(neckControl + ".ry", -90)
+                cmds.makeIdentity(neckControl, t = 1, r = 1, s = 1, apply = True)
+                
+                
+                neckControlEnd = self.createControl("circle", 25, "neck_02_fk_anim")
+                constraint = cmds.parentConstraint("driver_neck_02", neckControlEnd)[0]
+                cmds.delete(constraint)
+                
+                neckControlEndGrp = cmds.group(empty = True, name = "neck_02_fk_anim_grp")
+                constraint = cmds.parentConstraint("driver_neck_02", neckControlEndGrp)[0]
+                cmds.delete(constraint)
+                
+                cmds.parent(neckControlEnd, neckControlEndGrp)
+                cmds.setAttr(neckControlEnd + ".ry", -90)
+                cmds.makeIdentity(neckControlEnd, t = 1, r = 1, s = 1, apply = True)
+                
+                #setup neck hiearchy
+                cmds.parent(neckControlEndGrp, neckControl)
+                
+                
+            if numNeckBones == 3:
+                
+                #create the FK controls for the neck
+                neckControl = self.createControl("circle", 25, "neck_01_fk_anim")
+                constraint = cmds.parentConstraint("driver_neck_01", neckControl)[0]
+                cmds.delete(constraint)
+                
+                neckControlGrp = cmds.group(empty = True, name = "neck_01_fk_anim_grp")
+                constraint = cmds.parentConstraint("driver_neck_01", neckControlGrp)[0]
+                cmds.delete(constraint)
+                
+                cmds.parent(neckControl, neckControlGrp)
+                cmds.setAttr(neckControl + ".ry", -90)
+                cmds.makeIdentity(neckControl, t = 1, r = 1, s = 1, apply = True)
+                
+                
+                neckControlMid = self.createControl("circle", 25, "neck_02_fk_anim")
+                constraint = cmds.parentConstraint("driver_neck_02", neckControlMid)[0]
+                cmds.delete(constraint)
+                
+                neckControlMidGrp = cmds.group(empty = True, name = "neck_02_fk_anim_grp")
+                constraint = cmds.parentConstraint("driver_neck_02", neckControlMidGrp)[0]
+                cmds.delete(constraint)
+                
+                cmds.parent(neckControlMid, neckControlMidGrp)
+                cmds.setAttr(neckControlMid + ".ry", -90)
+                cmds.makeIdentity(neckControlMid, t = 1, r = 1, s = 1, apply = True)
+                
+                
+                
+                neckControlEnd = self.createControl("circle", 25, "neck_03_fk_anim")
+                constraint = cmds.parentConstraint("driver_neck_03", neckControlEnd)[0]
+                cmds.delete(constraint)
+                
+                neckControlEndGrp = cmds.group(empty = True, name = "neck_03_fk_anim_grp")
+                constraint = cmds.parentConstraint("driver_neck_03", neckControlEndGrp)[0]
+                cmds.delete(constraint)
+                
+                cmds.parent(neckControlEnd, neckControlEndGrp)
+                cmds.setAttr(neckControlEnd + ".ry", -90)
+                cmds.makeIdentity(neckControlEnd, t = 1, r = 1, s = 1, apply = True)
+                
+                #setup neck hiearchy
+                cmds.parent(neckControlEndGrp, neckControlMid)
+                cmds.parent(neckControlMidGrp, neckControl)
+                
+
+	
+            #create the FK control for the head
+            headControl = self.createControl("circle", 30, "head_fk_anim")
+            constraint = cmds.parentConstraint("driver_head", headControl)[0]
+            cmds.delete(constraint)
+            
+            headControlGrp = cmds.group(empty = True, name = "head_fk_anim_grp")
+            constraint = cmds.parentConstraint("driver_head", headControlGrp)[0]
+            cmds.delete(constraint)
+            
+            cmds.parent(headControl, headControlGrp)
+            cmds.setAttr(headControl + ".ry", 90)
+	    cmds.setAttr(headControl + ".rz", -35)
+            cmds.makeIdentity(headControl, t = 1, r = 1, s = 1, apply = True)
+            
+            
+            #setup head orientation
+            orientNodes = self.setupHeadOrientation(neckControlEnd, headControl)
+            neckOrientNodes = self.setupNeckOrientation(neckControl)
+            
+            #hook into spine
+            cmds.parent(headControlGrp, orientNodes[0])
+            cmds.parent(orientNodes[0], neckControlEnd)
+            numSpineBones = self.getSpineJoints()
+	    cmds.parent(neckOrientNodes[0], neckControlGrp)
+	    cmds.parent(neckControl, neckOrientNodes[0])
+	    cmds.parentConstraint("driver_spine_0" + str(len(numSpineBones)), neckControlGrp, mo = True)
+
+	    
+            #constrain driver joints to controls
+            if numNeckBones == 2:
+                cmds.parentConstraint(neckControl, "driver_neck_01", mo = True)
+                cmds.parentConstraint(neckControlEnd, "driver_neck_02", mo = True)
+		cmds.connectAttr(neckControl + ".scale", "driver_neck_01.scale")
+		cmds.connectAttr(neckControlEnd + ".scale", "driver_neck_02.scale")
+		
+                
+            if numNeckBones == 3:
+                cmds.parentConstraint(neckControl, "driver_neck_01", mo = True)
+                cmds.parentConstraint(neckControlMid, "driver_neck_02", mo = True)
+                cmds.parentConstraint(neckControlEnd, "driver_neck_03", mo = True)
+                cmds.connectAttr(neckControl + ".scale", "driver_neck_01.scale")
+		cmds.connectAttr(neckControlMid + ".scale", "driver_neck_02.scale")
+		cmds.connectAttr(neckControlEnd + ".scale", "driver_neck_03.scale")
+		
+            cmds.parentConstraint(headControl, "driver_head", mo = True)
+	    cmds.connectAttr(headControl + ".scale", "driver_head.scale")
+            
+            
+            #lock attrs, color controls, and clean up hierarchy
+            for control in [neckControl, neckControlEnd, neckControlMid, headControl]:
+                if cmds.objExists(control):
+                    for attr in [".v"]:
+                        cmds.setAttr(control + attr, lock = True, keyable = False)
+                        cmds.setAttr(control + ".overrideEnabled", 1)
+                        
+            for control in [neckControl, neckControlEnd, neckControlMid]:
+                if cmds.objExists(control):
+                    cmds.setAttr(control + ".overrideColor", 18)
+                    
+                    
+            cmds.setAttr(headControl + ".overrideEnabled", 1)
+            cmds.setAttr(headControl + ".overrideColor", 17)
+            
+            masterGrp = cmds.group(empty = True, name = "head_sys_grp")
+	    cmds.parent(orientNodes[4], masterGrp)
+	    cmds.parent(neckOrientNodes[3], masterGrp)
+	    #cmds.parent(neckOrientNodes[0], masterGrp)
+            
+            
+
+    # # # # # # # # # # # # # # # # # # # # # # # # # # # # # # # # # # # # # # # # # # # # # # # # # # # # # # # # # # # # # # # 
+    # # # # # # # # # # # # # # # # # # # # # # # # # # # # # # # # # # # # # # # # # # # # # # # # # # # # # # # # # # # # # # # 
+    # # # # # # # # # # # # # # # # # # # # # # # # # # # # # # # # # # # # # # # # # # # # # # # # # # # # # # # # # # # # # # # 
+    # # # # # # # # # # # # # # # # # # # # # # # # # # # # # # # # # # # # # # # # # # # # # # # # # # # # # # # # # # # # # # # 
+    # # # # # # # # # # # # # # # # # # # # # # # # # # # # # # # # # # # # # # # # # # # # # # # # # # # # # # # # # # # # # # # 
+    def setupHeadOrientation(self, neckControl, headControl):
+        
+        #create head orient controls (neck, shoulder, body, world)
+        orientMaster = cmds.group(empty = True, name = "head_fk_orient_master")
+        constraint = cmds.parentConstraint("driver_head", orientMaster)[0]
+        cmds.delete(constraint)
+        
+        orientNeck = cmds.duplicate(orientMaster, po = True, name = "head_fk_orient_neck")[0]
+        orientShoulder = cmds.duplicate(orientMaster, po = True, name = "head_fk_orient_shoulder")[0]
+        orientBody = cmds.duplicate(orientMaster, po = True, name = "head_fk_orient_body")[0]
+        orientWorld = cmds.duplicate(orientMaster, po = True, name = "head_fk_orient_world")[0]
+        
+        fkHeadOrientConstraint = cmds.orientConstraint([orientNeck, orientShoulder, orientBody, orientWorld], orientMaster)[0]
+        
+        numSpineBones = self.getSpineJoints()
+        
+        cmds.parent(orientNeck, neckControl)
+        cmds.parent(orientShoulder, "driver_spine_0" + str(len(numSpineBones)))
+        cmds.parent(orientBody, "body_anim")
+        
+        
+        #add the fk orient attr to the head control
+        cmds.select(headControl)
+        cmds.addAttr(longName= "fkOrientation", at = 'enum', en = "neck:shoulder:body:world:", keyable = True)
+        
+        
+        #setup sdks for controlling constraint weight
+        cmds.setAttr(headControl + ".fkOrientation", 0)
+        cmds.setAttr(fkHeadOrientConstraint + "." + orientNeck + "W0", 1)
+        cmds.setAttr(fkHeadOrientConstraint + "." + orientShoulder + "W1", 0)
+        cmds.setAttr(fkHeadOrientConstraint + "." + orientBody + "W2", 0)
+        cmds.setAttr(fkHeadOrientConstraint + "." + orientWorld + "W3", 0)
+        
+        cmds.setDrivenKeyframe([fkHeadOrientConstraint + "." + orientNeck + "W0", fkHeadOrientConstraint + "." + orientShoulder + "W1", fkHeadOrientConstraint + "." + orientBody + "W2", fkHeadOrientConstraint + "." + orientWorld + "W3"], cd = headControl + ".fkOrientation", itt = "linear", ott = "linear")
+
+        cmds.setAttr(headControl + ".fkOrientation", 1)
+        cmds.setAttr(fkHeadOrientConstraint + "." + orientNeck + "W0", 0)
+        cmds.setAttr(fkHeadOrientConstraint + "." + orientShoulder + "W1", 1)
+        cmds.setAttr(fkHeadOrientConstraint + "." + orientBody + "W2", 0)
+        cmds.setAttr(fkHeadOrientConstraint + "." + orientWorld + "W3", 0)
+        
+        cmds.setDrivenKeyframe([fkHeadOrientConstraint + "." + orientNeck + "W0", fkHeadOrientConstraint + "." + orientShoulder + "W1", fkHeadOrientConstraint + "." + orientBody + "W2", fkHeadOrientConstraint + "." + orientWorld + "W3"], cd = headControl + ".fkOrientation", itt = "linear", ott = "linear")
+
+        cmds.setAttr(headControl + ".fkOrientation", 2)
+        cmds.setAttr(fkHeadOrientConstraint + "." + orientNeck + "W0", 0)
+        cmds.setAttr(fkHeadOrientConstraint + "." + orientShoulder + "W1", 0)
+        cmds.setAttr(fkHeadOrientConstraint + "." + orientBody + "W2", 1)
+        cmds.setAttr(fkHeadOrientConstraint + "." + orientWorld + "W3", 0)
+        
+        cmds.setDrivenKeyframe([fkHeadOrientConstraint + "." + orientNeck + "W0", fkHeadOrientConstraint + "." + orientShoulder + "W1", fkHeadOrientConstraint + "." + orientBody + "W2", fkHeadOrientConstraint + "." + orientWorld + "W3"], cd = headControl + ".fkOrientation", itt = "linear", ott = "linear")
+
+
+        cmds.setAttr(headControl + ".fkOrientation", 3)
+        cmds.setAttr(fkHeadOrientConstraint + "." + orientNeck + "W0", 0)
+        cmds.setAttr(fkHeadOrientConstraint + "." + orientShoulder + "W1", 0)
+        cmds.setAttr(fkHeadOrientConstraint + "." + orientBody + "W2", 0)
+        cmds.setAttr(fkHeadOrientConstraint + "." + orientWorld + "W3", 1)
+        
+        cmds.setDrivenKeyframe([fkHeadOrientConstraint + "." + orientNeck + "W0", fkHeadOrientConstraint + "." + orientShoulder + "W1", fkHeadOrientConstraint + "." + orientBody + "W2", fkHeadOrientConstraint + "." + orientWorld + "W3"], cd = headControl + ".fkOrientation", itt = "linear", ott = "linear")
+
+        cmds.setAttr(headControl + ".fkOrientation", 0)
+        
+        return [orientMaster, orientNeck, orientShoulder, orientBody, orientWorld]
+            
+            
+    # # # # # # # # # # # # # # # # # # # # # # # # # # # # # # # # # # # # # # # # # # # # # # # # # # # # # # # # # # # # # # # 
+    # # # # # # # # # # # # # # # # # # # # # # # # # # # # # # # # # # # # # # # # # # # # # # # # # # # # # # # # # # # # # # # 
+    # # # # # # # # # # # # # # # # # # # # # # # # # # # # # # # # # # # # # # # # # # # # # # # # # # # # # # # # # # # # # # # 
+    # # # # # # # # # # # # # # # # # # # # # # # # # # # # # # # # # # # # # # # # # # # # # # # # # # # # # # # # # # # # # # # 
+    # # # # # # # # # # # # # # # # # # # # # # # # # # # # # # # # # # # # # # # # # # # # # # # # # # # # # # # # # # # # # # # 
+    def setupNeckOrientation(self, neckControl):
+        
+        #create head orient controls (neck, shoulder, body, world)
+        orientMaster = cmds.group(empty = True, name = "neck_fk_orient_master")
+        constraint = cmds.parentConstraint("driver_neck_01", orientMaster)[0]
+        cmds.delete(constraint)
+        
+        orientShoulder = cmds.duplicate(orientMaster, po = True, name = "neck_fk_orient_shoulder")[0]
+        orientBody = cmds.duplicate(orientMaster, po = True, name = "neck_fk_orient_body")[0]
+        orientWorld = cmds.duplicate(orientMaster, po = True, name = "neck_fk_orient_world")[0]
+        
+        fkHeadOrientConstraint = cmds.orientConstraint([orientShoulder, orientBody, orientWorld], orientMaster)[0]
+        
+        numSpineBones = self.getSpineJoints()
+        
+        cmds.parent(orientShoulder, "driver_spine_0" + str(len(numSpineBones)))
+        cmds.parent(orientBody, "body_anim")
+        
+        
+        #add the fk orient attr to the head control
+	cmds.select(neckControl)
+        cmds.addAttr(longName= "fkOrientation", at = 'enum', en = "shoulder:body:world:", keyable = True)
+        
+        
+        #setup sdks for controlling constraint weight
+        cmds.setAttr(neckControl + ".fkOrientation", 0)
+        cmds.setAttr(fkHeadOrientConstraint + "." + orientShoulder + "W0", 1)
+        cmds.setAttr(fkHeadOrientConstraint + "." + orientBody + "W1", 0)
+        cmds.setAttr(fkHeadOrientConstraint + "." + orientWorld + "W2", 0)
+        
+        cmds.setDrivenKeyframe([fkHeadOrientConstraint + "." + orientShoulder + "W0", fkHeadOrientConstraint + "." + orientBody + "W1", fkHeadOrientConstraint + "." + orientWorld + "W2"], cd = neckControl + ".fkOrientation", itt = "linear", ott = "linear")
+
+        cmds.setAttr(neckControl + ".fkOrientation", 1)
+        cmds.setAttr(fkHeadOrientConstraint + "." + orientShoulder + "W0", 0)
+        cmds.setAttr(fkHeadOrientConstraint + "." + orientBody + "W1", 1)
+        cmds.setAttr(fkHeadOrientConstraint + "." + orientWorld + "W2", 0)
+        
+	cmds.setDrivenKeyframe([fkHeadOrientConstraint + "." + orientShoulder + "W0", fkHeadOrientConstraint + "." + orientBody + "W1", fkHeadOrientConstraint + "." + orientWorld + "W2"], cd = neckControl + ".fkOrientation", itt = "linear", ott = "linear")
+
+        cmds.setAttr(neckControl + ".fkOrientation", 2)
+        cmds.setAttr(fkHeadOrientConstraint + "." + orientShoulder + "W0", 0)
+        cmds.setAttr(fkHeadOrientConstraint + "." + orientBody + "W1", 0)
+        cmds.setAttr(fkHeadOrientConstraint + "." + orientWorld + "W2", 1)
+        
+	cmds.setDrivenKeyframe([fkHeadOrientConstraint + "." + orientShoulder + "W0", fkHeadOrientConstraint + "." + orientBody + "W1", fkHeadOrientConstraint + "." + orientWorld + "W2"], cd = neckControl + ".fkOrientation", itt = "linear", ott = "linear")
+        
+        cmds.setAttr(neckControl + ".fkOrientation", 0)
+        
+        return [orientMaster, orientShoulder, orientBody, orientWorld]
+    
+    
+    # # # # # # # # # # # # # # # # # # # # # # # # # # # # # # # # # # # # # # # # # # # # # # # # # # # # # # # # # # # # # # # 
+    # # # # # # # # # # # # # # # # # # # # # # # # # # # # # # # # # # # # # # # # # # # # # # # # # # # # # # # # # # # # # # # 
+    # # # # # # # # # # # # # # # # # # # # # # # # # # # # # # # # # # # # # # # # # # # # # # # # # # # # # # # # # # # # # # # 
+    # # # # # # # # # # # # # # # # # # # # # # # # # # # # # # # # # # # # # # # # # # # # # # # # # # # # # # # # # # # # # # # 
+    # # # # # # # # # # # # # # # # # # # # # # # # # # # # # # # # # # # # # # # # # # # # # # # # # # # # # # # # # # # # # # # 
+    def rigLeafJoints(self):
+        
+        #find attrs on the skeleton settings node
+        createdControls = []
+        
+        attrs = cmds.listAttr("Skeleton_Settings")
+        for attr in attrs:
+            if attr.find("extraJoint") == 0:
+                attribute = cmds.getAttr("Skeleton_Settings." + attr, asString = True)
+                jointType = attribute.partition("/")[2].partition("/")[0]
+                name = attribute.rpartition("/")[2]
+                parent = attribute.partition("/")[0]
+		
+		if parent.find("(") != -1:
+		    parent = parent.partition(" (")[0]
+		
+                
+                if jointType == "leaf":
+                    attrs = name.partition("(")[2].partition(")")[0]
+                    name = name.partition(" (")[0]
+                    
+                    #create the control
+                    control = cmds.curve(name = (name + "_anim"), d = 1, p = [(0, 0, 1), (0, 0.5, 0.866025), (0, 0.866025, 0.5), (0, 1, 0), (0, 0.866025, -0.5), (0, 0.5, -0.866025), (0, 0, -1), (0, -0.5, -0.866025), (0, -0.866025, -0.5), (0, -1, 0), (0, -0.866025, 0.5), (0, -0.5, 0.866025), (0, 0, 1), (0.707107, 0, 0.707107), (1, 0, 0), (0.707107, 0, -0.707107), (0, 0, -1), (-0.707107, 0, -0.707107), (-1, 0, 0), (-0.866025, 0.5, 0), (-0.5, 0.866025, 0), (0, 1, 0), (0.5, 0.866025, 0), (0.866025, 0.5, 0), (1, 0, 0), (0.866025, -0.5, 0), (0.5, -0.866025, 0), (0, -1, 0), (-0.5, -0.866025, 0), (-0.866025, -0.5, 0), (-1, 0, 0), (-0.707107, 0, 0.707107), (0, 0, 1)])
+                    
+                    #scale up
+                    cmds.setAttr(control + ".sx", 9)
+                    cmds.setAttr(control + ".sy", 9)
+                    cmds.setAttr(control + ".sz", 9)
+                    
+                    #freeze transforms
+                    cmds.makeIdentity(control, t = 1, r = 1, s = 1, apply = True)
+                    
+                    #position control
+                    constraint = cmds.parentConstraint("driver_" + name, control)[0]
+                    cmds.delete(constraint)
+                    
+                    #create the control group
+                    ctrlGrp = cmds.group(empty = True, name = (name + "_anim_grp"))
+                    constraint = cmds.parentConstraint("driver_" + name, ctrlGrp)[0]
+                    cmds.delete(constraint)
+                    
+                    #create space switcher group
+		    spaceSwitcherFollow = cmds.duplicate(ctrlGrp, po = True, name = (name + "_anim_space_switcher_follow"))[0]
+                    spaceSwitcher = cmds.duplicate(ctrlGrp, po = True, name = (name + "_anim_space_switcher"))[0]
+                    
+                    #create the top parent group
+                    topParent = cmds.duplicate(ctrlGrp, po = True, name = (name + "_parent_grp"))[0]
+                    #parent control to group
+		    cmds.parent(spaceSwitcher, spaceSwitcherFollow)
+                    cmds.parent(spaceSwitcherFollow, topParent)
+                    cmds.parent(ctrlGrp, spaceSwitcher)
+                    cmds.parent(control, ctrlGrp)
+                    
+                    #constrain driver joint to control
+                    cmds.parentConstraint(control, "driver_" + name)
+		    cmds.connectAttr(control + ".scale", "driver_" + name + ".scale")
+		                     
+                    #lock attrs depending on type of control
+                    lockAttrs = []
+                    if attrs == "TR":
+                        lockAttrs = [".v"]
+                        
+                    if attrs == "T":
+                        lockAttrs = [".v", ".rx", ".ry", ".rz"]
+                        
+                    if attrs == "R":
+                        lockAttrs = [".v", ".tx", ".ty", ".tz"]
+                       
+                    
+                    for attr in lockAttrs:
+                        cmds.setAttr(control + attr, lock = True, keyable = False)
+                        
+                        
+                    #parent constrain the topParent to the parent of the control
+                    cmds.parentConstraint("driver_" + parent, topParent, mo = True)
+		    
+                    
+                    
+                    #color the control
+                    cmds.setAttr(control + ".overrideEnabled", 1)
+                    cmds.setAttr(control + ".overrideColor", 18)
+                    
+                    
+                    #add the topParent to the createdControls list
+                    createdControls.append(topParent)
+                    
+        return createdControls
+                    
+                    
+                
+    # # # # # # # # # # # # # # # # # # # # # # # # # # # # # # # # # # # # # # # # # # # # # # # # # # # # # # # # # # # # # # # 
+    # # # # # # # # # # # # # # # # # # # # # # # # # # # # # # # # # # # # # # # # # # # # # # # # # # # # # # # # # # # # # # # 
+    # # # # # # # # # # # # # # # # # # # # # # # # # # # # # # # # # # # # # # # # # # # # # # # # # # # # # # # # # # # # # # # 
+    # # # # # # # # # # # # # # # # # # # # # # # # # # # # # # # # # # # # # # # # # # # # # # # # # # # # # # # # # # # # # # # 
+    # # # # # # # # # # # # # # # # # # # # # # # # # # # # # # # # # # # # # # # # # # # # # # # # # # # # # # # # # # # # # # # 
+    def rigJiggleJoints(self):
+        
+        
+        #find attrs on the skeleton settings node
+        createdControls = []
+        
+        attrs = cmds.listAttr("Skeleton_Settings")
+        for attr in attrs:
+            if attr.find("extraJoint") == 0:
+                attribute = cmds.getAttr("Skeleton_Settings." + attr, asString = True)
+                jointType = attribute.partition("/")[2].partition("/")[0]
+                name = attribute.rpartition("/")[2]
+                parent = attribute.partition("/")[0]
+                
+                if jointType == "jiggle":
+                    
+                    #duplicate the driver joint
+                    jiggleStart = cmds.duplicate("driver_" + name, po = True, name = "rig_" + name + "_start")[0]
+                    cmds.parent(jiggleStart, world = True)
+                    
+                    jiggleEnd = cmds.duplicate(jiggleStart, po = True, name = "rig_" + name + "_end")[0]
+                    
+                    cmds.parent(jiggleEnd, jiggleStart)
+                    
+                    #move jiggleEnd down a bit in up axis
+                    scaleFactor = self.getScaleFactor()
+                    jointPos = cmds.xform(jiggleStart, q = True, ws = True, t = True)
+                    cmds.xform(jiggleEnd, ws = True, t = (jointPos[0], jointPos[1], (jointPos[2] - (25 * scaleFactor))))
+                    
+                    
+                    
+                    #@@@@@@@@@@@@@@@@@@@@@@@@@@@@@@@@@@@@@@@@@@@@@@@@@@@@@@@@@@@@@@@@@@@@@@@@@@@@@@@@@@@@@@@@@@@@@@@@@@@#
+                    #Create curve on joint chain
+                    #@@@@@@@@@@@@@@@@@@@@@@@@@@@@@@@@@@@@@@@@@@@@@@@@@@@@@@@@@@@@@@@@@@@@@@@@@@@@@@@@@@@@@@@@@@@@@@@@@@@#
+                    
+                    joints = [jiggleStart, jiggleEnd]
+                    
+                    positions = []
+                    
+                    #get the world space positions of each joint, and create a curve using those positions
+                    for i in range(int(len(joints))):
+                        pos = cmds.xform(joints[i], q = True, ws = True, t = True)
+                        positions.append(pos)
+                        
+                    createCurveCommand = "curve -d 1"
+                    for pos in positions:
+                        xPos = pos[0]
+                        yPos = pos[1]
+                        zPos = pos[2]
+                        createCurveCommand += " -p " + str(xPos) + " " + str(yPos) + " " + str(zPos)
+                        
+                    for i in range(int(len(positions))):
+                        createCurveCommand += " -k " + str(i)
+                        
+                    curve = mel.eval(createCurveCommand)
+		    curve = cmds.rename(curve, name + "_dynCurve")
+		    cmds.setAttr(curve + ".v", 0)
+
+                    #@@@@@@@@@@@@@@@@@@@@@@@@@@@@@@@@@@@@@@@@@@@@@@@@@@@@@@@@@@@@@@@@@@@@@@@@@@@@@@@@@@@@@@@@@@@@@@@@@@@#
+                    #Create hair system
+                    #@@@@@@@@@@@@@@@@@@@@@@@@@@@@@@@@@@@@@@@@@@@@@@@@@@@@@@@@@@@@@@@@@@@@@@@@@@@@@@@@@@@@@@@@@@@@@@@@@@@#
+                    
+                    cmds.select(curve)
+                    
+                    #find all hair systems in scene
+                    hairSystems = cmds.ls(type = "hairSystem")
+                    hairSys = ""
+                    
+                    
+                    #create the hair system and make the stiffness uniform
+                    madeHairCurve = True
+                    if hairSys == "":
+                        hairSys = cmds.createNode("hairSystem")
+			cmds.removeMultiInstance(hairSys + ".stiffnessScale[1]", b = True)
+			cmds.setAttr(hairSys + ".clumpWidth", 0.0)
+			cmds.connectAttr("time1.outTime", hairSys + ".currentTime")
+			
+		    hairSysParent = cmds.listRelatives(hairSys, parent = True)
+		    hairSysParent = cmds.rename(hairSysParent, name + "_hairSystem")
+		    cmds.setAttr(hairSysParent + ".v", 0)
+		    hairSys =  name + "_hairSystemShape"
+		    
+		    #create the hair follicle
+		    hair = cmds.createNode("follicle")
+		    cmds.setAttr(hair + ".parameterU", 0)
+		    cmds.setAttr(hair + ".parameterV", 0)
+		    hairTransforms = cmds.listRelatives(hair, p = True)
+		    hairDag = hairTransforms[0]
+		    hairDag = cmds.rename(hairDag, name + "_follicle")
+		    hair = name + "_follicleShape"
+		    
+		    cmds.setAttr(hairDag + ".v", 0)
+		    cmds.setAttr(hair + ".startDirection", 1)
+		    
+		    #get the curve CVs and set follicle degree to 1 if CVs are less than 3
+		    curveCVs = cmds.getAttr(curve + ".cp", size = True)
+		    if curveCVs < 3:
+			cmds.setAttr(hair + ".degree", 1)
+			
+		    #parent the curve to the follicle and connect the curve's worldspace[0] to the follicle startPos
+		    cmds.parent(curve, hairDag, relative = True)
+		    cmds.connectAttr(curve + ".worldSpace[0]", hair + ".startPosition")
+		    
+		    #connect the hair follicle to the hair system
+		    cmds.connectAttr(hair + ".outHair", hairSys + ".inputHair[0]")
+		    
+		    #create a new curve and connect the follicle's outCurve attr to the new curve
+		    cmds.connectAttr(hairSys + ".outputHair[0]", hair + ".currentPosition")
+		    crv = cmds.createNode("nurbsCurve")
+		    crvParent = cmds.listRelatives(crv, parent = True)[0]
+		    crvParent = cmds.rename(crvParent, name + "_track_rt_curve")
+		    crv =  name + "_track_rt_curveShape"
+		    cmds.setAttr(crvParent + ".v", 0)
+		    
+		    cmds.connectAttr(hair + ".outCurve", crv + ".create")
+		    
+		    #set the hair follicle attrs 
+		    if len(hairDag) > 0:
+			cmds.setAttr(hairDag + ".pointLock", 3)
+			cmds.setAttr(hairDag + ".restPose", 1)
+			
+		    cmds.select(hairSys)
+		    
+
+
+                    #@@@@@@@@@@@@@@@@@@@@@@@@@@@@@@@@@@@@@@@@@@@@@@@@@@@@@@@@@@@@@@@@@@@@@@@@@@@@@@@@@@@@@@@@@@@@@@@@@@@#
+                    #Create Spline Handle for the selected chain and the duplicated curve. the original is driven by hair
+                    #@@@@@@@@@@@@@@@@@@@@@@@@@@@@@@@@@@@@@@@@@@@@@@@@@@@@@@@@@@@@@@@@@@@@@@@@@@@@@@@@@@@@@@@@@@@@@@@@@@@#
+		    
+		    ikNodes = cmds.ikHandle(sol = "ikSplineSolver",  ccv = False, pcv = False, snc = True, sj = jiggleStart, ee = jiggleEnd, c = crv)[0]
+		    cmds.setAttr(ikNodes + ".v", 0)
+		    ikNodes = cmds.rename(ikNodes, name + "_dynChain_ikHandle")
+
+                    #@@@@@@@@@@@@@@@@@@@@@@@@@@@@@@@@@@@@@@@@@@@@@@@@@@@@@@@@@@@@@@@@@@@@@@@@@@@@@@@@@@@@@@@@@@@@@@@@@@@#
+                    #Create a duplicate joint chain for manual animation
+                    #@@@@@@@@@@@@@@@@@@@@@@@@@@@@@@@@@@@@@@@@@@@@@@@@@@@@@@@@@@@@@@@@@@@@@@@@@@@@@@@@@@@@@@@@@@@@@@@@@@@#
+		    
+		    dupeChain = cmds.duplicate(jiggleStart, rr = True, rc = True)
+		    dupeStartJoint = dupeChain[0]
+		    
+		    dupeJoints = cmds.listRelatives(dupeStartJoint, ad = True)
+		    joints = cmds.listRelatives(jiggleStart, ad = True)
+		    
+		    #rename duped joints and connect real joints to duped joints
+		    for i in range(int(len(joints))):
+			if cmds.objectType(dupeJoints[i], isType = 'joint'):
+			    cmds.rename(dupeJoints[i], "ANIM_" + joints[i])
+			    cmds.connectAttr("ANIM_" + joints[i] + ".r", joints[i] + ".r", force = True)
+			    
+			else:
+			    cmds.delete(dupeJoints[i])
+			    
+		    #connect up  start joint to ANIM start joint
+		    cmds.connectAttr(dupeStartJoint + ".t", jiggleStart + ".t")
+		    cmds.connectAttr(dupeStartJoint + ".r", jiggleStart + ".r")
+		    cmds.connectAttr(dupeStartJoint + ".s", jiggleStart + ".s")
+		    
+		    dupeStartJoint = cmds.rename(dupeStartJoint, "ANIM_" + jiggleStart)
+		    #cmds.parent(dupeStartJoint, world = True)
+		    
+                    #@@@@@@@@@@@@@@@@@@@@@@@@@@@@@@@@@@@@@@@@@@@@@@@@@@@@@@@@@@@@@@@@@@@@@@@@@@@@@@@@@@@@@@@@@@@@@@@@@@@#
+                    #Create skinCluster between duplicate curve and animation joint chain(dupe chain)
+                    #@@@@@@@@@@@@@@@@@@@@@@@@@@@@@@@@@@@@@@@@@@@@@@@@@@@@@@@@@@@@@@@@@@@@@@@@@@@@@@@@@@@@@@@@@@@@@@@@@@@#
+		    cmds.select(dupeStartJoint, hi = True)
+		    dupeSkel = cmds.ls(sl = True, type = "joint")
+		    cmds.select(curve)
+		    cmds.select(dupeSkel, add = True)
+		    skinCluster = cmds.skinCluster(tsb = True, mi = 3, dr = 4)
+		    
+
+                    #@@@@@@@@@@@@@@@@@@@@@@@@@@@@@@@@@@@@@@@@@@@@@@@@@@@@@@@@@@@@@@@@@@@@@@@@@@@@@@@@@@@@@@@@@@@@@@@@@@@#
+                    #Create the control that has all of our dynamic attrs
+                    #@@@@@@@@@@@@@@@@@@@@@@@@@@@@@@@@@@@@@@@@@@@@@@@@@@@@@@@@@@@@@@@@@@@@@@@@@@@@@@@@@@@@@@@@@@@@@@@@@@@#
+		    control = self.createControl("square", 15, name + "_anim")
+		    constraint = cmds.parentConstraint(jiggleStart, control)[0]
+		    cmds.delete(constraint)
+		    
+		    ctrlGrp = cmds.group(empty = True, name = control + "_grp")
+		    constraint = cmds.parentConstraint(jiggleStart, ctrlGrp)[0]
+		    cmds.delete(constraint)
+		    
+		    cmds.parent(control, ctrlGrp)
+		    cmds.parentConstraint(control, dupeStartJoint)
+		    
+		    #lock attrs
+		    cmds.setAttr(control + ".sx", lock = True, keyable = False)
+		    cmds.setAttr(control + ".sy", lock = True, keyable = False)
+		    cmds.setAttr(control + ".sz", lock = True, keyable = False)
+		    cmds.setAttr(control + ".v", lock = True, keyable = False)
+		    
+		    
+		    #add attrs
+		    cmds.select(control)
+		    
+		    cmds.addAttr(ln = "___DYNAMICS___", at = "double", keyable = True)
+		    cmds.setAttr(control + ".___DYNAMICS___", lock = True)
+		    cmds.addAttr(ln = "chainAttach", at = "enum", en = "No Attach:Base:Tip:Both End:", dv = 1, keyable = True)
+		    cmds.addAttr(ln = "chainStartEnvelope", at = "double", min = 0, max = 1, dv = 1, keyable = True)
+		    cmds.addAttr(ln = "chainStartFrame", at = "double", dv = 1, keyable = True)
+		    
+		    cmds.addAttr(ln = "___BEHAVIOR___", at = "double", keyable = True)
+		    cmds.setAttr(control + ".___BEHAVIOR___", lock = True)
+		    cmds.addAttr(ln = "chainStiffness", at = "double", min = 0, dv = .1, keyable = True)
+		    cmds.addAttr(ln = "chainDamping", at = "double", min = 0, dv = 0.2, keyable = True)
+		    cmds.addAttr(ln = "chainGravity", at = "double", min = 0, dv = 1, keyable = True)
+		    cmds.addAttr(ln = "chainIteration", at = "long", min = 0, dv = 1, keyable = True)
+		    
+		    cmds.addAttr(ln = "___COLLISIONS___", at = "double", keyable = True)
+		    cmds.setAttr(control + ".___COLLISIONS___", lock = True)
+		    cmds.addAttr(ln = "chainCollide", at = "bool", dv = 0, keyable = True)
+		    cmds.addAttr(ln = "chainWidthBase", at = "double", min = 0, dv = 1, keyable = True)
+		    cmds.addAttr(ln = "chainWidthExtremity", at = "double", min = 0, dv = 1, keyable = True)
+		    cmds.addAttr(ln = "chainCollideGround", at = "bool", dv = 0, keyable = True)
+		    cmds.addAttr(ln = "chainCollideGroundHeight", at = "double", dv = 0, keyable = True)
+		    
+		    
+		    #connect attrs
+		    cmds.connectAttr(control + ".chainStartEnvelope", ikNodes + ".ikBlend")
+		    cmds.connectAttr(control + ".chainAttach", hair + ".pointLock")
+		    cmds.connectAttr(control + ".chainStartFrame", hairSys + ".startFrame")
+		    cmds.connectAttr(control + ".chainStiffness", hairSys + ".stiffness")
+		    cmds.connectAttr(control + ".chainDamping", hairSys + ".damp")
+		    cmds.connectAttr(control + ".chainGravity", hairSys + ".gravity")
+		    cmds.connectAttr(control + ".chainIteration", hairSys + ".iterations")
+		    cmds.connectAttr(control + ".chainCollide", hairSys + ".collide")
+		    cmds.connectAttr(control + ".chainWidthBase", hairSys + ".clumpWidth")
+		    cmds.connectAttr(control + ".chainWidthExtremity", hairSys + ".clumpWidthScale[1].clumpWidthScale_FloatValue")
+		    cmds.connectAttr(control + ".chainCollideGround", hairSys + ".collideGround")
+		    cmds.connectAttr(control + ".chainCollideGroundHeight", hairSys + ".groundHeight")
+		    
+		    
+                    #@@@@@@@@@@@@@@@@@@@@@@@@@@@@@@@@@@@@@@@@@@@@@@@@@@@@@@@@@@@@@@@@@@@@@@@@@@@@@@@@@@@@@@@@@@@@@@@@@@@#
+                    #Create the expression for real time
+                    #@@@@@@@@@@@@@@@@@@@@@@@@@@@@@@@@@@@@@@@@@@@@@@@@@@@@@@@@@@@@@@@@@@@@@@@@@@@@@@@@@@@@@@@@@@@@@@@@@@@#
+		    track_RealTime = cmds.spaceLocator(name = name + "_track_rt_loc")[0]
+		    cmds.pointConstraint(dupeSkel[len(dupeSkel) -1], track_RealTime)
+		    connections = cmds.listConnections(hairSys + ".currentTime", p = True, c = True)
+		    cmds.disconnectAttr(connections[1], connections[0])
+		    
+		    expressionString = "if(frame!= " + hairSys + ".startFrame)\n\t" + hairSys + ".currentTime = " + hairSys + ".currentTime + 1 + " + track_RealTime + ".tx - " + track_RealTime + ".tx + " + track_RealTime + ".ty - " + track_RealTime + ".ty + " + track_RealTime + ".tz - " + track_RealTime + ".tz + " + control + ".chainWidthBase - "+ control + ".chainWidthBase + "+ control + ".chainWidthExtremity - "+ control + ".chainWidthExtremity + " + control + ".chainGravity - "+ control + ".chainGravity;\n" +"else\n\t" + hairSys + ".currentTime = " + hairSys + ".startFrame;"
+		    cmds.expression(name = "EXP_" + hairSys + "_TRACK_RealTime", string = expressionString)
+		    cmds.setAttr(track_RealTime + ".v", 0)
+		    
+                    #@@@@@@@@@@@@@@@@@@@@@@@@@@@@@@@@@@@@@@@@@@@@@@@@@@@@@@@@@@@@@@@@@@@@@@@@@@@@@@@@@@@@@@@@@@@@@@@@@@@#
+                    #Set Defaults
+                    #@@@@@@@@@@@@@@@@@@@@@@@@@@@@@@@@@@@@@@@@@@@@@@@@@@@@@@@@@@@@@@@@@@@@@@@@@@@@@@@@@@@@@@@@@@@@@@@@@@@#
+		    cmds.setAttr(hairSys + ".drawCollideWidth", 1)
+		    cmds.setAttr(hairSys + ".widthDrawSkip", 0)
+		    cmds.setAttr(hair + ".degree", 1)
+		    
+		    cmds.parentConstraint("driver_" + parent, ctrlGrp, mo = True)
+		    
+		    if cmds.objExists("dynHairChain") == False:
+			cmds.group(empty = True, name = "dynHairChain")
+			
+		    if cmds.objExists(jiggleStart + "_HairControls") == False:
+			group = cmds.group([ikNodes, hairSys, hair, track_RealTime, ctrlGrp, crvParent], name = jiggleStart + "_HairControls")
+			
+		    else:
+			cmds.parent([ikNodes, hairSys, hair, ctrlGrp,crvParent ], jiggleStart + "_HairControls")
+			
+		    cmds.parent(group, "dynHairChain")
+		    
+                    #@@@@@@@@@@@@@@@@@@@@@@@@@@@@@@@@@@@@@@@@@@@@@@@@@@@@@@@@@@@@@@@@@@@@@@@@@@@@@@@@@@@@@@@@@@@@@@@@@@@#
+                    #Cleanup
+                    #@@@@@@@@@@@@@@@@@@@@@@@@@@@@@@@@@@@@@@@@@@@@@@@@@@@@@@@@@@@@@@@@@@@@@@@@@@@@@@@@@@@@@@@@@@@@@@@@@@@#
+		    jointsGrp = cmds.group(empty = True, name = name + "_jiggle_jointsGrp")
+		    createdControls.append(jointsGrp)
+		    cmds.parent([jiggleStart, dupeStartJoint], jointsGrp)
+		    
+		    cmds.setAttr(control + ".overrideEnabled", 1)
+                    cmds.setAttr(control + ".overrideColor", 18)
+		    
+		    #constrain driver joints
+		    cmds.parentConstraint(jiggleStart, "driver_" + name, mo = True)
+		    
+		    
+	#return top level group
+	return createdControls
+		    
+		    
+		    
+
+
+        
+    # # # # # # # # # # # # # # # # # # # # # # # # # # # # # # # # # # # # # # # # # # # # # # # # # # # # # # # # # # # # # # # 
+    # # # # # # # # # # # # # # # # # # # # # # # # # # # # # # # # # # # # # # # # # # # # # # # # # # # # # # # # # # # # # # # 
+    # # # # # # # # # # # # # # # # # # # # # # # # # # # # # # # # # # # # # # # # # # # # # # # # # # # # # # # # # # # # # # # 
+    # # # # # # # # # # # # # # # # # # # # # # # # # # # # # # # # # # # # # # # # # # # # # # # # # # # # # # # # # # # # # # # 
+    # # # # # # # # # # # # # # # # # # # # # # # # # # # # # # # # # # # # # # # # # # # # # # # # # # # # # # # # # # # # # # # 
+    def rigCustomJointChains(self):
+        
+        
+        #find attrs on the skeleton settings node
+        createdControls = []
+        rootControl = ""
+        attrs = cmds.listAttr("Skeleton_Settings")
+        for attr in attrs:
+            if attr.find("extraJoint") == 0:
+                attribute = cmds.getAttr("Skeleton_Settings." + attr, asString = True)
+                jointType = attribute.partition("/")[2].partition("/")[0]
+                name = attribute.rpartition("/")[2]
+                parent = attribute.partition("/")[0]
+                
+                if jointType == "chain":
+		    
+		    numJointsInChain = name.partition("(")[2].partition(")")[0]
+		    name = name.partition(" (")[0]
+		    
+
+		    #@@@@@@@@@@@@@@@@@@@@@@@@@@@@@@@@@@@@@@@@@@@@@@@@@@@@@@@@@@@@@@@@@@@@@@@@@@@@@@@@@@@@@@@@@@@@@@@@@@@#
+		    #FK RIG
+		    #@@@@@@@@@@@@@@@@@@@@@@@@@@@@@@@@@@@@@@@@@@@@@@@@@@@@@@@@@@@@@@@@@@@@@@@@@@@@@@@@@@@@@@@@@@@@@@@@@@@#
+		    fkJoints = []
+		    frameCacheNodes = []
+		    fkRootGrp = ""
+		    
+		    for i in range(int(numJointsInChain)):
+			jointNum = i + 1
+			if jointNum == 1:
+			    firstControl = "fk_" + name + "_0" + str(jointNum) + "_anim"
+			
+			#create and position the joint
+			if cmds.objExists("rig_fk_" + name + "_0" + str(jointNum)):
+			    cmds.delete("rig_fk_" + name + "_0" + str(jointNum))
+			    
+			cmds.select(clear = True)
+			joint = cmds.joint(name = "rig_fk_" + name + "_0" + str(jointNum))
+			cmds.select(clear = True)
+			fkJoints.append(joint)
+			
+			constraint = cmds.parentConstraint("driver_" + name + "_0" + str(jointNum), joint)[0]
+			cmds.delete(constraint)
+			
+			#create the control and position
+			control = self.createControl("circle", 15, "fk_" + name + "_0" + str(jointNum) + "_anim")
+			cmds.setAttr(control + ".rx", 90)
+			cmds.makeIdentity(control, t = 1, r = 1, s = 1, apply = True)
+			constraint = cmds.parentConstraint(joint, control)[0]
+			cmds.delete(constraint)
+			cmds.makeIdentity(control, t = 0, r = 1, s = 0, apply = True)
+			
+			#cmds.setAttr(control + ".rz", -90)
+			cmds.makeIdentity(control, t = 1, r = 1, s = 1, apply = True)
+			
+			#create the control grp and parent the control to the group
+			ctrlGrp = cmds.group(empty = True, name = "fk_" + name + "_0" + str(jointNum) + "_grp")
+			if i == 0:
+			    fkRootGrp = ctrlGrp
+			
+			constraint = cmds.parentConstraint(joint, ctrlGrp)[0]
+			cmds.delete(constraint)
+			cmds.parent(control, ctrlGrp)
+			cmds.makeIdentity(control, t = 1, r = 1, s = 1, apply = True)
+			
+			#duplicate the ctrl grp for the lag mode
+			lagGrp = cmds.duplicate(ctrlGrp, po = True, name = "fk_" + name + "_0" + str(jointNum) + "_lag_grp")[0]
+			cmds.parent(lagGrp, ctrlGrp)
+			cmds.parent(control, lagGrp)
+			
+			
+			#color the control
+			cmds.setAttr(control + ".overrideEnabled", 1)
+			cmds.setAttr(control + ".overrideColor", 18)
+			
+			
+			if jointNum != 1:
+			    cmds.setAttr(control + ".sx", lock = True, keyable = False)
+			    cmds.setAttr(control + ".sy", lock = True, keyable = False)
+			    cmds.setAttr(control + ".sz", lock = True, keyable = False)
+			    
+			else:
+			    #aliasAttr one of the scale axis and connect the other two to that one
+			    cmds.aliasAttr("global_scale", control + ".scaleZ")
+			    cmds.connectAttr(control + ".scaleZ", control + ".scaleX")
+			    cmds.connectAttr(control + ".scaleZ", control + ".scaleY")
+			    
+			    cmds.setAttr(control + ".sx", lock = True, keyable = False)
+			    cmds.setAttr(control + ".sy", lock = True, keyable = False)			    
+
+			    
+			    
+			cmds.setAttr(control + ".v", lock = True, keyable = False)
+
+		    
+			#parent the joint to the control
+			cmds.parent(joint, control)
+			
+			
+			
+			#TEMP!
+			cmds.parentConstraint(joint, "driver_" + name + "_0" + str(jointNum))
+			
+			
+			#add attr to root joint of chain for turning on "lag" mode
+			
+			if i == 0:
+			    rootControl = ctrlGrp
+			    cmds.select(control)
+			    cmds.addAttr(longName='lagMode', defaultValue=0, minValue=0, maxValue=1, keyable = True)
+			    cmds.addAttr(longName='lagValue', defaultValue= 3, minValue= 0, maxValue=100, keyable = False)
+			    
+			    #setup lag mode node chain
+			    frameCacheX = cmds.createNode("frameCache")
+			    frameCacheX = cmds.rename(frameCacheX, name + "_frameCacheX")
+			    frameCacheY = cmds.createNode("frameCache")
+			    frameCacheY = cmds.rename(frameCacheY, name + "_frameCacheY")
+			    frameCacheZ = cmds.createNode("frameCache")
+			    frameCacheZ = cmds.rename(frameCacheZ, name + "_frameCacheZ")
+			    frameCacheNodes.append(frameCacheX)
+			    frameCacheNodes.append(frameCacheY)
+			    frameCacheNodes.append(frameCacheZ)
+			    
+			    
+
+			    
+			    #create a switcher node
+			    switchNode = cmds.shadingNode("multiplyDivide", asUtility = True, name = name + "_switcherNode")
+			    frameCacheNodes.append(switchNode)
+			    cmds.connectAttr(control + ".rotateX", switchNode + ".input1X")
+			    cmds.connectAttr(control + ".lagMode", switchNode + ".input2X")
+			    
+			    cmds.connectAttr(control + ".rotateY", switchNode + ".input1Y")
+			    cmds.connectAttr(control + ".lagMode", switchNode + ".input2Y")
+			    
+			    cmds.connectAttr(control + ".rotateZ", switchNode + ".input1Z")
+			    cmds.connectAttr(control + ".lagMode", switchNode + ".input2Z")
+			    
+			    cmds.connectAttr(switchNode + ".outputX", frameCacheX + ".stream")
+			    cmds.connectAttr(switchNode + ".outputY", frameCacheY + ".stream")
+			    cmds.connectAttr(switchNode + ".outputZ", frameCacheZ + ".stream")
+			    
+			    mainControl = control
+
+
+
+					   
+			#setup FK hierarchy
+			lagValue = cmds.getAttr(mainControl + ".lagValue")
+			if i != 0:
+			    cmds.parent(ctrlGrp, lastControl)
+			    
+			    #connect framecache results to lag Grps
+			    
+			    mode = "past"
+
+			    cmds.connectAttr(frameCacheX + "." + mode + "[" + str(int(abs(lagValue)) * (i + 1)) + "]", lagGrp + ".rotateX")
+			    cmds.connectAttr(frameCacheY + "." + mode + "[" + str(int(abs(lagValue)) * (i + 1)) + "]", lagGrp + ".rotateY")
+			    cmds.connectAttr(frameCacheZ + "." + mode + "[" + str(int(abs(lagValue)) * (i + 1)) + "]", lagGrp + ".rotateZ")
+			    
+			
+			lastControl = control
+			
+			
+
+		    #add nodes to container
+		    lagContainer = cmds.container(name = (name + "_lag_container"))
+		    for node in frameCacheNodes:
+			cmds.container(lagContainer, edit = True, addNode = node, includeNetwork = True, ihb = True)
+			
+		    #constrain root of fk chain to driver's parent joint
+		    parent = cmds.listRelatives("driver_" + name + "_01", parent = True)[0]
+		    cmds.parentConstraint(parent, rootControl, mo = True)
+		    
+			
+			
+			
+			
+			
+			
+		    #@@@@@@@@@@@@@@@@@@@@@@@@@@@@@@@@@@@@@@@@@@@@@@@@@@@@@@@@@@@@@@@@@@@@@@@@@@@@@@@@@@@@@@@@@@@@@@@@@@@#
+		    #IK RIG
+		    #@@@@@@@@@@@@@@@@@@@@@@@@@@@@@@@@@@@@@@@@@@@@@@@@@@@@@@@@@@@@@@@@@@@@@@@@@@@@@@@@@@@@@@@@@@@@@@@@@@@#
+		    
+
+		    ikJoints = []
+		    clusterControls = []
+		    for i in range(int(numJointsInChain)):
+			jointNum = i + 1
+			
+			#create and position the joint
+			if cmds.objExists("rig_ik_" + name + "_0" + str(jointNum)):
+			    cmds.delete("rig_ik_" + name + "_0" + str(jointNum))
+			    
+			    
+			cmds.select(clear = True)
+			joint = cmds.joint(name = "rig_ik_" + name + "_0" + str(jointNum))
+			cmds.select(clear = True)
+			ikJoints.append(joint)
+			
+			constraint = cmds.parentConstraint("driver_" + name + "_0" + str(jointNum), joint)[0]
+			cmds.delete(constraint)
+			
+			#recreate the joint heirarchy
+			if i != 0:
+			    cmds.parent(joint, lastJoint)
+			    
+			lastJoint = joint
+			
+		    
+
+		    startJoint = ikJoints[0]
+		    endJoint = ikJoints[(len(ikJoints) - 1)]
+		    cmds.makeIdentity(startJoint, r = 1, t = 0, s = 0, apply = True)
+		    
+		    #create the spline IK
+		    ikNodes = cmds.ikHandle(sj = startJoint, ee = endJoint, sol = "ikSplineSolver", createCurve = True, simplifyCurve = False, parentCurve = False, name = str(ikJoints[0]) + "_splineIK")   
+		    ikHandle = ikNodes[0]
+		    ikCurve = ikNodes[2]
+		    ikCurve = cmds.rename(ikCurve, name + "_splineIK_curve")
+		    cmds.setAttr(ikCurve + ".inheritsTransform", 0)
+		    cmds.setAttr(ikHandle + ".v", 0)
+		    cmds.setAttr(ikCurve + ".v", 0)
+		    
+		    #create the three joints to skin the curve to
+		    
+		    if int(numJointsInChain) <= 6:
+			#3 joints for the curve
+			skinJoints = 3
+			
+		    if int(numJointsInChain) >= 7:
+			#add 1 joint to the curve every odd number
+			oddJoints = []
+			for i in range(7, int(numJointsInChain)):
+			    if i % 2 != 0:
+				oddJoints.append(i)
+				    
+			#now we have a list of the total number of odd joints in our numJointsInChain. Take the length of the list + 3 to get the joints to create for our curve
+			skinJoints = 3 + len(oddJoints)
+
+			
+		    #create the joints to skin to the curve
+		    curveJoints = []
+		    if skinJoints == 3:
+			botJoint = cmds.duplicate(startJoint, name = name + "_splineIK_skin_joint_1", parentOnly = True)[0]
+			topJoint = cmds.duplicate(endJoint, name = name + "_splineIK_skin_joint_2", parentOnly = True)[0]
+			midJoint = cmds.duplicate(topJoint, name = name + "_splineIK_skin_joint_3", parentOnly = True)[0]
+			
+			cmds.parent([botJoint, topJoint,midJoint], world = True)
+			
+			constraint = cmds.pointConstraint([botJoint, topJoint], midJoint)[0]
+			cmds.delete(constraint)
+			curveJoints.append(botJoint)
+			curveJoints.append(topJoint)
+			curveJoints.append(midJoint)
+
+		    else:
+			for i in range(skinJoints):
+			    if i == 0:
+				joint = cmds.duplicate(ikJoints[i], name = name + "_splineIK_skin_joint_" + str(i), parentOnly = True)[0]
+				curveJoints.append(joint)
+				
+			    else:
+				joint = cmds.duplicate(ikJoints[i + i], name = name + "_splineIK_skin_joint_" + str(i), parentOnly = True)[0]
+				curveJoints.append(joint)				
+
+				
+			#parent all of the joints to the world
+			for joint in curveJoints:
+			    try:
+				cmds.parent(joint, world = True)
+			    except:
+				print joint
+				pass
+			    
+			    
+
+		    
+		    #skin the joints to the curve
+		    
+		    cmds.select(curveJoints)
+		    cmds.skinCluster( curveJoints, ikCurve, toSelectedBones = True )
+		    
+		    #find number of CVs on created curve
+		    numSpans = cmds.getAttr(ikCurve + ".spans")
+		    degree = cmds.getAttr(ikCurve + ".degree")
+		    numCVs = numSpans + degree
+		    
+		    #for each cv, create a cluster, then create the control
+		    clusters = []
+		    for cv in range(int(numCVs)):
+			cmds.select(ikCurve + ".cv[" + str(cv) + "]" )
+			cluster = cmds.cluster(name = name + "_cluster_" + str(cv))
+			clusters.append(cluster)
+			
+			
+		    #cleanup clusters list
+		    cmds.delete(clusters[1])
+		    cmds.delete(clusters[(len(clusters) - 2)])
+		    clusters.pop(1)
+		    clusters.pop(len(clusters) - 2)
+		    clusterNodes = []
+		    ikAnimGrps = []
+		    
+		    #create the controls for each cluster
+		    for i in range(int(len(clusters))):
+			
+			cluster = cmds.rename(clusters[i][1], name + "_cluster_" + str(i))
+			clusterNodes.append(cluster)
+			
+			cmds.setAttr(cluster + ".v", 0)
+			
+			control = cmds.spaceLocator(name = name + "_cv_" + str(i) + "_anim")[0]
+			constraint = cmds.parentConstraint(ikJoints[i], control)[0]
+			cmds.delete(constraint)
+			clusterControls.append(control)
+			
+			
+			#scale up the locator
+			scaleFactor = self.getScaleFactor()
+			shape = cmds.listRelatives(control, shapes = True)[0]
+			cmds.setAttr(shape + ".localScaleX", 15 * scaleFactor)
+			cmds.setAttr(shape + ".localScaleY", 15 * scaleFactor)
+			cmds.setAttr(shape + ".localScaleZ", 15 * scaleFactor)
+			
+			
+			ctrlGrp = cmds.group(empty = True, name = control + "_grp")
+			ikAnimGrps.append(ctrlGrp)
+			
+			constraint = cmds.pointConstraint(ikJoints[i], ctrlGrp)[0]
+			cmds.delete(constraint)
+			
+			cmds.parent(control, ctrlGrp)
+			cmds.makeIdentity(control, t = 1, r = 1, s = 1, apply = True)
+			
+			#point constrain the ctrlGrp to the corresponding joint
+			cmds.pointConstraint(ikJoints[i], ctrlGrp)
+			
+			#connect the clusters translate to the control's so the cluster will move when the control does
+			cmds.connectAttr(control + ".translate", cluster + ".translate")
+			
+			#color controls
+			cmds.setAttr(control + ".overrideEnabled", 1)
+			cmds.setAttr(control + ".overrideColor", 18)
+			
+			
+
+		    #hookup stretch to joint scale
+		    cmds.select(ikCurve)
+		    curveInfoNode = cmds.arclen(cmds.ls(sl = True), ch = True )
+		    originalLength = cmds.getAttr(curveInfoNode + ".arcLength")
+		    
+		    #create the multiply/divide node that will get the scale factor
+		    divideNode = cmds.shadingNode("multiplyDivide", asUtility = True, name = name + "_divideNode")
+		    cmds.setAttr(divideNode + ".operation", 2)
+		    cmds.setAttr(divideNode + ".input2X", originalLength)
+		    
+		    #create the blendcolors node
+		    blenderNode = cmds.shadingNode("blendColors", asUtility = True, name = name + "_blenderNode")
+		    cmds.setAttr(blenderNode + ".color2R", 1)
+		    
+		    #connect attrs
+		    cmds.connectAttr(curveInfoNode + ".arcLength", divideNode + ".input1X")
+		    
+		    
+		    for joint in ikJoints:
+			cmds.connectAttr(divideNode + ".outputX", joint + ".scaleX")
+			cmds.connectAttr(divideNode + ".outputX", joint + ".scaleY")
+			cmds.connectAttr(divideNode + ".outputX", joint + ".scaleZ")
+			
+			
+		    #create the control curves for the ik curve joints
+		    i = 1
+		    ikControls = []
+		    for joint in curveJoints:
+			if joint != curveJoints[0]:
+			    control = self.createControl("circle", 25, name + "_ik_" + str(i) + "_anim")
+			    cmds.setAttr(control + ".rx", 90)
+			    cmds.makeIdentity(control, t = 1, r = 1, s = 1, apply = True)
+			    ikControls.append(control)
+			    
+			    #position
+			    constraint = cmds.parentConstraint(joint, control)[0]
+			    cmds.delete(constraint)
+			    
+			    #create grp
+			    controlGrp = cmds.group(empty = True, name = control + "_grp")
+			    constraint = cmds.parentConstraint(joint, controlGrp)[0]
+			    cmds.delete(constraint)
+			    
+			    #setup hierarchy
+			    cmds.parent(control, controlGrp)
+			    cmds.parent(joint, control)
+			    
+			    #color controls
+			    cmds.setAttr(control + ".overrideEnabled", 1)
+			    cmds.setAttr(control + ".overrideColor", 18)
+			    
+			    
+			    #lock attrs
+			    cmds.setAttr(control + ".sx", lock = True, keyable = False)
+			    cmds.setAttr(control + ".sy", lock = True, keyable = False)
+			    cmds.setAttr(control + ".sz", lock = True, keyable = False)
+			    cmds.setAttr(control + ".v", lock = True, keyable = False)			    
+			    
+			else:
+			    #find parent of base joint and constrain base joint to the parent
+			    parent = cmds.listRelatives("driver_" + name + "_01", parent = True)[0]
+			    
+			    
+			    #create a control for the base
+			    control = self.createControl("circle", 30, name + "_ik_base_anim")
+			    cmds.setAttr(control + ".rx", 90)
+			    cmds.makeIdentity(control, t = 1, r = 1, s = 1, apply = True)
+			    
+			    #position
+			    constraint = cmds.parentConstraint(joint, control)[0]
+			    cmds.delete(constraint)
+			    
+			    #create grp
+			    controlGrp = cmds.group(empty = True, name = control + "_grp")
+			    constraint = cmds.parentConstraint(joint, controlGrp)[0]
+			    cmds.delete(constraint)
+			    
+			    #setup hierarchy
+			    cmds.parent(control, controlGrp)
+			    cmds.parent(joint, control)
+			    
+			    #color controls
+			    cmds.setAttr(control + ".overrideEnabled", 1)
+			    cmds.setAttr(control + ".overrideColor", 17)
+			    
+			    
+			    #lock attrs
+			    cmds.aliasAttr("global_scale", control + ".scaleZ")
+			    cmds.connectAttr(control + ".scaleZ", control + ".scaleX")
+			    cmds.connectAttr(control + ".scaleZ", control + ".scaleY")
+			    
+			    cmds.setAttr(control + ".sx", lock = True, keyable = False)
+			    cmds.setAttr(control + ".sy", lock = True, keyable = False)
+			    cmds.setAttr(control + ".v", lock = True, keyable = False)
+			    
+			    
+			    #hook the base control grp to the chain's parent
+			    cmds.parentConstraint(parent, controlGrp, mo = True)
+			    
+			    
+			    
+			i = i + 1
+			
+		    #parent the other IK grp controls under the base
+		    for control in ikControls:
+			grp = control + "_grp"
+			cmds.parent(grp, name + "_ik_base_anim")
+			
+			
+		    #tip control only:
+		    #add attr to show clusters on tip control
+		    tipControl = ikControls[len(ikControls) - 1]
+		    cmds.select(tipControl)
+		    cmds.addAttr(longName=("clusterControlVis"), at = 'bool', dv = 0, keyable = True)
+			    
+		    for control in clusterControls:
+			cmds.connectAttr(tipControl + ".clusterControlVis", control + ".v")
+			cmds.setAttr(control + ".sx", lock = True, keyable = False)
+			cmds.setAttr(control + ".sy", lock = True, keyable = False)
+			cmds.setAttr(control + ".sz", lock = True, keyable = False)
+			cmds.setAttr(control + ".v", lock = True, keyable = False)
+			cmds.setAttr(control + ".rx", lock = True, keyable = False)
+			cmds.setAttr(control + ".ry", lock = True, keyable = False)
+			cmds.setAttr(control + ".rz", lock = True, keyable = False)
+			
+		    #set color for tip
+		    cmds.setAttr(tipControl + ".overrideColor", 17)
+		    
+		    #lock tip attrs
+		    cmds.setAttr(tipControl + ".sx", lock = True, keyable = False)
+		    cmds.setAttr(tipControl + ".sy", lock = True, keyable = False)
+		    cmds.setAttr(tipControl + ".sz", lock = True, keyable = False)
+		    cmds.setAttr(tipControl + ".v", lock = True, keyable = False)		    
+			    
+
+			
+		    #@@@@@@@@@@@@@@@@@@@@@@@@@@@@@@@@@@@@@@@@@@@@@@@@@@@@@@@@@@@@@@@@@@@@@@@@@@@@@@@@@@@@@@@@@@@@@@@@@@@#
+		    #Dynamics RIG
+		    #@@@@@@@@@@@@@@@@@@@@@@@@@@@@@@@@@@@@@@@@@@@@@@@@@@@@@@@@@@@@@@@@@@@@@@@@@@@@@@@@@@@@@@@@@@@@@@@@@@@#
+		    dynamicJoints = []
+		    
+		    for i in range(int(numJointsInChain)):
+			jointNum = i + 1
+			
+			#create and position the joint
+			if cmds.objExists("rig_dyn_" + name + "_0" + str(jointNum)):
+			    cmds.delete("rig_dyn_" + name + "_0" + str(jointNum))
+			    
+			    
+			cmds.select(clear = True)
+			joint = cmds.joint(name = "rig_dyn_" + name + "_0" + str(jointNum))
+			if joint.find("|") == 0:
+			    joint = joint.partition("|")[2]
+			    
+			cmds.select(clear = True)
+			dynamicJoints.append(joint)
+			
+			constraint = cmds.parentConstraint("driver_" + name + "_0" + str(jointNum), joint)[0]
+			cmds.delete(constraint)
+			
+			#recreate the joint heirarchy
+			if i != 0:
+			    cmds.parent(joint, lastJoint)
+			    
+			
+			lastJoint = joint
+			
+			
+		    #freeze rotations on joints
+		    cmds.makeIdentity(dynamicJoints[0], t = False, r = True, scale = False, apply = True)
+		    
+		    
+		    #Create curve on joint chain
+		    positions = []
+		    
+		    #get the world space positions of each joint, and create a curve using those positions
+		    for i in range(int(len(dynamicJoints))):
+			pos = cmds.xform(dynamicJoints[i], q = True, ws = True, t = True)
+			positions.append(pos)
+			
+		    createCurveCommand = "curve -d 1"
+		    for pos in positions:
+			xPos = pos[0]
+			yPos = pos[1]
+			zPos = pos[2]
+			createCurveCommand += " -p " + str(xPos) + " " + str(yPos) + " " + str(zPos)
+			
+		    for i in range(int(len(positions))):
+			createCurveCommand += " -k " + str(i)
+			
+		    curve = mel.eval(createCurveCommand)
+		    curve = cmds.rename(curve, name + "_dynCurve")
+		    cmds.setAttr(curve + ".v", 0)
+		    
+		    
+		    
+		    #Create hair system
+		    cmds.select(curve)
+		    
+		    #find all hair systems in scene
+		    hairSystems = cmds.ls(type = "hairSystem")
+		    hairSys = ""
+		    
+		    
+		    #create the hair system and make the stiffness uniform
+		    madeHairCurve = True
+		    if hairSys == "":
+			hairSys = cmds.createNode("hairSystem")
+			cmds.removeMultiInstance(hairSys + ".stiffnessScale[1]", b = True)
+			cmds.setAttr(hairSys + ".clumpWidth", 0.0)
+			cmds.connectAttr("time1.outTime", hairSys + ".currentTime")
+			
+		    hairSysParent = cmds.listRelatives(hairSys, parent = True)
+		    hairSysParent = cmds.rename(hairSysParent, name + "_hairSystem")
+		    cmds.setAttr(hairSysParent + ".v", 0)
+		    hairSys =  name + "_hairSystemShape"
+		    
+		    
+		    #create the hair follicle
+		    hair = cmds.createNode("follicle")
+		    cmds.setAttr(hair + ".parameterU", 0)
+		    cmds.setAttr(hair + ".parameterV", 0)
+		    hairTransforms = cmds.listRelatives(hair, p = True)
+		    hairDag = hairTransforms[0]
+		    hairDag = cmds.rename(hairDag, name + "_follicle")
+		    hair = name + "_follicleShape"
+		    
+		    cmds.setAttr(hairDag + ".v", 0)
+		    cmds.setAttr(hair + ".startDirection", 1)
+		    
+		    #get the curve CVs and set follicle degree to 1 if CVs are less than 3
+		    curveCVs = cmds.getAttr(curve + ".cp", size = True)
+		    if curveCVs < 3:
+			cmds.setAttr(hair + ".degree", 1)
+			
+		    #parent the curve to the follicle and connect the curve's worldspace[0] to the follicle startPos
+		    cmds.parent(curve, hairDag, relative = True)
+		    cmds.connectAttr(curve + ".worldSpace[0]", hair + ".startPosition")
+		    
+		    #connect the hair follicle to the hair system
+		    cmds.connectAttr(hair + ".outHair", hairSys + ".inputHair[0]")
+		    
+		    #create a new curve and connect the follicle's outCurve attr to the new curve
+		    cmds.connectAttr(hairSys + ".outputHair[0]", hair + ".currentPosition")
+		    crv = cmds.createNode("nurbsCurve")
+		    crvParent = cmds.listRelatives(crv, parent = True)[0]
+		    crvParent = cmds.rename(crvParent, name + "_track_rt_curve")
+		    crv =  name + "_track_rt_curveShape"
+		    cmds.setAttr(crvParent + ".v", 0)
+		    
+		    cmds.connectAttr(hair + ".outCurve", crv + ".create")
+		    
+		    #set the hair follicle attrs 
+		    if len(hairDag) > 0:
+			cmds.setAttr(hairDag + ".pointLock", 3)
+			cmds.setAttr(hairDag + ".restPose", 1)
+			
+		    cmds.select(hairSys)
+		    
+		    
+		    
+		    #Create Spline Handle for the selected chain and the duplicated curve. the original is driven by hair
+		    ikNodes = cmds.ikHandle(sol = "ikSplineSolver",  ccv = False, pcv = False, snc = True, rootTwistMode = False, sj = dynamicJoints[0], ee = dynamicJoints[len(dynamicJoints) - 1], c = crv)[0]
+		    cmds.setAttr(ikNodes + ".v", 0)
+		    ikNodes = cmds.rename(ikNodes, name + "_dynChain_ikHandle")
+
+		    
+		    #Create a duplicate joint chain for manual animation
+		    dupeChain = cmds.duplicate(dynamicJoints[0], rr = True, rc = True)
+		    dupeStartJoint = dupeChain[0]
+		    
+		    dupeJoints = cmds.listRelatives(dupeStartJoint, ad = True)
+		    joints = cmds.listRelatives(dynamicJoints[0], ad = True)
+		    
+		    #rename duped joints and connect real joints to duped joints
+		    for i in range(int(len(joints))):
+			if cmds.objectType(dupeJoints[i], isType = 'joint'):
+			    cmds.rename(dupeJoints[i], "ANIM_" + joints[i])
+			    cmds.connectAttr("ANIM_" + joints[i] + ".r", joints[i] + ".r", force = True)
+			    
+			else:
+			    cmds.delete(dupeJoints[i])
+			    
+		    #connect up  start joint to ANIM start joint
+		    cmds.connectAttr(dupeStartJoint + ".t", dynamicJoints[0] + ".t")
+		    cmds.connectAttr(dupeStartJoint + ".r", dynamicJoints[0] + ".r")
+		    cmds.connectAttr(dupeStartJoint + ".s", dynamicJoints[0] + ".s")
+		    
+		    dupeStartJoint = cmds.rename(dupeStartJoint, "ANIM_" + dynamicJoints[0])
+		    #cmds.parent(dupeStartJoint, world = True)
+		    
+		    #Create skinCluster between duplicate curve and animation joint chain(dupe chain)
+		    cmds.select(dupeStartJoint, hi = True)
+		    dupeSkel = cmds.ls(sl = True, type = "joint")
+		    cmds.select(curve)
+		    cmds.select(dupeSkel, add = True)
+		    skinCluster = cmds.skinCluster(tsb = True, mi = 3, dr = 4)
+		    
+		    
+		    #Create the control that has all of our dynamic attrs
+		    control = self.createControl("square", 30, name + "_dyn_anim")
+		    constraint = cmds.parentConstraint(dynamicJoints[0], control)[0]
+		    cmds.delete(constraint)
+		    
+		    cmds.makeIdentity(control, r = 1, t = 0, s = 0, apply = True)
+		    cmds.setAttr(control + ".rx", 90)
+		    
+		    ctrlGrp = cmds.group(empty = True, name = control + "_grp")
+		    constraint = cmds.parentConstraint(dynamicJoints[0], ctrlGrp)[0]
+		    cmds.delete(constraint)
+		    
+		    cmds.parent(control, ctrlGrp)
+		    cmds.makeIdentity(control, r = 1, t = 1, s = 1, apply = True)
+		    cmds.parentConstraint(control, dupeStartJoint)
+		    
+		    #lock attrs
+		    cmds.setAttr(control + ".sx", lock = True, keyable = False)
+		    cmds.setAttr(control + ".sy", lock = True, keyable = False)
+		    cmds.setAttr(control + ".sz", lock = True, keyable = False)
+		    cmds.setAttr(control + ".v", lock = True, keyable = False)
+		    
+		    
+		    
+		    
+		    #add attrs
+		    cmds.select(control)
+		    
+		    cmds.addAttr(ln = "___DYNAMICS___", at = "double", keyable = True)
+		    cmds.setAttr(control + ".___DYNAMICS___", lock = True)
+		    cmds.addAttr(ln = "chainAttach", at = "enum", en = "No Attach:Base:Tip:Both End:", dv = 1, keyable = True)
+		    cmds.addAttr(ln = "chainStartEnvelope", at = "double", min = 0, max = 1, dv = 1, keyable = True)
+		    cmds.addAttr(ln = "chainStartFrame", at = "double", dv = 1, keyable = True)
+		    
+		    cmds.addAttr(ln = "___BEHAVIOR___", at = "double", keyable = True)
+		    cmds.setAttr(control + ".___BEHAVIOR___", lock = True)
+		    cmds.addAttr(ln = "chainStiffness", at = "double", min = 0, dv = .1, keyable = True)
+		    cmds.addAttr(ln = "chainDamping", at = "double", min = 0, dv = 0.2, keyable = True)
+		    cmds.addAttr(ln = "chainGravity", at = "double", min = 0, dv = 1, keyable = True)
+		    cmds.addAttr(ln = "chainIteration", at = "long", min = 0, dv = 1, keyable = True)
+		    
+		    cmds.addAttr(ln = "___COLLISIONS___", at = "double", keyable = True)
+		    cmds.setAttr(control + ".___COLLISIONS___", lock = True)
+		    cmds.addAttr(ln = "chainCollide", at = "bool", dv = 0, keyable = True)
+		    cmds.addAttr(ln = "chainWidthBase", at = "double", min = 0, dv = 1, keyable = True)
+		    cmds.addAttr(ln = "chainWidthExtremity", at = "double", min = 0, dv = 1, keyable = True)
+		    cmds.addAttr(ln = "chainCollideGround", at = "bool", dv = 0, keyable = True)
+		    cmds.addAttr(ln = "chainCollideGroundHeight", at = "double", dv = 0, keyable = True)
+		    
+		    
+		    
+		    
+		    #connect attrs
+		    cmds.connectAttr(control + ".chainStartEnvelope", ikNodes + ".ikBlend")
+		    cmds.connectAttr(control + ".chainAttach", hair + ".pointLock")
+		    cmds.connectAttr(control + ".chainStartFrame", hairSys + ".startFrame")
+		    cmds.connectAttr(control + ".chainStiffness", hairSys + ".stiffness")
+		    cmds.connectAttr(control + ".chainDamping", hairSys + ".damp")
+		    cmds.connectAttr(control + ".chainGravity", hairSys + ".gravity")
+		    cmds.connectAttr(control + ".chainIteration", hairSys + ".iterations")
+		    cmds.connectAttr(control + ".chainCollide", hairSys + ".collide")
+		    cmds.connectAttr(control + ".chainWidthBase", hairSys + ".clumpWidth")
+		    cmds.connectAttr(control + ".chainWidthExtremity", hairSys + ".clumpWidthScale[1].clumpWidthScale_FloatValue")
+		    cmds.connectAttr(control + ".chainCollideGround", hairSys + ".collideGround")
+		    cmds.connectAttr(control + ".chainCollideGroundHeight", hairSys + ".groundHeight")
+		    
+		    
+		    #Create the expression for real time
+		    track_RealTime = cmds.spaceLocator(name = name + "_track_rt_loc")[0]
+		    cmds.pointConstraint(dupeSkel[len(dupeSkel) -1], track_RealTime)
+		    connections = cmds.listConnections(hairSys + ".currentTime", p = True, c = True)
+		    cmds.disconnectAttr(connections[1], connections[0])
+		    
+		    expressionString = "if(frame!= " + hairSys + ".startFrame)\n\t" + hairSys + ".currentTime = " + hairSys + ".currentTime + 1 + " + track_RealTime + ".tx - " + track_RealTime + ".tx + " + track_RealTime + ".ty - " + track_RealTime + ".ty + " + track_RealTime + ".tz - " + track_RealTime + ".tz + " + control + ".chainWidthBase - "+ control + ".chainWidthBase + "+ control + ".chainWidthExtremity - "+ control + ".chainWidthExtremity + " + control + ".chainGravity - "+ control + ".chainGravity;\n" +"else\n\t" + hairSys + ".currentTime = " + hairSys + ".startFrame;"
+		    cmds.expression(name = "EXP_" + hairSys + "_TRACK_RealTime", string = expressionString)
+		    cmds.setAttr(track_RealTime + ".v", 0)
+		    
+		    #Set Defaults
+		    cmds.setAttr(hairSys + ".drawCollideWidth", 1)
+		    cmds.setAttr(hairSys + ".widthDrawSkip", 0)
+		    cmds.setAttr(hair + ".degree", 1)
+		    
+		    cmds.parentConstraint(parent, ctrlGrp, mo = True)
+		    
+		    if cmds.objExists("dynHairChain") == False:
+			cmds.group(empty = True, name = "dynHairChain")
+			
+		    if cmds.objExists(dynamicJoints[0] + "_HairControls") == False:
+			group = cmds.group([ikNodes, hairSys, hair, track_RealTime, ctrlGrp, crvParent], name = dynamicJoints[0] + "_HairControls")
+			
+		    else:
+			cmds.parent([ikNodes, hairSys, hair, ctrlGrp,crvParent ], dynamicJoints[0] + "_HairControls")
+			
+		    cmds.parent(group, "dynHairChain")
+		    
+		    #Cleanup
+		    jointsGrp = cmds.group(empty = True, name = name + "_jointGrp")
+		    cmds.parent([dynamicJoints[0], dupeStartJoint], jointsGrp)
+		    
+		    cmds.setAttr(control + ".overrideEnabled", 1)
+		    cmds.setAttr(control + ".overrideColor", 18)
+		    
+		    
+		    #@@@@@@@@@@@@@@@@@@@@@@@@@@@@@@@@@@@@@@@@@@@@@@@@@@@@@@@@@@@@@@@@@@@@@@@@@@@@@@@@@@@@@@@@@@@@@@@@@@@#
+		    #CLEAN UP SCENE
+		    #@@@@@@@@@@@@@@@@@@@@@@@@@@@@@@@@@@@@@@@@@@@@@@@@@@@@@@@@@@@@@@@@@@@@@@@@@@@@@@@@@@@@@@@@@@@@@@@@@@@#
+		    ikGrp = cmds.group(empty = True, name = name + "_ik_ctrl_grp")
+		    clustersGrp = cmds.group(empty = True, name = name + "_ik_clusters_grp")
+		    dynGrp = cmds.group(empty = True, name = name + "_dyn_ctrl_grp")
+		    masterGrp = cmds.group(empty = True, name = name + "_master_ctrl_grp")
+		    createdControls.append(masterGrp)
+		    
+		    #need to parent control groups in here
+		    cmds.parent([ikHandle, ikJoints[0], ikCurve, name + "_ik_base_anim_grp"], ikGrp)
+		    
+		    for cluster in clusterNodes:
+			cmds.parent(cluster, clustersGrp)
+			
+		    for grp in ikAnimGrps:
+			cmds.parent(grp, ikGrp)
+			
+		    cmds.parent(clustersGrp, "master_anim")
+		    cmds.setAttr(clustersGrp + ".inheritsTransform", 0)
+		    
+		    cmds.parent(name + "_jointGrp", dynGrp)
+		    cmds.parent([ikGrp, dynGrp, fkRootGrp], masterGrp)
+		    
+
+		    
+		    #@@@@@@@@@@@@@@@@@@@@@@@@@@@@@@@@@@@@@@@@@@@@@@@@@@@@@@@@@@@@@@@@@@@@@@@@@@@@@@@@@@@@@@@@@@@@@@@@@@@#
+		    #HOOKUP RIGS TO RIG SETTINGS
+		    #@@@@@@@@@@@@@@@@@@@@@@@@@@@@@@@@@@@@@@@@@@@@@@@@@@@@@@@@@@@@@@@@@@@@@@@@@@@@@@@@@@@@@@@@@@@@@@@@@@@#
+		    cmds.select("Rig_Settings")
+		    cmds.addAttr(longName= (name + "_fk"), defaultValue=1, minValue=0, maxValue=1, keyable = True)
+		    cmds.addAttr(longName= (name + "_ik"), defaultValue=0, minValue=0, maxValue=1, keyable = True)
+		    cmds.addAttr(longName= (name + "_dynamic"), defaultValue=0, minValue=0, maxValue=1, keyable = True)
+
+		    for i in range(int(len(dynamicJoints))):
+			driverJoint = dynamicJoints[i].replace("rig_dyn_", "driver_")
+			
+			constraint = cmds.parentConstraint([fkJoints[i], ikJoints[i], dynamicJoints[i]], driverJoint)[0]
+			
+			
+			cmds.connectAttr("Rig_Settings." + name + "_fk", constraint + "." + fkJoints[i] + "W0")
+			cmds.connectAttr("Rig_Settings." + name + "_ik", constraint + "." + ikJoints[i] + "W1")
+			cmds.connectAttr("Rig_Settings." + name + "_dynamic", constraint + "." + dynamicJoints[i] + "W2")
+			
+			#create blend Color nodes for scale
+			scaleBlendColors = cmds.shadingNode("blendColors", asUtility = True, name = name + "_scale_blend")
+			cmds.connectAttr(firstControl + ".scale", scaleBlendColors + ".color1")
+			cmds.connectAttr(name + "_ik_base_anim" + ".scale", scaleBlendColors + ".color2")
+			cmds.connectAttr(scaleBlendColors + ".output", driverJoint + ".scale")
+			
+			cmds.connectAttr("Rig_Settings." + name + "_fk", scaleBlendColors + ".blender")
+			
+		
+			
+
+		    #setup visibility connections
+		    cmds.connectAttr("Rig_Settings." + name + "_fk", fkRootGrp + ".v")
+		    cmds.connectAttr("Rig_Settings." + name + "_ik", ikGrp + ".v")
+		    cmds.connectAttr("Rig_Settings." + name + "_dynamic", name + "_dyn_anim_grp.v")
+			
+
+		    
+	return createdControls
+			
+
+    # # # # # # # # # # # # # # # # # # # # # # # # # # # # # # # # # # # # # # # # # # # # # # # # # # # # # # # # # # # # # # # 
+    # # # # # # # # # # # # # # # # # # # # # # # # # # # # # # # # # # # # # # # # # # # # # # # # # # # # # # # # # # # # # # # 
+    # # # # # # # # # # # # # # # # # # # # # # # # # # # # # # # # # # # # # # # # # # # # # # # # # # # # # # # # # # # # # # # 
+    # # # # # # # # # # # # # # # # # # # # # # # # # # # # # # # # # # # # # # # # # # # # # # # # # # # # # # # # # # # # # # # 
+    # # # # # # # # # # # # # # # # # # # # # # # # # # # # # # # # # # # # # # # # # # # # # # # # # # # # # # # # # # # # # # # 
+    def createDriverSkeleton(self):
+        
+        dupe = cmds.duplicate("root", rc = True)[0]
+        cmds.select("root", hi = True)
+        joints = cmds.ls(sl = True)
+        
+        cmds.select(dupe, hi = True)
+        dupeJoints = cmds.ls(sl = True)
+        
+        driverJoints = []
+        for i in range(int(len(dupeJoints))):
+            
+            if cmds.objExists(dupeJoints[i]):
+                driverJoint = cmds.rename(dupeJoints[i], "driver_" + joints[i])
+                driverJoints.append(driverJoint)
+                
+
+        #create a direct connection between the driver and the export joints
+	exceptionJoints = ["upperarm_l", "upperarm_r"]
+	
+        for joint in driverJoints:
+            exportJoint = joint.partition("_")[2]
+	    
+	    if exportJoint not in exceptionJoints:
+		cmds.connectAttr(joint + ".translate", exportJoint + ".translate")
+		cmds.connectAttr(joint + ".rotate", exportJoint + ".rotate")
+		cmds.connectAttr(joint + ".scale", exportJoint + ".scale")
+	    else:
+		cmds.connectAttr(joint + ".translate", exportJoint + ".translate")
+		cmds.connectAttr(joint + ".scale", exportJoint + ".scale")	
+		cmds.orientConstraint(joint, exportJoint)
+		
+
+
+    # # # # # # # # # # # # # # # # # # # # # # # # # # # # # # # # # # # # # # # # # # # # # # # # # # # # # # # # # # # # # # # 
+    # # # # # # # # # # # # # # # # # # # # # # # # # # # # # # # # # # # # # # # # # # # # # # # # # # # # # # # # # # # # # # # 
+    # # # # # # # # # # # # # # # # # # # # # # # # # # # # # # # # # # # # # # # # # # # # # # # # # # # # # # # # # # # # # # # 
+    # # # # # # # # # # # # # # # # # # # # # # # # # # # # # # # # # # # # # # # # # # # # # # # # # # # # # # # # # # # # # # # 
+    # # # # # # # # # # # # # # # # # # # # # # # # # # # # # # # # # # # # # # # # # # # # # # # # # # # # # # # # # # # # # # # 
+    def getSpineJoints(self):
+        
+        numSpineBones = int(cmds.getAttr("Skeleton_Settings.numSpineBones"))
+        spineJoints = []
+        
+        for i in range(int(numSpineBones)):
+            
+            if i < 10:
+                spineJoint = "spine_0" + str(i + 1)
+                
+            else:
+                spineJoint = "spine_" + (i + 1)
+                
+            spineJoints.append(spineJoint)
+            
+        return spineJoints
+    
+    
+    # # # # # # # # # # # # # # # # # # # # # # # # # # # # # # # # # # # # # # # # # # # # # # # # # # # # # # # # # # # # # # # 
+    # # # # # # # # # # # # # # # # # # # # # # # # # # # # # # # # # # # # # # # # # # # # # # # # # # # # # # # # # # # # # # # 
+    # # # # # # # # # # # # # # # # # # # # # # # # # # # # # # # # # # # # # # # # # # # # # # # # # # # # # # # # # # # # # # # 
+    # # # # # # # # # # # # # # # # # # # # # # # # # # # # # # # # # # # # # # # # # # # # # # # # # # # # # # # # # # # # # # # 
+    # # # # # # # # # # # # # # # # # # # # # # # # # # # # # # # # # # # # # # # # # # # # # # # # # # # # # # # # # # # # # # # 
+    def createControl(self, controlType, size, name):
+        
+        scale = self.getScaleFactor()
+        
+        if controlType == "circle":
+            control = cmds.circle(c = (0,0,0), sw = 360, r = size * scale, d = 3, name = name)[0]
+            
+        if controlType == "circleSpecial":
+            control = cmds.circle(c = (0,0,0), sw = 360, r = 1, d = 3, name = name)[0]
+            
+            side = name.rpartition("_")[2]
+            
+            if side == "l":
+                cmds.xform(control, piv = (0, -1, 0))
+                
+            else:
+                cmds.xform(control, piv = (0, 1, 0))
+                
+            cmds.setAttr(control + ".scaleX", size * scale)
+            cmds.setAttr(control + ".scaleY", size * scale)
+            cmds.setAttr(control + ".scaleZ", size * scale)
+
+        if controlType == "square":
+            control = cmds.circle(c = (0,0,0), s = 4, sw = 360, r = size * scale, d = 1, name = name)[0]
+            cmds.setAttr(control + ".rz", 45)
+            
+        if controlType == "foot":
+            control = cmds.curve(name = name, d = 3, p = [(0, 40, 0), (-3.42, 39, 0), (-10.2, 37, 0), (-13, 22, 0), (-15.7, 13.2, 0), (-20, -14, 0), (-18.1, -25.6, 0), (-15, -44.8, 0), (1.1, -41.2, 0), (4.8, -41.7, 0), (15.5, -31.9, 0), (16.9, -22.7, 0), (18.6, -15.2, 0), (16.5, -.5, 0), (11.2, 29.2, 0), (10.7, 39.7, 0), (3.6, 39.9, 0), (0, 40, 0)])
+            
+            footLoc = cmds.spaceLocator(name = (name + "_end_loc"))[0]
+            cmds.parent(footLoc, control)
+            cmds.setAttr(footLoc + ".ty", -40)
+            cmds.setAttr(footLoc + ".v", 0)
+                         
+            cmds.setAttr(control + ".scaleX", size * scale)
+            cmds.setAttr(control + ".scaleY", size * scale)
+            cmds.setAttr(control + ".scaleZ", size * scale)
+            
+        
+        if controlType == "arrow":
+            control = cmds.curve(name = name, d = 1, p = [(0, -45, 0), (5, -45, 0), (5, -62, 0), (10, -62, 0), (0, -72, 0), (-10, -62, 0), (-5, -62, 0), (-5, -45, 0), (0, -45, 0)])
+            cmds.xform(control, cp = True)
+            cmds.setAttr(control + ".ty", 58.5)
+            cmds.makeIdentity(control, t = 1, apply = True)
+            cmds.xform(control, piv = (0, 13.5, 0))
+            
+            cmds.setAttr(control + ".scaleX", size * scale)
+            cmds.setAttr(control + ".scaleY", size * scale)
+            cmds.setAttr(control + ".scaleZ", size * scale)
+            
+        if controlType == "arrowOnBall":
+            
+            control = cmds.curve(name = name, d = 1, p = [(0.80718, 0.830576, 8.022739), (0.80718, 4.219206, 7.146586 ), (0.80718, 6.317059, 5.70073), (2.830981, 6.317059, 5.70073), (0, 8.422749, 2.94335), (-2.830981, 6.317059, 5.70073), (-0.80718, 6.317059, 5.70073), (-0.80718, 4.219352, 7.146486), (-0.80718, 0.830576, 8.022739), (-4.187851, 0.830576, 7.158003), (-6.310271, 0.830576, 5.705409), (-6.317059, 2.830981, 5.7007), (-8.422749, 0, 2.94335), (-6.317059, -2.830981, 5.70073), (-6.317059, -0.830576, 5.70073), (-4.225134, -0.830576, 7.142501), (-0.827872, -0.830576, 8.017446), (-0.80718, -4.176512, 7.160965), (-0.80718, -6.317059, 5.70073), (-2.830981, -6.317059, 5.70073), (0, -8.422749, 2.94335), (2.830981, -6.317059, 5.70073), (0.80718, -6.317059, 5.70073), (0.80718, -4.21137, 7.151987), (0.80718, -0.830576, 8.022739), (4.183345, -0.830576, 7.159155), (6.317059, -0.830576, 5.70073), (6.317059, -2.830981, 5.70073), (8.422749, 0, 2.94335), (6.317059, 2.830981, 5.70073), (6.317059, 0.830576, 5.70073), (4.263245, 0.830576, 7.116234), (0.80718, 0.830576, 8.022739)])
+            
+            cmds.setAttr(control + ".scaleX", size * scale)
+            cmds.setAttr(control + ".scaleY", size * scale)
+            cmds.setAttr(control + ".scaleZ", size * scale)
+            
+            
+        if controlType == "semiCircle":
+            
+            control = cmds.curve(name = name, d = 3, p = [(0,0,0), (7, 0, 0), (8, 0, 0), (5, 4, 0), (0, 5, 0), (-5, 4, 0), (-8, 0, 0), (-7, 0, 0), (0,0,0)])
+            cmds.xform(control, ws = True, t = (0, 5, 0))
+            cmds.xform(control, ws = True, piv = (0,0,0))
+            cmds.makeIdentity(control, t = 1, apply = True)
+            
+            
+            cmds.setAttr(control + ".scaleX", size * scale)
+            cmds.setAttr(control + ".scaleY", size * scale)
+            cmds.setAttr(control + ".scaleZ", size * scale)
+            
+            
+        if controlType == "pin":
+            control = cmds.curve(name = name, d = 1, p = [(12,0,0), (0, 0, 0), (-12, -12, 0), (-12, 12, 0), (0, 0, 0)])
+            cmds.xform(control, ws = True, piv = [12,0,0])
+            cmds.setAttr(control + ".scaleY", .5)
+            cmds.makeIdentity(control, t = 1, apply = True)
+            
+            
+            cmds.setAttr(control + ".scaleX", size * scale)
+            cmds.setAttr(control + ".scaleY", size * scale)
+            cmds.setAttr(control + ".scaleZ", size * scale)
+            
+        
+        if controlType == "sphere":
+
+            points = [(0, 0, 1), (0, 0.5, 0.866), (0, 0.866025, 0.5), (0, 1, 0), (0, 0.866025, -0.5), (0, 0.5, -0.866025), (0, 0, -1), (0, -0.5, -0.866025), (0, -0.866025, -0.5), (0, -1, 0), (0, -0.866025, 0.5), (0, -0.5, 0.866025), (0, 0, 1), (0.707107, 0, 0.707107), (1, 0, 0), (0.707107, 0, -0.707107), (0, 0, -1), (-0.707107, 0, -0.707107), (-1, 0, 0), (-0.866025, 0.5, 0), (-0.5, 0.866025, 0), (0, 1, 0), (0.5, 0.866025, 0), (0.866025, 0.5, 0), (1, 0, 0), (0.866025, -0.5, 0), (0.5, -0.866025, 0), (0, -1, 0), (-0.5, -0.866025, 0), (-0.866025, -0.5, 0), (-1, 0, 0), (-0.707107, 0, 0.707107), (0, 0, 1)]
+            control = cmds.curve(name = name, d = 1, p = points)
+            
+            cmds.setAttr(control + ".scaleX", size * scale)
+            cmds.setAttr(control + ".scaleY", size * scale)
+            cmds.setAttr(control + ".scaleZ", size * scale)
+            
+
+        return control
+    
+    # # # # # # # # # # # # # # # # # # # # # # # # # # # # # # # # # # # # # # # # # # # # # # # # # # # # # # # # # # # # # # # 
+    # # # # # # # # # # # # # # # # # # # # # # # # # # # # # # # # # # # # # # # # # # # # # # # # # # # # # # # # # # # # # # # 
+    # # # # # # # # # # # # # # # # # # # # # # # # # # # # # # # # # # # # # # # # # # # # # # # # # # # # # # # # # # # # # # # 
+    # # # # # # # # # # # # # # # # # # # # # # # # # # # # # # # # # # # # # # # # # # # # # # # # # # # # # # # # # # # # # # # 
+    # # # # # # # # # # # # # # # # # # # # # # # # # # # # # # # # # # # # # # # # # # # # # # # # # # # # # # # # # # # # # # # 
+    def getScaleFactor(self):
+        
+        headLoc = cmds.spaceLocator(name = "headLoc")[0]
+        cmds.parentConstraint("head", headLoc)
+        
+        height = cmds.getAttr(headLoc + ".tz")
+        
+        defaultHeight = 400
+        scaleFactor = height/defaultHeight
+        cmds.delete(headLoc)
+        
+        return scaleFactor
+    
+    
+    # # # # # # # # # # # # # # # # # # # # # # # # # # # # # # # # # # # # # # # # # # # # # # # # # # # # # # # # # # # # # # # 
+    # # # # # # # # # # # # # # # # # # # # # # # # # # # # # # # # # # # # # # # # # # # # # # # # # # # # # # # # # # # # # # # 
+    # # # # # # # # # # # # # # # # # # # # # # # # # # # # # # # # # # # # # # # # # # # # # # # # # # # # # # # # # # # # # # # 
+    # # # # # # # # # # # # # # # # # # # # # # # # # # # # # # # # # # # # # # # # # # # # # # # # # # # # # # # # # # # # # # # 
+    # # # # # # # # # # # # # # # # # # # # # # # # # # # # # # # # # # # # # # # # # # # # # # # # # # # # # # # # # # # # # # # 
+    def getUpAxis(self, obj):
+        
+        cmds.xform(obj, ws = True, relative = True, t = [0, 0, 10])
+        translate = cmds.getAttr(obj + ".translate")[0]
+        newTuple = (abs(translate[0]), abs(translate[1]), abs(translate[2]))
+        cmds.xform(obj, ws = True, relative = True, t = [0, 0, -10])
+        
+        highestVal = max(newTuple)
+        axis = newTuple.index(highestVal)
+        upAxis = None
+        
+        if axis == 0:
+            upAxis = "X"
+            
+        if axis == 1:
+            upAxis = "Y"
+            
+        if axis == 2:
+            upAxis = "Z"
+            
+        return upAxis
+        
+        
+        
+        
+    # # # # # # # # # # # # # # # # # # # # # # # # # # # # # # # # # # # # # # # # # # # # # # # # # # # # # # # # # # # # # # # 
+    # # # # # # # # # # # # # # # # # # # # # # # # # # # # # # # # # # # # # # # # # # # # # # # # # # # # # # # # # # # # # # # 
+    # # # # # # # # # # # # # # # # # # # # # # # # # # # # # # # # # # # # # # # # # # # # # # # # # # # # # # # # # # # # # # # 
+    # # # # # # # # # # # # # # # # # # # # # # # # # # # # # # # # # # # # # # # # # # # # # # # # # # # # # # # # # # # # # # # 
+    # # # # # # # # # # # # # # # # # # # # # # # # # # # # # # # # # # # # # # # # # # # # # # # # # # # # # # # # # # # # # # # 
+    def normalizeSubVector(self, vector1, vector2):
+	    import math
+	    returnVec = []
+	    for i in range(len(vector1)):
+		    returnVec.append(vector1[i] - vector2[i])
+	    
+	    #get length of vector
+	    length = math.sqrt( (returnVec[0] * returnVec[0]) + (returnVec[1] * returnVec[1]) + (returnVec[2] * returnVec[2]) )
+	    
+	    #normalize the vector
+	    normalizedVector = []
+	    for i in range(len(returnVec)):
+		    normalizedVector.append(returnVec[i]/length)
+	    
+	    absVector = []
+	    for i in range(len(normalizedVector)):
+		    absVector.append(abs(normalizedVector[i]))
+	    
+	    aimAxis = max(absVector)
+	    aimAxisIndex = absVector.index(aimAxis)
+	    
+	    if aimAxisIndex == 0:
+		    if normalizedVector[0] < 0:
+			    axis = "X"
+		    if normalizedVector[0] > 0:
+			    axis = "-X"
+    
+	    if aimAxisIndex == 1:
+		    if normalizedVector[1] < 0:
+			    axis = "Y"
+		    if normalizedVector[1] > 0:
+			    axis = "-Y"
+			    
+	    if aimAxisIndex == 2:
+		    if normalizedVector[2] < 0:
+			    axis = "Z"
+		    if normalizedVector[2] > 0:
+			    axis = "-Z"
+			    
+	    return axis
+
+        
+        
+        
