Adding ART scripts to repository

diff --git a/Engine/Extras/Maya_AnimationRiggingTools/MayaTools/General/Scripts/ART_exportMotion.py b/Engine/Extras/Maya_AnimationRiggingTools/MayaTools/General/Scripts/ART_exportMotion.py
new file mode 100644
--- /dev/null
+++ b/Engine/Extras/Maya_AnimationRiggingTools/MayaTools/General/Scripts/ART_exportMotion.py
@@ -0,0 +1,1196 @@
+import maya.cmds as cmds
+import maya.mel as mel
+from functools import partial
+import os, cPickle
+
+class ExportMotionUI():
+    
+    def __init__(self):
+        
+        #create class variables
+        self.widgets = {}
+        
+        #find out which project we are in
+        references = cmds.ls(type = "reference")
+	for ref in references:
+	    try:
+		self.project = cmds.referenceQuery(ref, filename = True).rpartition("Projects/")[2].partition("/")[0]
+		
+	    except:
+		pass
+        
+        
+        #get access to our maya tools
+        toolsPath = cmds.internalVar(usd = True) + "mayaTools.txt"
+        if os.path.exists(toolsPath):
+            
+            f = open(toolsPath, 'r')
+            self.mayaToolsDir = f.readline()
+            f.close()
+            
+            
+        #check to see if window exists, if so, delete
+        if cmds.window("exportMotionUI", exists = True):
+            cmds.deleteUI("exportMotionUI")
+            
+        
+        #build window
+        self.widgets["window"] = cmds.window("exportMotionUI", w = 700, h = 400, title = "Export Motion", sizeable = False)
+        
+        #create the main layout
+        self.widgets["topLevelLayout"] = cmds.columnLayout()
+        
+        #create the rowColumnLayout (left side for the different ways one can import motion, the right side with that method's settings
+        self.widgets["rowColumnLayout"] = cmds.rowColumnLayout(w = 700, h = 400, nc = 2, cw = [(1, 150), (2, 550)], parent = self.widgets["topLevelLayout"])
+        
+        #create the columnLayout for the left side
+        self.widgets["leftSideButtonColumn"] = cmds.columnLayout(w = 150, h = 400, parent = self.widgets["rowColumnLayout"], cat = ["both", 5], rs = 5)
+        
+        
+        #and create the frame layout for the right side
+        self.widgets["rightSideFrame"] = cmds.frameLayout(w = 550, h = 400, collapsable = False, borderStyle = "etchedIn", labelVisible = False, parent = self.widgets["rowColumnLayout"])
+        
+        
+        
+        #create the buttons for the different methods of importing motion
+        self.widgets["ExportMotionMethods"] = cmds.iconTextRadioCollection()
+        self.widgets["exportMotion_FBX"] = cmds.iconTextRadioButton(select = True,  w = 140, h = 50,  parent = self.widgets["leftSideButtonColumn"], image = self.mayaToolsDir + "/General/Icons/ART/exportMotion_exportFBX_off.bmp", selectionImage = self.mayaToolsDir + "/General/Icons/ART/exportMotion_exportFBX_on.bmp", onc = partial(self.switchMode, "fbx"))
+        self.widgets["exportMotion_Anim"] = cmds.iconTextRadioButton(select = False,  w = 140, h = 50,  parent = self.widgets["leftSideButtonColumn"], image = self.mayaToolsDir + "/General/Icons/ART/exportMotion_exportAnim_off.bmp", selectionImage = self.mayaToolsDir + "/General/Icons/ART/exportMotion_exportAnim_on.bmp", onc = partial(self.switchMode, "anim"))
+
+
+
+        #create the elements for the right column 
+        # # # # # # # # # # # # # # # # # # # # # # # # # # # # # # # # # # # # # # # # # # # # # # # # # # # # # # # # # # # # # # # 
+        
+        # # # # # # # # # # # # # # # # # # # # # # # # # # # # # # # # # # # # # # # # # # # # # # # # # # # # # # # # # # # # # # # 
+    
+        # # # # # # # # # # # # # # # # # # # # # # # # # # # # # # # # # # # # # # # # # # # # # # # # # # # # # # # # # # # # # # # 
+        #Export FBX
+        self.widgets["exportFBXFormLayout"] = cmds.formLayout(w = 500, h = 400, parent = self.widgets["rightSideFrame"])
+        self.widgets["exportFBX_addAnimButton"] = cmds.symbolButton(w = 100, h = 30, image = self.mayaToolsDir + "/General/Icons/ART/exportMotion_addAnim.bmp", parent = self.widgets["exportFBXFormLayout"])
+        cmds.formLayout(self.widgets["exportFBXFormLayout"], edit = True, af = [(self.widgets["exportFBX_addAnimButton"], "top", 15), (self.widgets["exportFBX_addAnimButton"], "left", 257)])
+
+        #character list
+        self.widgets["exportFBX_characterList"] = cmds.optionMenu(w = 180, h = 30)
+        self.widgets["exportFBX_characterThumb"] = cmds.image(w = 50, h = 50)
+        cmds.formLayout(self.widgets["exportFBXFormLayout"], edit = True, af = [(self.widgets["exportFBX_characterList"], "top", 15), (self.widgets["exportFBX_characterList"], "left", 10)])
+        cmds.formLayout(self.widgets["exportFBXFormLayout"], edit = True, af = [(self.widgets["exportFBX_characterThumb"], "top", 5), (self.widgets["exportFBX_characterThumb"], "left", 195)])
+        
+        #weapon list
+	self.widgets["exportFBX_weaponList"] = cmds.optionMenu(w = 180, h = 30)
+	cmds.menuItem(label = "No Weapon", parent = self.widgets["exportFBX_weaponList"])
+	self.widgets["exportFBX_weaponSuffix"] = cmds.textField(w = 160, h = 30, text = "Weapon", ann = "Export FBX weapon suffix")
+	
+	cmds.formLayout(self.widgets["exportFBXFormLayout"], edit = True, af = [(self.widgets["exportFBX_weaponList"], "top", 60), (self.widgets["exportFBX_weaponList"], "left", 10)])
+	cmds.formLayout(self.widgets["exportFBXFormLayout"], edit = True, af = [(self.widgets["exportFBX_weaponSuffix"], "top", 60), (self.widgets["exportFBX_weaponSuffix"], "left", 197)])
+	
+	
+        self.widgets["exportFBXAnimScrollList"] = cmds.scrollLayout(w = 350, h = 295, hst = 0)
+        self.widgets["exportFBXAnimScrollColumn"] = cmds.columnLayout(parent = self.widgets["exportFBXAnimScrollList"])
+
+        cmds.formLayout(self.widgets["exportFBXFormLayout"], edit = True, af = [(self.widgets["exportFBXAnimScrollList"], "top", 95), (self.widgets["exportFBXAnimScrollList"], "left", 10)])
+        
+        self.widgets["exportFBX_exportMorphsCB"] = cmds.checkBox(label = "Export Morphs", v = False, parent = self.widgets["exportFBXFormLayout"], onc = self.findMorphs, ofc = self.clearMorphs)
+        cmds.formLayout(self.widgets["exportFBXFormLayout"], edit = True, af = [(self.widgets["exportFBX_exportMorphsCB"], "top", 13), (self.widgets["exportFBX_exportMorphsCB"], "left", 370)])
+        
+        self.widgets["exportFBX_morphsList"] = cmds.textScrollList(w = 165, h = 250, parent = self.widgets["exportFBXFormLayout"], allowMultiSelection=True)
+        cmds.formLayout(self.widgets["exportFBXFormLayout"], edit = True, af = [(self.widgets["exportFBX_morphsList"], "top", 50), (self.widgets["exportFBX_morphsList"], "left", 370)])
+
+        self.widgets["exportFBX_exportButton"] = cmds.symbolButton(c = self.exportFBX, w = 165, h = 50, image = self.mayaToolsDir + "/General/Icons/ART/exportMotion_exportButton.bmp", parent = self.widgets["exportFBXFormLayout"])
+        cmds.formLayout(self.widgets["exportFBXFormLayout"], edit = True, af = [(self.widgets["exportFBX_exportButton"], "bottom", 10), (self.widgets["exportFBX_exportButton"], "left", 370)])
+
+        cmds.symbolButton(self.widgets["exportFBX_addAnimButton"], edit = True, c = partial(self.addAnim,self.widgets["exportFBXAnimScrollColumn"], "fbx")) 
+
+
+
+
+        #EXPORT ANIMATION
+        self.widgets["exportAnimationFormLayout"] = cmds.formLayout(w = 500, h = 400, parent = self.widgets["rightSideFrame"], visible = False)
+        
+        
+        #character list
+        self.widgets["exportanim_characterList"] = cmds.optionMenu(w = 180, h = 30, cc = self.getControls)
+        self.widgets["exportanim_characterThumb"] = cmds.image(w = 50, h = 50)
+        cmds.formLayout(self.widgets["exportAnimationFormLayout"], edit = True, af = [(self.widgets["exportanim_characterList"], "top", 15), (self.widgets["exportanim_characterList"], "left", 10)])
+        cmds.formLayout(self.widgets["exportAnimationFormLayout"], edit = True, af = [(self.widgets["exportanim_characterThumb"], "top", 5), (self.widgets["exportanim_characterThumb"], "left", 195)])
+        
+        
+
+        self.widgets["exportanimAnimScrollList"] = cmds.scrollLayout(w = 350, h = 325, hst = 0)
+        self.widgets["exportanimAnimScrollColumn"] = cmds.columnLayout(parent = self.widgets["exportanimAnimScrollList"])
+
+        cmds.formLayout(self.widgets["exportAnimationFormLayout"], edit = True, af = [(self.widgets["exportanimAnimScrollList"], "top", 60), (self.widgets["exportanimAnimScrollList"], "left", 10)])
+        
+        
+
+        self.widgets["exportanim_exportButton"] = cmds.symbolButton(w = 165, h = 50, image = self.mayaToolsDir + "/General/Icons/ART/exportMotion_exportButton.bmp", parent = self.widgets["exportAnimationFormLayout"], c = self.exportAnimation)
+        cmds.formLayout(self.widgets["exportAnimationFormLayout"], edit = True, af = [(self.widgets["exportanim_exportButton"], "bottom", 10), (self.widgets["exportanim_exportButton"], "left", 370)])
+        
+        
+        #add a formLayout to the columnLayout
+        self.widgets["addAnimationSequence_FormLayout"] = cmds.formLayout(h = 200, parent = self.widgets["exportanimAnimScrollColumn"])
+        
+        
+        #add the textField and browse button
+	text = cmds.text(label = "Animation File Name:")
+        self.widgets["addAnimSequence_TextField"] = cmds.textField(w = 290, h = 30)
+        
+        cmds.formLayout(self.widgets["addAnimationSequence_FormLayout"], edit = True, af = [(self.widgets["addAnimSequence_TextField"], "top", 25), (self.widgets["addAnimSequence_TextField"], "left", 5)])
+        cmds.formLayout(self.widgets["addAnimationSequence_FormLayout"], edit = True, af = [(text, "top", 5), (text, "left", 5)])
+	
+	#Category
+	text = cmds.text(label = "Animation Category:")
+	self.widgets["exportanim_exportCategory"] = cmds.optionMenu(w = 290, h = 30)
+	self.widgets["exportanim_addCategory"] = cmds.symbolButton(image = self.mayaToolsDir + "/General/Icons/ART/add.bmp", ann = "Add new category", w = 30, h = 30, c = self.addAnimCategory)
+	cmds.menuItem(label = "None", parent = self.widgets["exportanim_exportCategory"])
+	
+        cmds.formLayout(self.widgets["addAnimationSequence_FormLayout"], edit = True, af = [(self.widgets["exportanim_exportCategory"], "top", 80), (self.widgets["exportanim_exportCategory"], "left", 5)])
+	cmds.formLayout(self.widgets["addAnimationSequence_FormLayout"], edit = True, af = [(self.widgets["exportanim_addCategory"], "top", 80), (self.widgets["exportanim_addCategory"], "left", 300)])
+        cmds.formLayout(self.widgets["addAnimationSequence_FormLayout"], edit = True, af = [(text, "top", 60), (text, "left", 5)])
+        
+
+        
+        #add separator
+        sep = cmds.separator(w = 320, style = "out")
+        cmds.formLayout(self.widgets["addAnimationSequence_FormLayout"], edit = True, af = [(sep, "top", 120), (sep, "left", 8)])
+
+        
+        
+
+        #show window
+        cmds.showWindow(self.widgets["window"])
+        
+        #populate character list
+        self.getCharacters()
+	self.getControls()
+        self.changeActiveCharacter("fbx")
+        self.changeActiveCharacter("anim")
+	
+	#find weapons
+	self.getWeapons()
+        
+        #add the first anim sequence to the list
+	if cmds.objExists("ExportAnimationSettings") == False:
+	    self.addAnim(self.widgets["exportFBXAnimScrollColumn"], "fbx")
+	    
+	else:
+	    #load up info from cached settings
+	    sequeneces = cmds.listAttr("ExportAnimationSettings", string = "sequence*")
+	    
+	    i = 0
+	    for seq in sequeneces:
+		self.addAnim(self.widgets["exportFBXAnimScrollColumn"], "fbx")
+		
+		#get data
+		data = cmds.getAttr("ExportAnimationSettings." + seq)
+		name = data.partition("::")[0]
+		data = data.partition("::")[2]
+		start = data.partition("::")[0]
+		data = data.partition("::")[2]
+		end = data.partition("::")[0]
+		data = data.partition("::")[2]
+		fps = data.partition("::")[0]
+		interp = data.partition("::")[2]
+		
+		
+		
+		#set the data
+		cmds.textField(self.widgets["addAnimSequence_" + str(i + 1) + "_TextField"], edit = True, text = name)
+		cmds.intField(self.widgets["addAnimSequence_" + str(i + 1) + "_startField"], edit = True, v = int(start))
+		cmds.intField(self.widgets["addAnimSequence_" + str(i + 1) + "_endField"], edit = True, v = int(end))
+		cmds.optionMenu(self.widgets["addAnimSequence_" + str(i + 1) + "_fpsField"], edit = True, v = fps)
+		cmds.optionMenu(self.widgets["addAnimSequence_" + str(i + 1) + "_interpMenu"], edit = True, v = interp)
+		
+		i = i + 1
+		
+		
+	
+	
+	#get animation categories
+	self.getAnimCategories()
+        
+    # # # # # # # # # # # # # # # # # # # # # # # # # # # # # # # # # # # # # # # # # # # # # # # # # # # # # # # # # # # # # # # 
+    # # # # # # # # # # # # # # # # # # # # # # # # # # # # # # # # # # # # # # # # # # # # # # # # # # # # # # # # # # # # # # # 
+    # # # # # # # # # # # # # # # # # # # # # # # # # # # # # # # # # # # # # # # # # # # # # # # # # # # # # # # # # # # # # # # 
+    # # # # # # # # # # # # # # # # # # # # # # # # # # # # # # # # # # # # # # # # # # # # # # # # # # # # # # # # # # # # # # # 
+    # # # # # # # # # # # # # # # # # # # # # # # # # # # # # # # # # # # # # # # # # # # # # # # # # # # # # # # # # # # # # # #
+    def getAnimCategories(self, *args):
+	
+	projectPath = self.mayaToolsDir + "/General/ART/Projects/" + self.project + "/Animations/"
+	if not os.path.exists(projectPath):
+	    os.makedirs(projectPath)
+	categories = os.listdir(projectPath)
+	
+	children = cmds.optionMenu(self.widgets["exportanim_exportCategory"], q = True, ill = True)
+	existingCategories = []
+	for child in children:
+	    label = cmds.menuItem(child, q = True, label = True)
+	    existingCategories.append(label)
+	
+	for category in categories:
+	    if category not in existingCategories:
+		cmds.menuItem(label = category, parent = self.widgets["exportanim_exportCategory"])
+	    
+    # # # # # # # # # # # # # # # # # # # # # # # # # # # # # # # # # # # # # # # # # # # # # # # # # # # # # # # # # # # # # # # 
+    # # # # # # # # # # # # # # # # # # # # # # # # # # # # # # # # # # # # # # # # # # # # # # # # # # # # # # # # # # # # # # # 
+    # # # # # # # # # # # # # # # # # # # # # # # # # # # # # # # # # # # # # # # # # # # # # # # # # # # # # # # # # # # # # # # 
+    # # # # # # # # # # # # # # # # # # # # # # # # # # # # # # # # # # # # # # # # # # # # # # # # # # # # # # # # # # # # # # # 
+    # # # # # # # # # # # # # # # # # # # # # # # # # # # # # # # # # # # # # # # # # # # # # # # # # # # # # # # # # # # # # # #
+    def getWeapons(self, *args):
+	
+	#find references in file
+	referenceNodes = []
+        references = cmds.ls(type = "reference")
+        
+        for reference in references:
+	    try:
+		filePath = cmds.referenceQuery(reference, filename = True)
+		
+		if filePath.find("Weapons") != -1:
+		    namespace = cmds.referenceQuery(reference, namespace = True)
+		    referenceNodes.append(namespace.partition(":")[2])
+	    
+	    except:
+		pass
+	    
+		
+	    
+
+
+        for node in referenceNodes:
+            cmds.menuItem(label = node, parent = self.widgets["exportFBX_weaponList"])
+
+    # # # # # # # # # # # # # # # # # # # # # # # # # # # # # # # # # # # # # # # # # # # # # # # # # # # # # # # # # # # # # # # 
+    # # # # # # # # # # # # # # # # # # # # # # # # # # # # # # # # # # # # # # # # # # # # # # # # # # # # # # # # # # # # # # # 
+    # # # # # # # # # # # # # # # # # # # # # # # # # # # # # # # # # # # # # # # # # # # # # # # # # # # # # # # # # # # # # # # 
+    # # # # # # # # # # # # # # # # # # # # # # # # # # # # # # # # # # # # # # # # # # # # # # # # # # # # # # # # # # # # # # # 
+    # # # # # # # # # # # # # # # # # # # # # # # # # # # # # # # # # # # # # # # # # # # # # # # # # # # # # # # # # # # # # # #
+    def switchMode(self, mode, *args):
+        
+        if mode == "fbx":
+            cmds.formLayout(self.widgets["exportFBXFormLayout"], edit = True, visible = True)
+            cmds.formLayout(self.widgets["exportAnimationFormLayout"], edit = True, visible = False)
+            
+        if mode == "anim":
+            cmds.formLayout(self.widgets["exportFBXFormLayout"], edit = True, visible = False)
+            cmds.formLayout(self.widgets["exportAnimationFormLayout"], edit = True, visible = True)
+        
+        
+        
+    # # # # # # # # # # # # # # # # # # # # # # # # # # # # # # # # # # # # # # # # # # # # # # # # # # # # # # # # # # # # # # # 
+    # # # # # # # # # # # # # # # # # # # # # # # # # # # # # # # # # # # # # # # # # # # # # # # # # # # # # # # # # # # # # # # 
+    # # # # # # # # # # # # # # # # # # # # # # # # # # # # # # # # # # # # # # # # # # # # # # # # # # # # # # # # # # # # # # # 
+    # # # # # # # # # # # # # # # # # # # # # # # # # # # # # # # # # # # # # # # # # # # # # # # # # # # # # # # # # # # # # # # 
+    # # # # # # # # # # # # # # # # # # # # # # # # # # # # # # # # # # # # # # # # # # # # # # # # # # # # # # # # # # # # # # #
+    def exportAnimation(self,  *args):
+        character = cmds.optionMenu(self.widgets["exportanim_characterList"], q = True, value = True)
+	
+	
+	animLayers = cmds.ls(type = "animLayer")
+	
+	if animLayers != []:
+	    
+	    if cmds.animLayer("BaseAnimation", q = True, exists = True) == False:
+		cmds.confirmDialog(title = "Export Animation", message = "Base animation layer should be named BaseAnimation(default creation name). Please rename to BaseAnimation and try again.")
+		return
+	    
+	    for layer in animLayers:
+		cmds.animLayer(layer, edit = True, selected = False)
+	    cmds.animLayer("BaseAnimation", edit = True, selected = True, preferred = True)
+
+		
+
+	    
+	#add items to control list
+	exportControls = list(self.controls)
+	
+	cmds.select(character + ":*_space_switcher_follow")
+	nodes = cmds.ls(sl = True)
+	spaceSwitchers = []
+	for node in nodes:
+	    if node.find("invis") == -1:
+		spaceSwitchers.append(node)
+		
+	checkSpaceObjs = []
+        for control in spaceSwitchers:
+            spaceSwitchNode = control.rpartition("_follow")[0]
+	    exportControls.append(spaceSwitchNode.partition(":")[2])
+	    checkSpaceObjs.append(spaceSwitchNode)
+	cmds.select(clear = True)
+	
+	#add the chest space control
+	exportControls.append("chest_ik_world_aligned")
+	
+
+	for node in ["fk_orient_world_loc_l", "fk_orient_world_loc_r", "fk_orient_body_loc_l", "fk_orient_body_loc_r", "head_fk_orient_neck", "head_fk_orient_shoulder", "head_fk_orient_body", "head_fk_orient_world"]:
+	    exportControls.append(node)
+	    
+	
+	#find custom nodes in the scene that the space switching system is using.
+	deleteObjects = []
+	for spaceNode in checkSpaceObjs:
+	    attrs = cmds.listAttr(spaceNode, string = "space_*")
+	    spaceControlName = spaceNode.rpartition("_space_switcher")[0]
+	    
+	    for attr in attrs:
+		obj = attr.partition("space_")[2]
+		if obj not in exportControls:
+		    if obj != "world":
+			
+			#warn user about space switching not exporting
+			result = cmds.confirmDialog(title = "Export Animation", icon = "warning", message =(spaceControlName + " has a custom space that cannot be exported."), button = ["Ignore", "Bake Down " + spaceControlName, "Cancel"], defaultButton = "Continue", cancelButton = "Cancel", dismissString = "Cancel")
+			
+			if result == "Cancel":
+			    return
+			
+			if result == "Bake Down " + spaceControlName:
+			    self.bakeSpace(spaceNode + "_follow")
+			    cmds.select(spaceNode)
+			    cmds.setToolTo("moveSuperContext")
+			    cmds.select(clear = True)
+			    
+			    
+	#Also check for any controls that have constraints
+	toBake = []
+	constraints = ["parentConstraint", "pointConstraint", "orientConstraint", "pairBlend"]
+	constraintAttrs = ["translateX", "translateY", "translateZ", "rotateX", "rotateY", "rotateZ"]
+	
+	for control in exportControls:
+	    if control != "chest_ik_world_aligned":
+		if cmds.objExists(character + ":" + control):
+		    attrs = cmds.listAttr(character + ":" + control, keyable = True)
+		    for attr in attrs:
+			if attr in constraintAttrs:
+			    connection = cmds.listConnections(character + ":" + control + "." + attr, plugs = True)
+			    if connection != None:
+				for each in connection:
+				    if cmds.nodeType(each) in constraints:
+					toBake.append(character + ":" + control)
+	toBake = list(set(toBake))
+	
+	
+	if toBake != []:
+	    #warn user about space switching not exporting
+	    message = ""
+	    for each in toBake:
+		message += each + "\n"
+	    result = cmds.confirmDialog(title = "Export Animation", icon = 'warning', message =("The following controls have constraints:\n\n" + message + "\nThis data won't be exported unless baked."), button = ["Ignore", "Bake Down Data", "Cancel"], defaultButton = "Continue", cancelButton = "Cancel", dismissString = "Cancel")
+	    
+	    if result == "Bake Down Data":
+		cmds.select(clear = True)
+		for each in toBake:
+		    cmds.select(each, add = True)
+		    
+		start = cmds.playbackOptions(q = True, min = True)
+		end = cmds.playbackOptions(q = True, max = True)
+		cmds.bakeResults(simulation = True, t = (start, end))
+
+		
+		
+		
+		
+		    
+
+        animData = []
+        for control in exportControls:
+	    controlData = []
+	    if cmds.objExists(character + ":" + control):
+		controlData.append(control)
+	
+		#get the attrs for this control that are keyable
+		attrs = cmds.listAttr(character + ":" + control, keyable = True)
+		
+		#if baseAnimation layer exists, select it.
+		if cmds.animLayer("BaseAnimation", q = True, exists = True):
+		    animLayers = cmds.ls(type = "animLayer")
+		    for layer in animLayers:
+			cmds.animLayer(layer, edit = True, selected = False, preferred = False)
+			
+		    cmds.animLayer("BaseAnimation", edit = True, selected = True, preferred = True)
+		    cmds.select(character + ":" + control)
+		    
+		#create a list that will hold all of the keyframe information for each keyable attribute as well as the layer that has the keys
+		keyframeInformation = []
+		allAttrData = []
+		
+		#get the keys and key values for each attr
+		for attr in attrs:
+		    keys = cmds.keyframe(character + ":" + control, at = attr, q = True)
+		    keyValues = cmds.keyframe(character + ":" + control,at = attr,  q = True, vc = True)
+		    
+		    attrData = [attr]
+		    keyData = []
+		    
+		    if keys != None:
+			for i in range(int(len(keys))):
+			    #find tangent info of key
+			    tangentData = []
+			    inTanType = cmds.keyTangent(character + ":" + control + "." + attr, q = True, itt = True)[i]
+			    outTanType = cmds.keyTangent(character + ":" + control + "." + attr, q = True, ott = True)[i]
+			    inTanAngle = cmds.keyTangent(character + ":" + control + "." + attr, q = True, ia = True)[i]
+			    outTanAngle = cmds.keyTangent(character + ":" + control + "." + attr, q = True, oa = True)[i]
+			    inWeight = cmds.keyTangent(character + ":" + control + "." + attr, q = True, iw = True)[i]
+			    outWeight = cmds.keyTangent(character + ":" + control + "." + attr, q = True, ow = True)[i]
+			    
+			    
+			    #add to tangent list 
+			    tangentData.append(inTanType)
+			    tangentData.append(outTanType)
+			    tangentData.append(inTanAngle)
+			    tangentData.append(outTanAngle)
+			    tangentData.append(inWeight)
+			    tangentData.append(outWeight)
+			        
+			    keyData.append([keys[i], keyValues[i], tangentData])
+			    
+			attrData.append(keyData)
+		    
+			#add the attrData to the overall list of attributes with keys and their values
+			keyframeInformation.append(attrData)
+		    
+		allAttrData.append([None, keyframeInformation])
+			
+			
+			
+		#find any additional animation layers
+		layers = cmds.listConnections(character + ":" + control, type = "animLayer")
+		keyframeInformation = []
+		
+		if layers != None:
+		    layerData = list(set(layers))
+		    
+		    for layer in layerData:
+			#select the layer to get the keys off of it and add this layer to the keyframe information list
+			for l in animLayers:
+			    cmds.animLayer(l, edit = True, selected = False)
+			cmds.animLayer(layer, edit = True, selected = True, preferred = True)
+			cmds.select(character + ":" + control)
+			
+			for attr in attrs:
+			    keys = cmds.keyframe(character + ":" + control, at = attr, q = True)
+			    keyValues = cmds.keyframe(character + ":" + control,at = attr,  q = True, vc = True)
+			    
+			    attrData = [attr]
+			    keyData = []
+			    
+			    if keys!= None:
+				for i in range(int(len(keys))):
+				    
+				    #find tangent info of key
+				    tangentData = []
+				    inTanType = cmds.keyTangent(character + ":" + control + "." + attr, q = True, itt = True)[i]
+				    outTanType = cmds.keyTangent(character + ":" + control + "." + attr, q = True, ott = True)[i]
+				    inTanAngle = cmds.keyTangent(character + ":" + control + "." + attr, q = True, ia = True)[i]
+				    outTanAngle = cmds.keyTangent(character + ":" + control + "." + attr, q = True, oa = True)[i]
+				    inWeight = cmds.keyTangent(character + ":" + control + "." + attr, q = True, iw = True)[i]
+				    outWeight = cmds.keyTangent(character + ":" + control + "." + attr, q = True, ow = True)[i]
+				    
+				    #add to tangent list 
+				    tangentData.append(inTanType)
+				    tangentData.append(outTanType)
+				    tangentData.append(inTanAngle)
+				    tangentData.append(outTanAngle)
+				    tangentData.append(inWeight)
+				    tangentData.append(outWeight)
+					
+				    keyData.append([keys[i], keyValues[i], tangentData])	
+				    
+				attrData.append(keyData)
+				
+				#add the attrData to the overall list of attributes with keys and their values
+				keyframeInformation.append(attrData)
+				
+			
+			allAttrData.append([layer, keyframeInformation])
+
+	    #compile all the data into the controlData
+	    controlData.append(allAttrData)
+	    animData.append(controlData)
+	    
+	    #testing
+	    #for data in animData:
+		#print data
+		
+	    
+
+	    
+	
+        #export path
+        animName = cmds.textField(self.widgets["addAnimSequence_TextField"], q = True, text = True)
+	filepath = self.mayaToolsDir + "/General/ART/Projects/" + self.project + "/Animations/"
+	if not os.path.exists(filepath):
+	    os.makedirs(filepath)
+	    
+	category = cmds.optionMenu(self.widgets["exportanim_exportCategory"], q = True, v = True)
+	
+	filepath = filepath + category + "/"
+	if not os.path.exists(filepath):
+	    os.makedirs(filepath)
+	    
+	filepath = filepath + animName + ".txt"
+			
+			
+        f = open(filepath, 'w')
+	cPickle.dump(animData, f)
+	f.close()
+	
+	cmds.headsUpMessage( "Animation Exported!", time=3.0 )
+	
+	
+	#delete out temp nodes
+	if len(deleteObjects) > 0:
+	    for obj in deleteObjects:
+		cmds.delete(obj)
+
+		
+    # # # # # # # # # # # # # # # # # # # # # # # # # # # # # # # # # # # # # # # # # # # # # # # # # # # # # # # # # # # # # # # 
+    # # # # # # # # # # # # # # # # # # # # # # # # # # # # # # # # # # # # # # # # # # # # # # # # # # # # # # # # # # # # # # # 
+    # # # # # # # # # # # # # # # # # # # # # # # # # # # # # # # # # # # # # # # # # # # # # # # # # # # # # # # # # # # # # # # 
+    # # # # # # # # # # # # # # # # # # # # # # # # # # # # # # # # # # # # # # # # # # # # # # # # # # # # # # # # # # # # # # # 
+    # # # # # # # # # # # # # # # # # # # # # # # # # # # # # # # # # # # # # # # # # # # # # # # # # # # # # # # # # # # # # # #
+    def addAnimCategory(self,  *args): 
+	
+	cmds.promptDialog(title = "Add Category", message = "Category Name:")
+	text = cmds.promptDialog(q = True, text = True)
+	
+	#find existing categories
+	if not os.path.exists(self.mayaToolsDir + "/General/ART/Projects/" + self.project + "/Animations/"):
+	    os.makedirs(self.mayaToolsDir + "/General/ART/Projects/" + self.project + "/Animations/")
+	    
+	animCategories = os.listdir(self.mayaToolsDir + "/General/ART/Projects/" + self.project + "/Animations/")
+	
+	if text in animCategories:
+	    cmds.warning("Category already exists. Aborting operation.")
+	    return
+	
+	else:
+	    
+	    cmds.menuItem(label = text, parent = self.widgets["exportanim_exportCategory"])
+
+
+    # # # # # # # # # # # # # # # # # # # # # # # # # # # # # # # # # # # # # # # # # # # # # # # # # # # # # # # # # # # # # # # 
+    # # # # # # # # # # # # # # # # # # # # # # # # # # # # # # # # # # # # # # # # # # # # # # # # # # # # # # # # # # # # # # # 
+    # # # # # # # # # # # # # # # # # # # # # # # # # # # # # # # # # # # # # # # # # # # # # # # # # # # # # # # # # # # # # # # 
+    # # # # # # # # # # # # # # # # # # # # # # # # # # # # # # # # # # # # # # # # # # # # # # # # # # # # # # # # # # # # # # # 
+    # # # # # # # # # # # # # # # # # # # # # # # # # # # # # # # # # # # # # # # # # # # # # # # # # # # # # # # # # # # # # # #
+    def exportFBX(self,  *args):
+        
+        #delete the export skeleton if it exists:
+	if cmds.objExists("root"):
+	    cmds.delete("root")
+	    
+        #get the current time range values
+        originalStart = cmds.playbackOptions(q = True, min = True)
+        originalEnd = cmds.playbackOptions(q = True, max = True)
+        currentFPS = cmds.currentUnit(q = True, time = True)
+	currentInterp = 1
+	
+	#get the current rotation interp value
+	options = cmds.optionVar(list = True)
+	
+	for op in options:
+	    if op == "rotationInterpolationDefault":
+		currentInterp = cmds.optionVar(q = op)
+		
+
+        #get the number of anim sequences to export
+        children = cmds.columnLayout(self.widgets["exportFBXAnimScrollColumn"], q = True, childArray = True)
+        if children != None:
+            numAnims = len(children)
+            
+        else:
+            numAnims = 1
+            
+        for i in range(int(numAnims)):
+            #get the file name
+            fileName = cmds.textField(self.widgets["addAnimSequence_" + str(i + 1) + "_TextField"], q = True, text = True)
+            if fileName != "":
+            
+        
+                #get the frame range
+                startFrame = cmds.intField(self.widgets["addAnimSequence_" + str(i + 1) + "_startField"], q = True, v = True)
+                endFrame = cmds.intField(self.widgets["addAnimSequence_" + str(i + 1) + "_endField"], q = True, v = True)
+		
+		#get the fps
+		fps = cmds.optionMenu(self.widgets["addAnimSequence_" + str(i + 1) + "_fpsField"], q = True, v = True)
+		suffix = None
+		
+		if fps == "30 FPS":
+		    fps = "ntsc"
+		    
+		if fps == "60 FPS":
+		    fps = "ntscf"
+		    suffix = "_60fps"
+		    
+		if fps == "120 FPS":
+		    fps  = "120fps"
+		    suffix = "_120fps"
+
+		#get the rotation interpolation
+		interp = cmds.optionMenu(self.widgets["addAnimSequence_" + str(i + 1) + "_interpMenu"], q = True, v = True)
+		if interp == "Independent Euler Angle":
+		    cmds.optionVar(iv = ("rotationInterpolationDefault", 1))
+		    print "set to independent euler"
+		
+		if interp == "Synchronized Euler Angle":
+		    cmds.optionVar(iv = ("rotationInterpolationDefault", 2))
+		    print "set to sync euler"
+		    
+		if interp == "Quaternion Slerp":
+		    cmds.optionVar(iv = ("rotationInterpolationDefault", 3))
+		    print "set to qauternion slerp"
+		
+                
+
+                #get the blendshapes to export
+                blendshapes = cmds.textScrollList(self.widgets["exportFBX_morphsList"], q = True, selectItem = True)
+                if blendshapes != None:
+		    if cmds.objExists("custom_export"):
+			cmds.delete("custom_export")
+			
+                    cube = cmds.polyCube(name = "custom_export")[0]
+                    i = 1
+                    for shape in blendshapes:
+                        attrs = cmds.listAttr(shape, m = True, string = "weight")
+                        #counter for index
+                        
+                        
+                        for attr in attrs:
+                            keys = cmds.keyframe( shape, attribute=attrs, query=True, timeChange = True )
+                            keyValues = cmds.keyframe( shape, attribute=attrs, query=True, valueChange = True )
+                            
+                            morph = cmds.polyCube(name = attr)[0]
+                            if cmds.objExists("custom_export_shapes"):
+                                cmds.blendShape("custom_export_shapes", edit = True, t = (cube, i, morph, 1.0))
+				print "editing blendshape to include " + cube + str(i) + morph
+                                
+                            else:
+                                cmds.select([morph, cube], r = True)
+                                cmds.blendShape(name = "custom_export_shapes")
+                                cmds.select(clear = True)
+                                
+                            cmds.delete(morph)
+                            print i
+                            #transfer keys from original to new morph
+                            if keys != None:
+                                for x in range(int(len(keys))):
+                                    cmds.setKeyframe("custom_export_shapes." + attr, t = (keys[x]), v = keyValues[x])
+                                
+			    if keys == None:
+				for x in range(startFrame, endFrame + 1):
+				    value = cmds.getAttr(shape + "." + attr, time = x)
+				    cmds.setKeyframe("custom_export_shapes." + attr, t = (x), v = value)
+				    
+                            i = i + 1
+                        
+		#duplicate the skeleton
+		character = cmds.optionMenu(self.widgets["exportFBX_characterList"], q = True, value = True)
+		newSkeleton = cmds.duplicate(character + ":" + "root", un = False, ic = False)
+		
+		joints = []
+		for each in newSkeleton:
+		    if cmds.nodeType(each) != "joint":
+			cmds.delete(each)
+			
+		    else:
+			joints.append(each)
+			
+
+		constraints = []
+		for joint in joints:
+		    #do some checks to make sure that this is valid
+		    parent = cmds.listRelatives(joint, parent = True)
+
+		    if parent != None:
+			if parent[0] in joints:
+			    constraint = cmds.parentConstraint(character + ":" + parent[0] + "|" + character + ":" + joint, joint)[0]
+			    constraints.append(constraint)
+			    constraint = cmds.scaleConstraint(character + ":" + parent[0] + "|" + character + ":" + joint, joint)[0]
+			    constraints.append(constraint)			    
+		    else:
+			#root bone?
+			if joint == "root":
+			    constraint = cmds.parentConstraint(character + ":" + joint, joint)[0]
+			    constraints.append(constraint)			    
+
+
+		    
+		    
+		#set time range to that of the start and end frame and then bake
+		if startFrame != "":
+		    cmds.playbackOptions(min = startFrame, animationStartTime = startFrame)
+		    cmds.playbackOptions(max = endFrame, animationEndTime = endFrame)
+		    
+		    #set to the new fps
+		    cmds.currentUnit(time = fps)
+		    
+		    #set the new frame range
+		    
+		    if fps == "ntscf":
+			cmds.playbackOptions(min = startFrame * 2, animationStartTime = startFrame * 2)
+			cmds.playbackOptions(max = endFrame * 2, animationEndTime = endFrame * 2)
+			startFrame = startFrame * 2
+			endFrame = endFrame * 2
+			
+		    if fps == "120fps":
+			cmds.playbackOptions(min = startFrame * 4, animationStartTime = startFrame * 4)
+			cmds.playbackOptions(max = endFrame * 4, animationEndTime = endFrame * 4)
+			startFrame = startFrame * 4
+			endFrame = endFrame * 4
+			
+
+		    cmds.select("root", hi = True)
+		    cmds.bakeResults(simulation = True, t = (startFrame, endFrame))
+		    cmds.delete(constraints)
+		    
+		    #run an euler filter
+		    cmds.select("root", hi = True)
+		    cmds.filterCurve()
+		
+		    #export selected
+		    #first change some fbx properties
+		    string = "FBXExportConstraints -v 1;"
+		    string += "FBXExportCacheFile -v 0;"
+		    mel.eval(string)
+		    
+		    cmds.select("root", hi = True)
+		    if cmds.objExists("custom_export"):
+			cmds.select("custom_export", add = True)
+			
+		    #Add FPS Suffix to animation file name
+		    if suffix != None:
+			prefix = fileName.rpartition(".")[0]
+			fileName = prefix + suffix + ".fbx"
+			
+		    cmds.file(fileName, es = True, force = True, prompt = False, type = "FBX export")
+		    
+
+			
+
+		    #clean scene
+		    if cmds.objExists("root"):
+			cmds.delete("root")
+		    if cmds.objExists("custom_export"):
+			cmds.delete("custom_export")	
+	    
+		    cmds.select(clear = True)
+		    
+		    
+		    
+		    
+
+		    # # # # # # # # # # # # # # # # # # # # # # # # # # # #
+		    #Export Weapon Animation
+		    if cmds.optionMenu(self.widgets["exportFBX_weaponList"], q = True, v = True) != "No Weapon":
+			
+			#duplicate weapon skeleton and constrain
+			namespace = cmds.optionMenu(self.widgets["exportFBX_weaponList"], q = True, v = True)
+			cmds.select(namespace + ":*")
+			joints = cmds.ls(sl = True, type = "joint")
+			
+			for joint in joints:
+			    parent = cmds.listRelatives(joint, parent = True)[0]
+			    if cmds.nodeType(parent) != "joint":
+				rootJoint = joint
+				
+				
+			newSkeleton = cmds.duplicate(rootJoint, un = False, ic = False)
+			root = newSkeleton[0]
+			cmds.parent(root, world = True)
+			
+			joints = []
+			for each in newSkeleton:
+			    if cmds.nodeType(each) != "joint":
+				cmds.delete(each)
+				
+			    else:
+				joints.append(each)
+	    
+			#constrain new skeleton to original
+			constraints = []
+			for joint in joints:
+			    constraint = cmds.parentConstraint(namespace + ":" + joint, joint)[0]
+			    constraints.append(constraint)
+			    
+			#bake
+			cmds.select(root, hi = True)
+			cmds.bakeResults(simulation = True, t = (startFrame, endFrame))
+			cmds.delete(constraints)			
+
+			#run an euler filter
+			cmds.select(root, hi = True)
+			cmds.filterCurve()
+			
+			#zero out the root
+			cmds.cutKey(root)
+			
+			for attr in [".tx", ".ty", ".tz", ".rx", ".ry", ".rz"]:
+			    cmds.setAttr(root + attr, 0)
+			    
+			#export selected
+			#first change some fbx properties
+			string = "FBXExportConstraints -v 1;"
+			string += "FBXExportCacheFile -v 0;"
+			mel.eval(string)
+			
+			cmds.select(root, hi = True)
+			    
+			#Add Weapon Suffix to animation file name
+			suffix = cmds.textField(self.widgets["exportFBX_weaponSuffix"], q = True, text = True)
+			if suffix == "":
+			    suffix = "Weapon"
+			    
+			cmds.textField(self.widgets["addAnimSequence_" + str(i + 1) + "_TextField"], q = True, text = True)
+			prefix = fileName.rpartition(".")[0]
+			fileName = prefix + "_" + suffix + ".fbx"
+			    
+			cmds.select(root, hi = True)
+			cmds.file(fileName, es = True, force = True, prompt = False, type = "FBX export")	
+			
+			#clean scene
+			cmds.delete(root)
+			
+
+
+		#reset fps to original
+		cmds.currentUnit(time = currentFPS)		
+
+		#reset timeslider to original values
+		cmds.playbackOptions(min = originalStart, max = originalEnd)
+		
+		#reset rotation interp
+		cmds.optionVar(iv = ("rotationInterpolationDefault", currentInterp))
+		
+
+
+	#export complete
+	cmds.confirmDialog(title = "Export", message = "Export Complete!", button = "OK")
+
+	#save export information to scene
+	if cmds.objExists("ExportAnimationSettings"):
+	    cmds.delete("ExportAnimationSettings")
+	    
+	
+	exportAnimSettings = cmds.createNode("unknown", name = "ExportAnimationSettings")
+	
+	for i in range(numAnims):
+	    
+	    #create a string attr with sequenceName|startFrame|endFrame|FPS
+	    string = str(cmds.textField(self.widgets["addAnimSequence_" + str(i + 1) + "_TextField"], q = True, text = True)) + "::"
+	    string += str(cmds.intField(self.widgets["addAnimSequence_" + str(i + 1) + "_startField"], q = True, v = True)) + "::"
+	    string += str(cmds.intField(self.widgets["addAnimSequence_" + str(i + 1) + "_endField"], q = True, v = True)) + "::"
+	    string += cmds.optionMenu(self.widgets["addAnimSequence_" + str(i + 1) + "_fpsField"], q = True, v = True) + "::"
+	    string += cmds.optionMenu(self.widgets["addAnimSequence_" + str(i + 1) + "_interpMenu"], q = True, v = True)
+	    
+	    #add the attr
+	    cmds.addAttr(exportAnimSettings, dt = "string", ln = "sequence" + str(i))
+	    cmds.setAttr(exportAnimSettings + ".sequence" + str(i),  str(string), type = "string")
+
+
+    # # # # # # # # # # # # # # # # # # # # # # # # # # # # # # # # # # # # # # # # # # # # # # # # # # # # # # # # # # # # # # # 
+    # # # # # # # # # # # # # # # # # # # # # # # # # # # # # # # # # # # # # # # # # # # # # # # # # # # # # # # # # # # # # # # 
+    # # # # # # # # # # # # # # # # # # # # # # # # # # # # # # # # # # # # # # # # # # # # # # # # # # # # # # # # # # # # # # # 
+    # # # # # # # # # # # # # # # # # # # # # # # # # # # # # # # # # # # # # # # # # # # # # # # # # # # # # # # # # # # # # # # 
+    # # # # # # # # # # # # # # # # # # # # # # # # # # # # # # # # # # # # # # # # # # # # # # # # # # # # # # # # # # # # # # #
+    def removeAnim(self, layout,  *args):
+        cmds.deleteUI(layout)
+        
+        
+        
+    # # # # # # # # # # # # # # # # # # # # # # # # # # # # # # # # # # # # # # # # # # # # # # # # # # # # # # # # # # # # # # # 
+    # # # # # # # # # # # # # # # # # # # # # # # # # # # # # # # # # # # # # # # # # # # # # # # # # # # # # # # # # # # # # # # 
+    # # # # # # # # # # # # # # # # # # # # # # # # # # # # # # # # # # # # # # # # # # # # # # # # # # # # # # # # # # # # # # # 
+    # # # # # # # # # # # # # # # # # # # # # # # # # # # # # # # # # # # # # # # # # # # # # # # # # # # # # # # # # # # # # # # 
+    # # # # # # # # # # # # # # # # # # # # # # # # # # # # # # # # # # # # # # # # # # # # # # # # # # # # # # # # # # # # # # #
+    def addAnim(self, layout, mode,  *args):
+        
+        #add an animation sequence to the columnLayout
+        
+        #find number of children to the columnLayout
+        children = cmds.columnLayout(layout, q = True, childArray = True)
+        if children != None:
+            numChildren = len(children)
+            
+        else:
+            numChildren = 0
+        
+        #add a formLayout to the columnLayout
+        self.widgets["addAnimSequence_" + str(numChildren + 1) + "_FormLayout"] = cmds.formLayout(h = 110, parent = layout)
+        
+        #add the textField and browse button
+        self.widgets["addAnimSequence_" + str(numChildren + 1) + "_TextField"] = cmds.textField(w = 290, h = 30)
+        self.widgets["addAnimSequence_" + str(numChildren + 1) + "_browseButton"] = cmds.symbolButton(w = 30, h = 30, image = self.mayaToolsDir + "/General/Icons/ART/browse.bmp", c = partial(self.fileBrowse, self.widgets["addAnimSequence_" + str(numChildren + 1) + "_TextField"], "fbx"))
+        
+        cmds.formLayout(self.widgets["addAnimSequence_" + str(numChildren + 1) + "_FormLayout"], edit = True, af = [(self.widgets["addAnimSequence_" + str(numChildren + 1) + "_TextField"], "top", 5), (self.widgets["addAnimSequence_" + str(numChildren + 1) + "_TextField"], "left", 5)])
+        cmds.formLayout(self.widgets["addAnimSequence_" + str(numChildren + 1) + "_FormLayout"], edit = True, af = [(self.widgets["addAnimSequence_" + str(numChildren + 1) + "_browseButton"], "top", 5), (self.widgets["addAnimSequence_" + str(numChildren + 1) + "_browseButton"], "left", 300)])
+        
+        #add the frame range fields
+        label = cmds.text(label = "Frame Range:", font = "boldLabelFont")
+        self.widgets["addAnimSequence_" + str(numChildren + 1) + "_startField"] = cmds.intField(v = 0, w = 60)
+        self.widgets["addAnimSequence_" + str(numChildren + 1) + "_endField"] = cmds.intField(v = 0, w = 60)
+        label2 = cmds.text(label = " to ")
+	
+	#add FPS options
+	self.widgets["addAnimSequence_" + str(numChildren + 1) + "_fpsField"] = cmds.optionMenu(w = 80)
+	cmds.menuItem(label = "30 FPS", parent = self.widgets["addAnimSequence_" + str(numChildren + 1) + "_fpsField"])
+	cmds.menuItem(label = "60 FPS", parent = self.widgets["addAnimSequence_" + str(numChildren + 1) + "_fpsField"])
+	cmds.menuItem(label = "120 FPS", parent = self.widgets["addAnimSequence_" + str(numChildren + 1) + "_fpsField"])
+	
+        #add the rotation interpolation options
+	#1 = independent euler angle, 2 = synchronized euler angle, 3 = quaternion slerp
+	interpLabel = cmds.text(label = "Rotation Interpolation: ", font = "boldLabelFont")
+	self.widgets["addAnimSequence_" + str(numChildren + 1) + "_interpMenu"] = cmds.optionMenu(w = 180)
+	cmds.menuItem(label = "Independent Euler Angle", parent = self.widgets["addAnimSequence_" + str(numChildren + 1) + "_interpMenu"])
+	cmds.menuItem(label = "Synchronized Euler Angle", parent = self.widgets["addAnimSequence_" + str(numChildren + 1) + "_interpMenu"])
+	cmds.menuItem(label = "Quaternion Slerp", parent = self.widgets["addAnimSequence_" + str(numChildren + 1) + "_interpMenu"])
+	
+	#place widgets
+        cmds.formLayout(self.widgets["addAnimSequence_" + str(numChildren + 1) + "_FormLayout"], edit = True, af = [(label, "top", 40), (label, "left", 5)])
+        cmds.formLayout(self.widgets["addAnimSequence_" + str(numChildren + 1) + "_FormLayout"], edit = True, af = [(self.widgets["addAnimSequence_" + str(numChildren + 1) + "_startField"], "top", 40), (self.widgets["addAnimSequence_" + str(numChildren + 1) + "_startField"], "left", 95)])
+        cmds.formLayout(self.widgets["addAnimSequence_" + str(numChildren + 1) + "_FormLayout"], edit = True, af = [(self.widgets["addAnimSequence_" + str(numChildren + 1) + "_endField"], "top", 40), (self.widgets["addAnimSequence_" + str(numChildren + 1) + "_endField"], "left", 190)])
+        cmds.formLayout(self.widgets["addAnimSequence_" + str(numChildren + 1) + "_FormLayout"], edit = True, af = [(label2, "top", 42), (label2, "left", 165)])
+        cmds.formLayout(self.widgets["addAnimSequence_" + str(numChildren + 1) + "_FormLayout"], edit = True, af = [(self.widgets["addAnimSequence_" + str(numChildren + 1) + "_fpsField"], "top", 41), (self.widgets["addAnimSequence_" + str(numChildren + 1) + "_fpsField"], "left", 255)])
+	
+	cmds.formLayout(self.widgets["addAnimSequence_" + str(numChildren + 1) + "_FormLayout"], edit = True, af = [(interpLabel, "top", 71), (interpLabel, "left", 5)])
+	cmds.formLayout(self.widgets["addAnimSequence_" + str(numChildren + 1) + "_FormLayout"], edit = True, af = [(self.widgets["addAnimSequence_" + str(numChildren + 1) + "_interpMenu"], "top", 71), (self.widgets["addAnimSequence_" + str(numChildren + 1) + "_interpMenu"], "left", 155)])
+
+	
+        #add separator
+        sep = cmds.separator(w = 320, style = "out")
+        cmds.formLayout(self.widgets["addAnimSequence_" + str(numChildren + 1) + "_FormLayout"], edit = True, af = [(sep, "top", 100), (sep, "left", 8)])
+        sep2 = cmds.separator(w = 320, style = "out")
+        cmds.formLayout(self.widgets["addAnimSequence_" + str(numChildren + 1) + "_FormLayout"], edit = True, af = [(sep2, "top", 105), (sep2, "left", 8)])
+        
+        if numChildren == 0:
+            #get info from the scene and auto-fill in the fields
+            startFrame = cmds.playbackOptions(q = True, min = True)
+            endFrame = cmds.playbackOptions(q = True, max = True)
+            
+            cmds.intField(self.widgets["addAnimSequence_" + str(numChildren + 1) + "_startField"], edit = True, value = startFrame)
+            cmds.intField(self.widgets["addAnimSequence_" + str(numChildren + 1) + "_endField"], edit = True, value = endFrame)
+            
+            
+            if mode == "fbx":
+                #file name + path
+                openedFile = cmds.file(q = True, sceneName = True)
+                if openedFile != "":
+                    fileName = openedFile.rpartition(".")[0]
+                    fileName += ".fbx"
+                    cmds.textField(self.widgets["addAnimSequence_" + str(numChildren + 1) + "_TextField"] , edit = True, text = fileName)
+                    
+                    
+                
+        if numChildren > 0:
+            #add popup menu to the formLayout for removing the anim sequence
+            self.widgets["addAnimSequence_" + str(numChildren + 1) + "_popupMenu"] = cmds.popupMenu(b = 3, parent = self.widgets["addAnimSequence_" + str(numChildren + 1) + "_FormLayout"])
+            removeMenu = cmds.menuItem(label = "Remove Sequence", parent = self.widgets["addAnimSequence_" + str(numChildren + 1) + "_popupMenu"], c = partial(self.removeAnim, self.widgets["addAnimSequence_" + str(numChildren + 1) + "_FormLayout"]))
+            
+            
+            
+    # # # # # # # # # # # # # # # # # # # # # # # # # # # # # # # # # # # # # # # # # # # # # # # # # # # # # # # # # # # # # # # 
+    # # # # # # # # # # # # # # # # # # # # # # # # # # # # # # # # # # # # # # # # # # # # # # # # # # # # # # # # # # # # # # # 
+    # # # # # # # # # # # # # # # # # # # # # # # # # # # # # # # # # # # # # # # # # # # # # # # # # # # # # # # # # # # # # # # 
+    # # # # # # # # # # # # # # # # # # # # # # # # # # # # # # # # # # # # # # # # # # # # # # # # # # # # # # # # # # # # # # # 
+    # # # # # # # # # # # # # # # # # # # # # # # # # # # # # # # # # # # # # # # # # # # # # # # # # # # # # # # # # # # # # # # 
+    def fileBrowse(self, textField, mode, *args):
+        path = cmds.fileDialog2(fm = 0,  okc = "OK")[0].rpartition(".")[0]
+        if path == None:
+            return
+        
+        if mode == "fbx":
+            #edit the text field with the above path passed in
+            cmds.textField(textField, edit = True, text = path + ".fbx")
+            
+        if mode == "anim":
+            cmds.textField(textField, edit = True, text = path + ".txt")
+        
+    # # # # # # # # # # # # # # # # # # # # # # # # # # # # # # # # # # # # # # # # # # # # # # # # # # # # # # # # # # # # # # # 
+    # # # # # # # # # # # # # # # # # # # # # # # # # # # # # # # # # # # # # # # # # # # # # # # # # # # # # # # # # # # # # # # 
+    # # # # # # # # # # # # # # # # # # # # # # # # # # # # # # # # # # # # # # # # # # # # # # # # # # # # # # # # # # # # # # # 
+    # # # # # # # # # # # # # # # # # # # # # # # # # # # # # # # # # # # # # # # # # # # # # # # # # # # # # # # # # # # # # # # 
+    # # # # # # # # # # # # # # # # # # # # # # # # # # # # # # # # # # # # # # # # # # # # # # # # # # # # # # # # # # # # # # # 
+    def findMorphs(self, *args):
+	
+        physiqueShapes = ["calf_l_shapes", "calf_r_shapes", "head_shapes", "l_elbow_shapes", "l_lowerarm_shapes", "l_shoulder_shapes", "l_upperarm_shapes", "neck1_shapes", "neck2_shapes", "neck3_shapes", "pelvis_shapes",
+	                  "r_elbow_shapes", "r_lowerarm_shapes", "r_shoulder_shapes", "r_upperarm_shapes", "spine1_shapes", "spine2_shapes", "spine3_shapes", "spine4_shapes", "spine5_shapes", "thigh_l_shapes", "thigh_r_shapes"]
+	
+	
+        blendshapes = cmds.ls(type = "blendShape")
+        for shape in blendshapes:
+	    niceName = shape.partition(":")[2]
+	    
+	    if niceName not in physiqueShapes:
+		cmds.textScrollList(self.widgets["exportFBX_morphsList"], edit = True, append = shape)
+            
+    # # # # # # # # # # # # # # # # # # # # # # # # # # # # # # # # # # # # # # # # # # # # # # # # # # # # # # # # # # # # # # # 
+    # # # # # # # # # # # # # # # # # # # # # # # # # # # # # # # # # # # # # # # # # # # # # # # # # # # # # # # # # # # # # # # 
+    # # # # # # # # # # # # # # # # # # # # # # # # # # # # # # # # # # # # # # # # # # # # # # # # # # # # # # # # # # # # # # # 
+    # # # # # # # # # # # # # # # # # # # # # # # # # # # # # # # # # # # # # # # # # # # # # # # # # # # # # # # # # # # # # # # 
+    # # # # # # # # # # # # # # # # # # # # # # # # # # # # # # # # # # # # # # # # # # # # # # # # # # # # # # # # # # # # # # # 
+    def clearMorphs(self, *args):
+        cmds.textScrollList(self.widgets["exportFBX_morphsList"], edit = True, removeAll = True)
+        
+    # # # # # # # # # # # # # # # # # # # # # # # # # # # # # # # # # # # # # # # # # # # # # # # # # # # # # # # # # # # # # # # 
+    # # # # # # # # # # # # # # # # # # # # # # # # # # # # # # # # # # # # # # # # # # # # # # # # # # # # # # # # # # # # # # # 
+    # # # # # # # # # # # # # # # # # # # # # # # # # # # # # # # # # # # # # # # # # # # # # # # # # # # # # # # # # # # # # # # 
+    # # # # # # # # # # # # # # # # # # # # # # # # # # # # # # # # # # # # # # # # # # # # # # # # # # # # # # # # # # # # # # # 
+    # # # # # # # # # # # # # # # # # # # # # # # # # # # # # # # # # # # # # # # # # # # # # # # # # # # # # # # # # # # # # # # 
+    def changeActiveCharacter(self, mode):
+        
+        if mode == "fbx":
+            characterName = cmds.optionMenu(self.widgets["exportFBX_characterList"], q = True, value = True)
+            
+        if mode == "anim":
+            characterName = cmds.optionMenu(self.widgets["exportanim_characterList"], q = True, value = True)
+        
+        thumbnailPath = self.mayaToolsDir + "/General/Icons/ART/Thumbnails/" + self.project + "/"
+        thumbs = os.listdir(thumbnailPath)
+        for thumb in thumbs:
+            if thumb.find("_small") != -1:
+                if thumb.find(characterName) == 0:
+                    
+                    if mode == "fbx":
+                        cmds.image(self.widgets["exportFBX_characterThumb"], edit = True, image = thumbnailPath + thumb, ann = characterName)
+                        
+                    if mode == "anim":
+                        cmds.image(self.widgets["exportanim_characterThumb"], edit = True, image = thumbnailPath + thumb, ann = characterName)
+        
+    # # # # # # # # # # # # # # # # # # # # # # # # # # # # # # # # # # # # # # # # # # # # # # # # # # # # # # # # # # # # # # # 
+    # # # # # # # # # # # # # # # # # # # # # # # # # # # # # # # # # # # # # # # # # # # # # # # # # # # # # # # # # # # # # # # 
+    # # # # # # # # # # # # # # # # # # # # # # # # # # # # # # # # # # # # # # # # # # # # # # # # # # # # # # # # # # # # # # # 
+    # # # # # # # # # # # # # # # # # # # # # # # # # # # # # # # # # # # # # # # # # # # # # # # # # # # # # # # # # # # # # # # 
+    # # # # # # # # # # # # # # # # # # # # # # # # # # # # # # # # # # # # # # # # # # # # # # # # # # # # # # # # # # # # # # # 
+    def getCharacters(self):
+        referenceNodes = []
+        references = cmds.ls(type = "reference")
+        
+        for reference in references:
+            niceName = reference.rpartition("RN")[0]
+            suffix = reference.rpartition("RN")[2]
+            if suffix != "":
+                if cmds.objExists(niceName + suffix + ":" + "Skeleton_Settings"):
+                    referenceNodes.append(niceName + suffix)
+                    
+            else:
+                if cmds.objExists(niceName + ":" + "Skeleton_Settings"):
+                    referenceNodes.append(niceName)
+                
+        for node in referenceNodes:
+            cmds.menuItem(label = node, parent = self.widgets["exportFBX_characterList"])
+        
+            cmds.menuItem(label = node, parent = self.widgets["exportanim_characterList"])
+
+
+
+
+
+
+    # # # # # # # # # # # # # # # # # # # # # # # # # # # # # # # # # # # # # # # # # # # # # # # # # # # # # # # # # # # # # # # 
+    # # # # # # # # # # # # # # # # # # # # # # # # # # # # # # # # # # # # # # # # # # # # # # # # # # # # # # # # # # # # # # # 
+    # # # # # # # # # # # # # # # # # # # # # # # # # # # # # # # # # # # # # # # # # # # # # # # # # # # # # # # # # # # # # # # 
+    # # # # # # # # # # # # # # # # # # # # # # # # # # # # # # # # # # # # # # # # # # # # # # # # # # # # # # # # # # # # # # # 
+    # # # # # # # # # # # # # # # # # # # # # # # # # # # # # # # # # # # # # # # # # # # # # # # # # # # # # # # # # # # # # # # 
+    def getControls(self):
+	#get all controls
+        self.controls = []
+        for control in ["head_fk_anim", "neck_01_fk_anim", "neck_02_fk_anim", "neck_03_fk_anim", "spine_01_anim", "spine_02_anim", "spine_03_anim", "spine_04_anim", "spine_05_anim", "mid_ik_anim", "chest_ik_anim",
+                        "body_anim", "hip_anim", "clavicle_l_anim", "clavicle_r_anim", "fk_arm_l_anim", "fk_arm_r_anim", "fk_elbow_l_anim", "fk_elbow_r_anim", "fk_wrist_l_anim", "fk_wrist_r_anim",
+                        "ik_elbow_l_anim", "ik_elbow_r_anim", "ik_wrist_l_anim", "ik_wrist_r_anim", "fk_thigh_l_anim", "fk_thigh_r_anim", "fk_calf_l_anim", "fk_calf_r_anim", "fk_foot_l_anim", "fk_foot_r_anim",
+                        "fk_ball_l_anim", "fk_ball_r_anim", "ik_foot_anim_l", "ik_foot_anim_r", "heel_ctrl_l", "heel_ctrl_r", "toe_wiggle_ctrl_l", "toe_wiggle_ctrl_r",
+                        "toe_tip_ctrl_l", "toe_tip_ctrl_r", "master_anim", "offset_anim", "root_anim", "upperarm_l_twist_anim", "upperarm_r_twist_anim", "thigh_l_twist_anim", "thigh_r_twist_anim",
+                        "pinky_metacarpal_ctrl_l", "pinky_metacarpal_ctrl_r", "pinky_finger_fk_ctrl_1_l", "pinky_finger_fk_ctrl_1_r", "pinky_finger_fk_ctrl_2_l", "pinky_finger_fk_ctrl_2_r", "pinky_finger_fk_ctrl_3_l", "pinky_finger_fk_ctrl_3_r",
+                        "ring_metacarpal_ctrl_l", "ring_metacarpal_ctrl_r", "ring_finger_fk_ctrl_1_l", "ring_finger_fk_ctrl_1_r", "ring_finger_fk_ctrl_2_l", "ring_finger_fk_ctrl_2_r", "ring_finger_fk_ctrl_3_l", "ring_finger_fk_ctrl_3_r",
+                        "middle_metacarpal_ctrl_l", "middle_metacarpal_ctrl_r", "middle_finger_fk_ctrl_1_l", "middle_finger_fk_ctrl_1_r", "middle_finger_fk_ctrl_2_l", "middle_finger_fk_ctrl_2_r", "middle_finger_fk_ctrl_3_l", "middle_finger_fk_ctrl_3_r",
+                        "index_metacarpal_ctrl_l", "index_metacarpal_ctrl_r", "index_finger_fk_ctrl_1_l", "index_finger_fk_ctrl_1_r", "index_finger_fk_ctrl_2_l", "index_finger_fk_ctrl_2_r", "index_finger_fk_ctrl_3_l", "index_finger_fk_ctrl_3_r",
+                        "thumb_finger_fk_ctrl_1_l", "thumb_finger_fk_ctrl_1_r", "thumb_finger_fk_ctrl_2_l", "thumb_finger_fk_ctrl_2_r", "thumb_finger_fk_ctrl_3_l", "thumb_finger_fk_ctrl_3_r",
+                        "index_l_ik_anim", "index_r_ik_anim", "middle_l_ik_anim", "middle_r_ik_anim", "ring_l_ik_anim", "ring_r_ik_anim", "pinky_l_ik_anim", "pinky_r_ik_anim", "thumb_l_ik_anim", "thumb_r_ik_anim",
+                        "index_l_poleVector", "index_r_poleVector", "middle_l_poleVector", "middle_r_poleVector", "ring_l_poleVector", "ring_r_poleVector", "pinky_l_poleVector", "pinky_r_poleVector", "thumb_l_poleVector", "thumb_r_poleVector",
+                        "l_global_ik_anim", "r_global_ik_anim", "Rig_Settings"]:
+	    self.controls.append(control)
+	    
+	#find custom joints
+	character = cmds.optionMenu(self.widgets["exportanim_characterList"], q = True, value = True)
+        customJoints = []
+        attrs = cmds.listAttr(character + ":" + "Skeleton_Settings")
+        for attr in attrs:
+            if attr.find("extraJoint") == 0:
+                customJoints.append(attr)
+                
+        for joint in customJoints:
+            attribute = cmds.getAttr(character + ":" + "Skeleton_Settings." + joint, asString = True)
+            jointType = attribute.partition("/")[2].partition("/")[0]
+            label = attribute.rpartition("/")[2]
+            
+            if jointType == "leaf":
+                label = label.partition(" (")[0]
+		control = label + "_anim"
+		self.controls.append(control)
+    
+            if jointType == "jiggle":
+		control = label + "_anim"
+		self.controls.append(control)
+		
+	    if jointType == "dynamic":
+		numJointsInChain = label.partition("(")[2].partition(")")[0]
+		label = label.partition(" (")[0]
+		
+		self.controls.append("fk_" + label + "_01_anim")
+		self.controls.append(label + "_cv_0_anim")
+		self.controls.append(label + "_ik_mid_anim")
+		self.controls.append(label + "_ik_top_anim")
+		self.controls.append(label + "_dyn_anim")
+		
+    # # # # # # # # # # # # # # # # # # # # # # # # # # # # # # # # # # # # # # # # # # # # # # # # # # # # # # # # # # # # # # # 
+    # # # # # # # # # # # # # # # # # # # # # # # # # # # # # # # # # # # # # # # # # # # # # # # # # # # # # # # # # # # # # # # 
+    # # # # # # # # # # # # # # # # # # # # # # # # # # # # # # # # # # # # # # # # # # # # # # # # # # # # # # # # # # # # # # # 
+    # # # # # # # # # # # # # # # # # # # # # # # # # # # # # # # # # # # # # # # # # # # # # # # # # # # # # # # # # # # # # # # 
+    # # # # # # # # # # # # # # # # # # # # # # # # # # # # # # # # # # # # # # # # # # # # # # # # # # # # # # # # # # # # # # # 
+    def bakeSpace(self, node,  *args):
+
+	#now we have a list of nodes to bake
+	locs = []
+	constraints = []
+	character = cmds.optionMenu(self.widgets["exportanim_characterList"], q = True, value = True)
+	
+	#create a locator constrained to the space switch node
+	spaceSwitchNode = node.rpartition("_follow")[0]
+	control = spaceSwitchNode.partition("_space_switcher")[0]
+	loc = cmds.spaceLocator(name = spaceSwitchNode + "_loc")[0]
+	locs.append(loc)
+	constraint = cmds.parentConstraint(control, loc)[0]
+	constraints.append(constraint)
+	    
+	#bake all of the locators
+	for loc in locs:
+	    cmds.select(loc, add = True)
+	    
+	start = cmds.playbackOptions(q = True, min = True)
+	end = cmds.playbackOptions(q = True, max = True)
+	cmds.bakeResults(simulation = True, t = (start, end))
+	cmds.delete(constraints)
+	cmds.select(clear = True)
+	
+	
+	#delete keys on all space switch attrs 
+	attrs = cmds.listAttr(spaceSwitchNode, string = "space_*")
+	cmds.cutKey(spaceSwitchNode)
+	
+	for attr in attrs:
+	    #zero out all attrs on the space switcher
+	    cmds.setAttr(spaceSwitchNode + "." + attr, 0)
+		
+	
+	#reverse the constraint so now the spaceSwitchNode will follow the loc
+	loc = spaceSwitchNode + "_loc"
+	if control.find(character + ":ik_elbow") == 0:
+	    constraint = cmds.pointConstraint(loc, control)[0]
+	    
+	else:
+	    constraint = cmds.parentConstraint(loc, control)[0]
+	    
+	cmds.select(control, add = True)
+	
+	    
+	
+	cmds.bakeResults(simulation = True, t = (start, end))
+	cmds.delete(locs)
\ No newline at end of file
