Adding ART scripts to repository

diff --git a/Engine/Extras/Maya_AnimationRiggingTools/MayaTools/General/Scripts/ART_importMotion.py b/Engine/Extras/Maya_AnimationRiggingTools/MayaTools/General/Scripts/ART_importMotion.py
new file mode 100644
--- /dev/null
+++ b/Engine/Extras/Maya_AnimationRiggingTools/MayaTools/General/Scripts/ART_importMotion.py
@@ -0,0 +1,1554 @@
+import maya.cmds as cmds
+import maya.mel as mel
+from functools import partial
+import os, cPickle
+import math
+
+class ImportMotionUI():
+    
+    def __init__(self):
+        
+        #create class variables
+        self.widgets = {}
+        
+        #find out which project we are in
+        references = cmds.ls(type = "reference")
+	for ref in references:
+	    try:
+		self.project = cmds.referenceQuery(ref, filename = True).rpartition("Projects/")[2].partition("/")[0]
+		
+	    except:
+		pass
+        
+        
+        #get access to our maya tools
+        toolsPath = cmds.internalVar(usd = True) + "mayaTools.txt"
+        if os.path.exists(toolsPath):
+            
+            f = open(toolsPath, 'r')
+            self.mayaToolsDir = f.readline()
+            f.close()
+            
+            
+        #check to see if window exists, if so, delete
+        if cmds.window("importMotionUI", exists = True):
+            cmds.deleteUI("importMotionUI")
+            
+        
+        #build window
+        self.widgets["window"] = cmds.window("importMotionUI", w = 700, h = 400, title = "Import Motion", sizeable = False)
+        
+        #create the main layout
+        self.widgets["topLevelLayout"] = cmds.columnLayout()
+        
+        #create the rowColumnLayout (left side for the different ways one can import motion, the right side with that method's settings
+        self.widgets["rowColumnLayout"] = cmds.rowColumnLayout(w = 700, h = 400, nc = 2, cw = [(1, 150), (2, 550)], parent = self.widgets["topLevelLayout"])
+        
+        #create the columnLayout for the left side
+        self.widgets["leftSideButtonColumn"] = cmds.columnLayout(w = 150, h = 400, parent = self.widgets["rowColumnLayout"], cat = ["both", 5], rs = 5)
+        
+        
+        #and create the frame layout for the right side
+        self.widgets["rightSideFrame"] = cmds.frameLayout(w = 550, h = 400, collapsable = False, borderStyle = "etchedIn", labelVisible = False, parent = self.widgets["rowColumnLayout"])
+        
+        
+        
+        #create the buttons for the different methods of importing motion
+        self.widgets["importMotionMethods"] = cmds.iconTextRadioCollection()
+        self.widgets["importMotion_mocap"] = cmds.iconTextRadioButton(select = True,  w = 140, h = 50,  parent = self.widgets["leftSideButtonColumn"], image = self.mayaToolsDir + "/General/Icons/ART/importMocap_off.bmp", selectionImage = self.mayaToolsDir + "/General/Icons/ART/importMocap_on.bmp")
+        self.widgets["importMotion_anims"] = cmds.iconTextRadioButton(select = False,  w = 140, h = 50,  parent = self.widgets["leftSideButtonColumn"], image = self.mayaToolsDir + "/General/Icons/ART/importAnim_off.bmp", selectionImage = self.mayaToolsDir + "/General/Icons/ART/importAnim_on.bmp")
+        
+        
+        #create the elements for the right column 
+        # # # # # # # # # # # # # # # # # # # # # # # # # # # # # # # # # # # # # # # # # # # # # # # # # # # # # # # # # # # # # # # 
+        
+        # # # # # # # # # # # # # # # # # # # # # # # # # # # # # # # # # # # # # # # # # # # # # # # # # # # # # # # # # # # # # # # 
+    
+        # # # # # # # # # # # # # # # # # # # # # # # # # # # # # # # # # # # # # # # # # # # # # # # # # # # # # # # # # # # # # # # 
+        #IMPORT MOCAP
+        self.widgets["importMocapForm"] = cmds.formLayout(w = 500, h = 400, parent = self.widgets["rightSideFrame"])
+        
+        #text labels
+        fbxLabel = cmds.text(label = "FBX File:", font = "boldLabelFont")
+        importMethodLabel = cmds.text(label = "Import Method:", font = "boldLabelFont")
+        frameOffsetLabel = cmds.text(label = "Frame Offset:", font = "boldLabelFont")
+        
+        cmds.formLayout(self.widgets["importMocapForm"], edit = True, af = [(fbxLabel, "top", 13), (fbxLabel, "left", 10)])
+        cmds.formLayout(self.widgets["importMocapForm"], edit = True, af = [(importMethodLabel, "top", 140), (importMethodLabel, "left", 10)])
+        cmds.formLayout(self.widgets["importMocapForm"], edit = True, af = [(frameOffsetLabel, "top", 210), (frameOffsetLabel, "left", 10)])
+        
+        #fbxImport
+        self.widgets["fbxImportTextField"] = cmds.textField(w = 400, text = "", enable = True)
+        self.widgets["fbxImportBrowseButton"] = cmds.symbolButton(w = 30, h = 30, image = self.mayaToolsDir + "/General/Icons/ART/browse.bmp", c = self.fbxBrowse)
+        
+        cmds.formLayout(self.widgets["importMocapForm"], edit = True, af = [(self.widgets["fbxImportTextField"], "top", 10), (self.widgets["fbxImportTextField"], "left", 70)])
+        cmds.formLayout(self.widgets["importMocapForm"], edit = True, af = [(self.widgets["fbxImportBrowseButton"], "top", 5), (self.widgets["fbxImportBrowseButton"], "left", 475)])
+        
+        #character list
+        self.widgets["importMocap_characterList"] = cmds.optionMenu(w = 240, h = 50)
+        self.widgets["importMocap_characterThumb"] = cmds.image(w = 50, h = 50)
+        cmds.formLayout(self.widgets["importMocapForm"], edit = True, af = [(self.widgets["importMocap_characterList"], "top", 45), (self.widgets["importMocap_characterList"], "left", 10)])
+        cmds.formLayout(self.widgets["importMocapForm"], edit = True, af = [(self.widgets["importMocap_characterThumb"], "top", 45), (self.widgets["importMocap_characterThumb"], "left", 255)])
+
+
+        #import method
+        self.widgets["importMethodRadioCollection"] = cmds.radioCollection()
+        self.widgets["importMethod_FK"] = cmds.radioButton(label = "FK", select = True)
+        self.widgets["importMethod_IK"] = cmds.radioButton(label = "IK")
+        self.widgets["importMethod_Both"] = cmds.radioButton(label = "Both")
+        
+        
+        cmds.formLayout(self.widgets["importMocapForm"], edit = True, af = [(self.widgets["importMethod_FK"], "top", 165), (self.widgets["importMethod_FK"], "left", 10)])
+        cmds.formLayout(self.widgets["importMocapForm"], edit = True, af = [(self.widgets["importMethod_IK"], "top", 165), (self.widgets["importMethod_IK"], "left", 80)])
+        cmds.formLayout(self.widgets["importMocapForm"], edit = True, af = [(self.widgets["importMethod_Both"], "top", 165), (self.widgets["importMethod_Both"], "left", 150)])
+        
+        #frame offset
+        self.widgets["frameOffsetField"] = cmds.intField(value=0, w = 100)
+        cmds.formLayout(self.widgets["importMocapForm"], edit = True, af = [(self.widgets["frameOffsetField"], "top", 235), (self.widgets["frameOffsetField"], "left", 10)])
+        
+        #apply motion to parts
+        self.widgets["importMotionTo_Frame"] = cmds.frameLayout( w= 220, h = 350, label = "Apply To Which Parts:", bs = "etchedIn", collapsable = False)
+        cmds.formLayout(self.widgets["importMocapForm"], edit = True, af = [(self.widgets["importMotionTo_Frame"], "top", 40), (self.widgets["importMotionTo_Frame"], "right", 10)])
+        self.widgets["importMotionTo_Form"] = cmds.formLayout( w= 240, h = 280, parent = self.widgets["importMotionTo_Frame"])
+        
+        self.widgets["importMotionTo_HeadButton"] = cmds.iconTextCheckBox(w = 55, h = 55, image = self.mayaToolsDir + "/General/Icons/ART/importMotion_head.bmp", selectionImage = self.mayaToolsDir + "/General/Icons/ART/importMotion_head_on.bmp", value = True)
+        self.widgets["importMotionTo_SpineButton"] = cmds.iconTextCheckBox(w = 55, h = 100, image = self.mayaToolsDir + "/General/Icons/ART/importMotion_torso.bmp", selectionImage = self.mayaToolsDir + "/General/Icons/ART/importMotion_torso_on.bmp", value = True)
+        self.widgets["importMotionTo_lArmButton"] = cmds.iconTextCheckBox(w = 30, h = 100, image = self.mayaToolsDir + "/General/Icons/ART/importMotion_arm.bmp", selectionImage = self.mayaToolsDir + "/General/Icons/ART/importMotion_arm_on.bmp", value = True)
+        self.widgets["importMotionTo_rArmButton"] = cmds.iconTextCheckBox(w = 30, h = 100, image = self.mayaToolsDir + "/General/Icons/ART/importMotion_arm_r.bmp", selectionImage = self.mayaToolsDir + "/General/Icons/ART/importMotion_arm_r_on.bmp", value = True)
+        self.widgets["importMotionTo_lLegButton"] = cmds.iconTextCheckBox(w = 30, h = 110, image = self.mayaToolsDir + "/General/Icons/ART/importMotion_leg.bmp", selectionImage = self.mayaToolsDir + "/General/Icons/ART/importMotion_leg_on.bmp", value = True)
+        self.widgets["importMotionTo_rLegButton"] = cmds.iconTextCheckBox(w = 30, h = 110, image = self.mayaToolsDir + "/General/Icons/ART/importMotion_leg.bmp", selectionImage = self.mayaToolsDir + "/General/Icons/ART/importMotion_leg_on.bmp", value = True)
+        
+        cmds.formLayout(self.widgets["importMotionTo_Form"], edit = True, af = [(self.widgets["importMotionTo_HeadButton"], "top", 25), (self.widgets["importMotionTo_HeadButton"], "right", 100)])
+        cmds.formLayout(self.widgets["importMotionTo_Form"], edit = True, af = [(self.widgets["importMotionTo_SpineButton"], "top", 80), (self.widgets["importMotionTo_SpineButton"], "right", 100)])
+        cmds.formLayout(self.widgets["importMotionTo_Form"], edit = True, af = [(self.widgets["importMotionTo_lArmButton"], "top", 80), (self.widgets["importMotionTo_lArmButton"], "right", 70)])
+        cmds.formLayout(self.widgets["importMotionTo_Form"], edit = True, af = [(self.widgets["importMotionTo_rArmButton"], "top", 80), (self.widgets["importMotionTo_rArmButton"], "right", 155)])
+        cmds.formLayout(self.widgets["importMotionTo_Form"], edit = True, af = [(self.widgets["importMotionTo_lLegButton"], "top", 180), (self.widgets["importMotionTo_lLegButton"], "right", 98)])
+        cmds.formLayout(self.widgets["importMotionTo_Form"], edit = True, af = [(self.widgets["importMotionTo_rLegButton"], "top", 180), (self.widgets["importMotionTo_rLegButton"], "right", 125)])
+        
+        #import button
+        self.widgets["importMocap_importButton"] = cmds.symbolButton(c = self.importMocap, w = 300, h = 50, image = self.mayaToolsDir + "/General/Icons/ART/importMotion_importButton.bmp", parent = self.widgets["importMocapForm"])
+        cmds.formLayout(self.widgets["importMocapForm"], edit = True, af = [(self.widgets["importMocap_importButton"], "bottom", 10), (self.widgets["importMocap_importButton"], "left", 10)])
+	
+	
+	#heel solve checkbox
+	self.widgets["heelSolverCB"] = cmds.checkBox(label = "Solve Foot Roll", v = False, parent = self.widgets["importMocapForm"])
+	self.widgets["kneeSolverCB"] = cmds.checkBox(label = "Solve Knee Vectors", v = True, parent = self.widgets["importMocapForm"])
+	cmds.formLayout(self.widgets["importMocapForm"], edit = True, af = [(self.widgets["heelSolverCB"], "bottom", 110), (self.widgets["heelSolverCB"], "left", 10)])
+	cmds.formLayout(self.widgets["importMocapForm"], edit = True, af = [(self.widgets["kneeSolverCB"], "bottom", 110), (self.widgets["kneeSolverCB"], "left", 140)])
+
+
+
+	# # # # # # # # # # # # # # # # # # # # # # # # # # # # # # # # # # # # # # # # # # # # # # # # 
+	#IMPORT ANIMATION
+	self.widgets["importAnimationFormLayout"] = cmds.formLayout(w = 500, h = 400, parent = self.widgets["rightSideFrame"], visible = False)
+	
+	
+        #character list
+        self.widgets["importAnim_characterList"] = cmds.optionMenu(w = 200, h = 30)
+        self.widgets["importAnim_characterThumb"] = cmds.image(w = 50, h = 50)
+        cmds.formLayout(self.widgets["importAnimationFormLayout"], edit = True, af = [(self.widgets["importAnim_characterList"], "top", 30), (self.widgets["importAnim_characterList"], "left", 275)])
+        cmds.formLayout(self.widgets["importAnimationFormLayout"], edit = True, af = [(self.widgets["importAnim_characterThumb"], "top", 30), (self.widgets["importAnim_characterThumb"], "right", 10)])
+
+
+
+	#create the projects drop down
+	label = cmds.text(label = "Projects:", align = 'right')
+	self.widgets["importAnimProjectsList"] = cmds.optionMenu(w = 250,h = 30, parent = self.widgets["importAnimationFormLayout"], cc = self.getProjCategories)
+	cmds.formLayout(self.widgets["importAnimationFormLayout"], edit = True, af = [(label, "top", 10), (label, "left", 10)])
+	cmds.formLayout(self.widgets["importAnimationFormLayout"], edit = True, af = [(self.widgets["importAnimProjectsList"], "top", 30), (self.widgets["importAnimProjectsList"], "left", 10)])
+	
+	
+	#create the categories layout
+	self.widgets["categoriesList_topLayout"] = cmds.frameLayout(w = 250, h = 300, bs = "etchedIn", cll = False, cl = False, lv = False, parent = self.widgets["importAnimationFormLayout"])
+	cmds.formLayout(self.widgets["importAnimationFormLayout"], edit = True, af = [(self.widgets["categoriesList_topLayout"], "bottom", 10), (self.widgets["categoriesList_topLayout"], "left", 10)])
+	self.widgets["categoriesList_scrollLayout"] = cmds.scrollLayout(w = 240, h = 300, hst = 0, parent = self.widgets["categoriesList_topLayout"])
+	self.widgets["categoriesList_columnLayout"] = cmds.columnLayout(w = 220, parent = self.widgets["categoriesList_scrollLayout"])
+	
+	
+	#create the animation list layout
+	self.widgets["animList_topLayout"] = cmds.frameLayout(w = 260, h = 300, bs = "etchedIn", cll = False, cl = False, lv = False, parent = self.widgets["importAnimationFormLayout"], bgc = [.2, .2, .2])
+	cmds.formLayout(self.widgets["importAnimationFormLayout"], edit = True, af = [(self.widgets["animList_topLayout"], "bottom", 10), (self.widgets["animList_topLayout"], "right", 10)])
+	self.widgets["animList_scrollLayout"] = cmds.scrollLayout(w = 260, h = 300, hst = 0, parent = self.widgets["animList_topLayout"], bgc = [.2, .2, .2])
+	self.widgets["animList_columnLayout"] = cmds.columnLayout(w = 220, parent = self.widgets["animList_scrollLayout"], bgc = [.2, .2, .2])
+	
+        # # # # # # # # # # # # # # # # # # # # # # # # # # # # # # # # # # # # # # # # # # # # # # # # 
+        #Edit radio button commands
+	cmds.iconTextRadioButton(self.widgets["importMotion_mocap"], edit = True, onc = partial(self.switchMode, "fbx"))
+	cmds.iconTextRadioButton(self.widgets["importMotion_anims"], edit = True, onc = partial(self.switchMode, "anim"))
+	
+	
+        #show the window
+        cmds.showWindow(self.widgets["window"])
+        
+        #populate the dropdown with the characters
+        self.getCharacters()
+        self.changeActiveCharacter()
+	
+	#populate the import animations project list
+	self.getProjects()
+	
+	
+    # # # # # # # # # # # # # # # # # # # # # # # # # # # # # # # # # # # # # # # # # # # # # # # # # # # # # # # # # # # # # # # 
+    # # # # # # # # # # # # # # # # # # # # # # # # # # # # # # # # # # # # # # # # # # # # # # # # # # # # # # # # # # # # # # # 
+    # # # # # # # # # # # # # # # # # # # # # # # # # # # # # # # # # # # # # # # # # # # # # # # # # # # # # # # # # # # # # # # 
+    # # # # # # # # # # # # # # # # # # # # # # # # # # # # # # # # # # # # # # # # # # # # # # # # # # # # # # # # # # # # # # # 
+    # # # # # # # # # # # # # # # # # # # # # # # # # # # # # # # # # # # # # # # # # # # # # # # # # # # # # # # # # # # # # # #
+    def getProjects(self, *args):
+	
+	
+	projectPath = self.mayaToolsDir + "/General/ART/Projects/"
+	projects = os.listdir(projectPath)
+	
+	for proj in projects:
+	    cmds.menuItem(label = proj, parent = self.widgets["importAnimProjectsList"])
+	
+	
+	#set to favorite if it exists
+	settingsLocation = self.mayaToolsDir + "/General/Scripts/projectSettings.txt"
+	if os.path.exists(settingsLocation):
+            f = open(settingsLocation, 'r')
+            settings = cPickle.load(f)
+	    favoriteProject = settings.get("FavoriteProject")
+	    
+	    try:
+		cmds.optionMenu(self.widgets["importAnimProjectsList"], edit = True, v = favoriteProject)	
+		
+	    except:
+		pass
+	    
+	    
+	    
+        self.getProjCategories()
+	
+    # # # # # # # # # # # # # # # # # # # # # # # # # # # # # # # # # # # # # # # # # # # # # # # # # # # # # # # # # # # # # # # 
+    # # # # # # # # # # # # # # # # # # # # # # # # # # # # # # # # # # # # # # # # # # # # # # # # # # # # # # # # # # # # # # # 
+    # # # # # # # # # # # # # # # # # # # # # # # # # # # # # # # # # # # # # # # # # # # # # # # # # # # # # # # # # # # # # # # 
+    # # # # # # # # # # # # # # # # # # # # # # # # # # # # # # # # # # # # # # # # # # # # # # # # # # # # # # # # # # # # # # # 
+    # # # # # # # # # # # # # # # # # # # # # # # # # # # # # # # # # # # # # # # # # # # # # # # # # # # # # # # # # # # # # # #
+    def getProjCategories(self, *args):
+	
+	#clear out all children first
+	children = cmds.columnLayout(self.widgets["categoriesList_columnLayout"], q = True, childArray = True)
+	if children != None:
+	    for child in children:
+		cmds.deleteUI(child)
+		
+		
+	selectedProj = cmds.optionMenu(self.widgets["importAnimProjectsList"], q = True, v = True)
+	categoryPath = self.mayaToolsDir + "/General/ART/Projects/" + selectedProj + "/Animations/"
+	
+	if not os.path.exists(categoryPath):
+	    os.makedirs(categoryPath)
+	
+	categories = os.listdir(categoryPath)
+	
+	self.widgets["animationCategories"] = cmds.iconTextRadioCollection()
+	
+	for item in categories:
+	    self.createCategoryEntry(item, selectedProj)
+	
+	
+    # # # # # # # # # # # # # # # # # # # # # # # # # # # # # # # # # # # # # # # # # # # # # # # # # # # # # # # # # # # # # # # 
+    # # # # # # # # # # # # # # # # # # # # # # # # # # # # # # # # # # # # # # # # # # # # # # # # # # # # # # # # # # # # # # # 
+    # # # # # # # # # # # # # # # # # # # # # # # # # # # # # # # # # # # # # # # # # # # # # # # # # # # # # # # # # # # # # # # 
+    # # # # # # # # # # # # # # # # # # # # # # # # # # # # # # # # # # # # # # # # # # # # # # # # # # # # # # # # # # # # # # # 
+    # # # # # # # # # # # # # # # # # # # # # # # # # # # # # # # # # # # # # # # # # # # # # # # # # # # # # # # # # # # # # # #
+    def createCategoryEntry(self, categoryName, project, *args):
+	
+
+	cmds.iconTextRadioButton(onc = partial(self.getAnimations, categoryName, project), parent = self.widgets["categoriesList_columnLayout"], image = "menuIconFile.png", w = 220, h = 30, style = "iconAndTextHorizontal", label = categoryName, cl = self.widgets["animationCategories"], sl =True)
+	
+	#get animations for seleted category
+	self.getAnimations(categoryName, project)
+	
+    # # # # # # # # # # # # # # # # # # # # # # # # # # # # # # # # # # # # # # # # # # # # # # # # # # # # # # # # # # # # # # # 
+    # # # # # # # # # # # # # # # # # # # # # # # # # # # # # # # # # # # # # # # # # # # # # # # # # # # # # # # # # # # # # # # 
+    # # # # # # # # # # # # # # # # # # # # # # # # # # # # # # # # # # # # # # # # # # # # # # # # # # # # # # # # # # # # # # # 
+    # # # # # # # # # # # # # # # # # # # # # # # # # # # # # # # # # # # # # # # # # # # # # # # # # # # # # # # # # # # # # # # 
+    # # # # # # # # # # # # # # # # # # # # # # # # # # # # # # # # # # # # # # # # # # # # # # # # # # # # # # # # # # # # # # #
+    def getAnimations(self, categoryName, project, *args):
+	
+	#clear out all animation files first
+	children = cmds.columnLayout(self.widgets["animList_columnLayout"], q = True, childArray = True)
+	if children != None:
+	    for child in children:
+		cmds.deleteUI(child)
+	    
+	    
+	#get animations and populate UI
+	animFiles = os.listdir(self.mayaToolsDir + "/General/ART/Projects/" + project + "/Animations/" + categoryName + "/")
+	
+	for file in animFiles:
+	    niceName = file.partition(".")[0]
+	    button = cmds.iconTextButton( parent = self.widgets["animList_columnLayout"], image = "ghostOff.png", w = 220, h = 30, bgc = [.2, .2, .2], style = "iconAndTextHorizontal", label = niceName, ann = (project + ", " + categoryName))
+
+	    #create the popup menu for the button
+	    menu = cmds.popupMenu(b = 1, parent =button)
+	    cmds.menuItem(label = "Import Options for " + file.partition(".")[0] + " animation:", parent = menu, enable = False)
+	    cmds.menuItem(divider = True, parent = menu)
+	    cmds.menuItem(label = "Import All Data", parent = menu, c = partial(self.importAnimation, file, project, categoryName))
+	    cmds.menuItem(label = "Import Onto Selected Controls", parent = menu, c = partial(self.importAnimationOnSelection, file, project, categoryName))
+
+    # # # # # # # # # # # # # # # # # # # # # # # # # # # # # # # # # # # # # # # # # # # # # # # # # # # # # # # # # # # # # # # 
+    # # # # # # # # # # # # # # # # # # # # # # # # # # # # # # # # # # # # # # # # # # # # # # # # # # # # # # # # # # # # # # # 
+    # # # # # # # # # # # # # # # # # # # # # # # # # # # # # # # # # # # # # # # # # # # # # # # # # # # # # # # # # # # # # # # 
+    # # # # # # # # # # # # # # # # # # # # # # # # # # # # # # # # # # # # # # # # # # # # # # # # # # # # # # # # # # # # # # # 
+    # # # # # # # # # # # # # # # # # # # # # # # # # # # # # # # # # # # # # # # # # # # # # # # # # # # # # # # # # # # # # # #
+    def importAnimationOnSelection(self, animationName, project, categoryName, *args):
+	
+	selected = cmds.ls(sl = True)
+	
+	animPath = self.mayaToolsDir + "/General/ART/Projects/" + project + "/Animations/" + categoryName + "/" + animationName
+	character = cmds.optionMenu(self.widgets["importAnim_characterList"], q = True, v = True)
+	
+	f = open(animPath, 'r')
+	animData = cPickle.load(f)
+	f.close()
+	
+	#create a progress window
+	progressWindow = cmds.progressWindow(title='Importing Animation', progress = 10, status='Importing...', isInterruptable=True )
+	progressIncrement = 100/len(animData)
+	
+	#go through all of the animData ([layer, control, [curves and key pairings]])
+	for data in animData:
+	    #cmds.progressWindow(edit = True, progress = progressIncrement, status='Importing...')
+	    
+	    #sort out the incoming data
+	    control = data[0]
+	    keyInfo = data[1:]
+	    layers = []
+	    
+	    for info in keyInfo:
+		for i in info:
+		    layer = i[0]
+		    if layer != None:
+			layers.append(layer)
+	    
+	    #if an object did have layers, we need to sort out the information
+	    if layers != []:
+		
+		#create needed layers
+		for layer in layers:
+		    if layer != "BaseAnimation":
+			#see if our layer exists, if not create and select
+			try:
+			    cmds.select(character + ":" + control)
+			    if cmds.animLayer(layer, q = True, exists = True) == False:
+				cmds.animLayer(layer, addSelectedObjects = True)
+					
+			    else:
+				cmds.animLayer(layer, edit = True, addSelectedObjects = True)
+				
+			except:
+			    pass
+			
+				
+		
+		#first setup base animation before other layers
+		animationLayersAll = cmds.ls(type = "animLayer")
+		for l in animationLayersAll:
+		    cmds.animLayer(l, edit = True, selected = False)
+		cmds.animLayer("BaseAnimation", edit = True, selected = True)
+		
+		for info in keyInfo:
+		    for i in info:
+			layer = i[0]
+			if layer == None:
+			    layer = "BaseAnimation"
+			    
+			attrs = i[1]
+		    
+		    for attr in attrs:
+			attribute = attr[0]
+			
+			if cmds.objExists(character + ":" + control + "." + attribute):
+			    keys = attr[1]
+			    for key in keys:
+				frame = key[0]
+				value = key[1]
+				
+				#grab tangent info if there was any
+				try: 
+				    tangentInfo = key[2]
+				except:
+				    tangentInfo = [None, None, None, None, None, None]
+				    pass
+				
+
+				
+				if cmds.objExists(character + ":" + control):
+				    if character + ":" + control in selected:
+				    
+					cmds.setKeyframe(character + ":" + control, animLayer = layer, at = attribute, t = frame, value = value, noResolve = True)
+					
+					if tangentInfo[1] != None:
+					    cmds.keyTangent(character + ":" + control + "." + attribute, t = (frame, frame), itt = tangentInfo[1])
+					    
+					if tangentInfo[1] != None:
+					    cmds.keyTangent(character + ":" + control + "." + attribute, t = (frame, frame), ott = tangentInfo[1])
+					    
+					if tangentInfo[2] != None:
+					    cmds.keyTangent(character + ":" + control + "." + attribute, t = (frame, frame), a = True,  ia = tangentInfo[2])
+					    
+					if tangentInfo[3] != None:
+					    cmds.keyTangent(character + ":" + control + "." + attribute, t = (frame, frame), a = True,  oa = tangentInfo[3])
+					    
+					if tangentInfo[4] != None:
+					    cmds.keyTangent(character + ":" + control + "." + attribute, t = (frame, frame), iw = tangentInfo[4])
+					    
+					if tangentInfo[5] != None:
+					    cmds.keyTangent(character + ":" + control + "." + attribute, t = (frame, frame), ow = tangentInfo[5])	
+					    
+					
+					#refresh scene
+					cmds.select(character + ":" + control)
+					cmds.setToolTo('moveSuperContext')
+					cmds.select(clear = True)
+
+				    
+				    
+	    
+	    else:
+		for info in keyInfo:
+		    for i in info:
+			attrs = i[1]
+			
+			for attr in attrs:
+			    attribute = attr[0]
+			    
+			    if cmds.objExists(character + ":" + control + "." + attribute):
+				keys = attr[1]
+				for key in keys:
+				    frame = key[0]
+				    value = key[1]
+				    
+				    #grab tangent info if there was any
+				    try: 
+					tangentInfo = key[2]
+				    except:
+					tangentInfo = [None, None, None, None, None, None]
+					pass				    
+
+				    
+				    if cmds.objExists(character + ":" + control):
+					
+					if character + ":" + control in selected:
+					    
+					    if cmds.animLayer("BaseAnimation", q = True, exists = True):
+						
+						cmds.setKeyframe(character + ":" + control, animLayer = "BaseAnimation", at = attribute, t = frame, value = value, noResolve = True)
+						
+						if tangentInfo[1] != None:
+						    cmds.keyTangent(character + ":" + control + "." + attribute, t = (frame, frame), itt = tangentInfo[1])
+						    
+						if tangentInfo[1] != None:
+						    cmds.keyTangent(character + ":" + control + "." + attribute, t = (frame, frame), ott = tangentInfo[1])
+						    
+						if tangentInfo[2] != None:
+						    cmds.keyTangent(character + ":" + control + "." + attribute, t = (frame, frame), a = True,  ia = tangentInfo[2])
+						    
+						if tangentInfo[3] != None:
+						    cmds.keyTangent(character + ":" + control + "." + attribute, t = (frame, frame), a = True,  oa = tangentInfo[3])
+						    
+						if tangentInfo[4] != None:
+						    cmds.keyTangent(character + ":" + control + "." + attribute, t = (frame, frame), iw = tangentInfo[4])
+						    
+						if tangentInfo[5] != None:
+						    cmds.keyTangent(character + ":" + control + "." + attribute, t = (frame, frame), ow = tangentInfo[5])						
+
+
+					    else:
+						cmds.setKeyframe(character + ":" + control, at = attribute, t = frame, value = value)
+						
+						if tangentInfo[1] != None:
+						    cmds.keyTangent(character + ":" + control + "." + attribute, t = (frame, frame), itt = tangentInfo[1])
+						    
+						if tangentInfo[1] != None:
+						    cmds.keyTangent(character + ":" + control + "." + attribute, t = (frame, frame), ott = tangentInfo[1])
+						    
+						if tangentInfo[2] != None:
+						    cmds.keyTangent(character + ":" + control + "." + attribute, t = (frame, frame), a = True,  ia = tangentInfo[2])
+						    
+						if tangentInfo[3] != None:
+						    cmds.keyTangent(character + ":" + control + "." + attribute, t = (frame, frame), a = True,  oa = tangentInfo[3])
+						    
+						if tangentInfo[4] != None:
+						    cmds.keyTangent(character + ":" + control + "." + attribute, t = (frame, frame), iw = tangentInfo[4])
+						    
+						if tangentInfo[5] != None:
+						    cmds.keyTangent(character + ":" + control + "." + attribute, t = (frame, frame), ow = tangentInfo[5])						
+					    
+		
+		
+		
+
+
+	cmds.progressWindow(progressWindow, endProgress=1)
+	
+	
+    # # # # # # # # # # # # # # # # # # # # # # # # # # # # # # # # # # # # # # # # # # # # # # # # # # # # # # # # # # # # # # # 
+    # # # # # # # # # # # # # # # # # # # # # # # # # # # # # # # # # # # # # # # # # # # # # # # # # # # # # # # # # # # # # # # 
+    # # # # # # # # # # # # # # # # # # # # # # # # # # # # # # # # # # # # # # # # # # # # # # # # # # # # # # # # # # # # # # # 
+    # # # # # # # # # # # # # # # # # # # # # # # # # # # # # # # # # # # # # # # # # # # # # # # # # # # # # # # # # # # # # # # 
+    # # # # # # # # # # # # # # # # # # # # # # # # # # # # # # # # # # # # # # # # # # # # # # # # # # # # # # # # # # # # # # #
+    def importAnimation(self, animationName, project, categoryName, *args):
+	
+	animPath = self.mayaToolsDir + "/General/ART/Projects/" + project + "/Animations/" + categoryName + "/" + animationName
+	character = cmds.optionMenu(self.widgets["importAnim_characterList"], q = True, v = True)
+	
+	f = open(animPath, 'r')
+	animData = cPickle.load(f)
+	f.close()
+	
+	#create a progress window
+	progressWindow = cmds.progressWindow(title='Importing Animation', progress = 10, status='Importing...', isInterruptable=True )
+	progressIncrement = 100/len(animData)
+	
+	#go through all of the animData ([layer, control, [curves and key pairings]])
+	for data in animData:
+	    
+	    #sort out the incoming data
+	    control = data[0]
+	    keyInfo = data[1:]
+	    layers = []
+	    
+	    for info in keyInfo:
+		for i in info:
+		    layer = i[0]
+		    if layer != None:
+			layers.append(layer)
+	    
+	    #if an object did have layers, we need to sort out the information
+	    if layers != []:
+		
+		#create needed layers
+		for layer in layers:
+		    if layer != "BaseAnimation":
+			
+			#see if our layer exists, if not create and select
+			try:
+			    cmds.select(character + ":" + control)
+			    if cmds.animLayer(layer, q = True, exists = True) == False:
+				cmds.animLayer(layer, addSelectedObjects = True)
+					
+			    else:
+				cmds.animLayer(layer, edit = True, addSelectedObjects = True)
+			except:
+			    pass
+			
+				
+		
+		#first setup base animation before other layers
+		animationLayersAll = cmds.ls(type = "animLayer")
+		for l in animationLayersAll:
+		    cmds.animLayer(l, edit = True, selected = False)
+		cmds.animLayer("BaseAnimation", edit = True, selected = True)
+		
+		
+		for info in keyInfo:
+		    #info is the array that includes all of the information for a control per layer. if info > 0, that means that control has animation on more than 1 layer. 
+		    for i in info:
+			#i is the array that has 2 pieces. the first element is the layer, the 2nd element is all of the keyframe data
+			layer = i[0]
+			if layer == None:
+			    layer = "BaseAnimation"
+			    
+			attrs = i[1]
+		    
+			for attr in attrs:
+			    #print layer, attr
+			    
+			    attribute = attr[0]
+			    
+			    if cmds.objExists(character + ":" + control + "." + attribute):
+				keys = attr[1]
+				
+				for key in keys:
+				    frame = key[0]
+				    value = key[1]
+				    
+				    #grab tangent info if there was any
+				    try: 
+					tangentInfo = key[2]
+				    except:
+					tangentInfo = [None, None, None, None, None, None]
+					pass
+				    
+    
+				    
+				    if cmds.objExists(character + ":" + control):
+
+					
+					cmds.setKeyframe(character + ":" + control, animLayer = layer, at = attribute, t = frame, value = value, noResolve = True)
+					
+					if tangentInfo[1] != None:
+					    cmds.keyTangent(character + ":" + control + "." + attribute, t = (frame, frame), itt = tangentInfo[1])
+					    
+					if tangentInfo[1] != None:
+					    cmds.keyTangent(character + ":" + control + "." + attribute, t = (frame, frame), ott = tangentInfo[1])
+					    
+					if tangentInfo[2] != None:
+					    cmds.keyTangent(character + ":" + control + "." + attribute, t = (frame, frame), a = True,  ia = tangentInfo[2])
+					    
+					if tangentInfo[3] != None:
+					    cmds.keyTangent(character + ":" + control + "." + attribute, t = (frame, frame), a = True,  oa = tangentInfo[3])
+					    
+					if tangentInfo[4] != None:
+					    cmds.keyTangent(character + ":" + control + "." + attribute, t = (frame, frame), iw = tangentInfo[4])
+					    
+					if tangentInfo[5] != None:
+					    cmds.keyTangent(character + ":" + control + "." + attribute, t = (frame, frame), ow = tangentInfo[5])	
+					    
+					
+					#refresh scene
+					cmds.select(character + ":" + control)
+					cmds.setToolTo('moveSuperContext')
+					cmds.select(clear = True)
+				    
+	    
+	    else:
+		for info in keyInfo:
+		    for i in info:
+			attrs = i[1]
+			
+			for attr in attrs:
+			    attribute = attr[0]
+			    
+			    if cmds.objExists(character + ":" + control + "." + attribute):
+				keys = attr[1]
+				for key in keys:
+				    frame = key[0]
+				    value = key[1]
+				    
+				    #grab tangent info if there was any
+				    try:
+					tangentInfo = key[2]
+				    except:
+					tangentInfo = [None, None, None, None, None, None]
+					pass
+				    
+				    
+				    
+				    if cmds.objExists(character + ":" + control):
+
+					
+					if cmds.animLayer("BaseAnimation", q = True, exists = True):
+					    cmds.setKeyframe(character + ":" + control, at = attribute, t = frame, value = value, animLayer = "BaseAnimation")
+					    
+					    if tangentInfo[1] != None:
+						cmds.keyTangent(character + ":" + control + "." + attribute, t = (frame, frame), itt = tangentInfo[1])
+						
+					    if tangentInfo[1] != None:
+						cmds.keyTangent(character + ":" + control + "." + attribute, t = (frame, frame), ott = tangentInfo[1])
+						
+					    if tangentInfo[2] != None:
+						cmds.keyTangent(character + ":" + control + "." + attribute, t = (frame, frame), a = True, ia = tangentInfo[2])
+						
+					    if tangentInfo[3] != None:
+						cmds.keyTangent(character + ":" + control + "." + attribute, t = (frame, frame),  a = True, oa = tangentInfo[3])
+						
+					    if tangentInfo[4] != None:
+						cmds.keyTangent(character + ":" + control + "." + attribute, t = (frame, frame), iw = tangentInfo[4])
+						
+					    if tangentInfo[5] != None:
+						cmds.keyTangent(character + ":" + control + "." + attribute, t = (frame, frame), ow = tangentInfo[5])
+
+						
+					else:
+					    cmds.setKeyframe(character + ":" + control, at = attribute, t = frame, value = value)
+					    if tangentInfo[1] != None:
+						cmds.keyTangent(character + ":" + control + "." + attribute, t = (frame, frame), itt = tangentInfo[1])
+								    
+					    if tangentInfo[1] != None:
+						cmds.keyTangent(character + ":" + control + "." + attribute, t = (frame, frame), ott = tangentInfo[1])
+								    
+					    if tangentInfo[2] != None:
+						cmds.keyTangent(character + ":" + control + "." + attribute, t = (frame, frame), a = True, ia = tangentInfo[2])
+									
+					    if tangentInfo[3] != None:
+						cmds.keyTangent(character + ":" + control + "." + attribute, t = (frame, frame), a = True,  oa = tangentInfo[3])
+									
+					    if tangentInfo[4] != None:
+						cmds.keyTangent(character + ":" + control + "." + attribute, t = (frame, frame), iw = tangentInfo[4])
+									
+					    if tangentInfo[5] != None:
+						cmds.keyTangent(character + ":" + control + "." + attribute, t = (frame, frame), ow = tangentInfo[5])		
+		
+		
+
+
+	cmds.progressWindow(progressWindow, endProgress=1)
+    # # # # # # # # # # # # # # # # # # # # # # # # # # # # # # # # # # # # # # # # # # # # # # # # # # # # # # # # # # # # # # # 
+    # # # # # # # # # # # # # # # # # # # # # # # # # # # # # # # # # # # # # # # # # # # # # # # # # # # # # # # # # # # # # # # 
+    # # # # # # # # # # # # # # # # # # # # # # # # # # # # # # # # # # # # # # # # # # # # # # # # # # # # # # # # # # # # # # # 
+    # # # # # # # # # # # # # # # # # # # # # # # # # # # # # # # # # # # # # # # # # # # # # # # # # # # # # # # # # # # # # # # 
+    # # # # # # # # # # # # # # # # # # # # # # # # # # # # # # # # # # # # # # # # # # # # # # # # # # # # # # # # # # # # # # #
+    def switchMode(self, mode, *args):
+        
+        if mode == "fbx":
+            cmds.formLayout(self.widgets["importMocapForm"], edit = True, visible = True)
+            cmds.formLayout(self.widgets["importAnimationFormLayout"], edit = True, visible = False)
+            
+        if mode == "anim":
+            cmds.formLayout(self.widgets["importMocapForm"], edit = True, visible = False)
+            cmds.formLayout(self.widgets["importAnimationFormLayout"], edit = True, visible = True)
+	    
+	    
+    # # # # # # # # # # # # # # # # # # # # # # # # # # # # # # # # # # # # # # # # # # # # # # # # # # # # # # # # # # # # # # # 
+    # # # # # # # # # # # # # # # # # # # # # # # # # # # # # # # # # # # # # # # # # # # # # # # # # # # # # # # # # # # # # # # 
+    # # # # # # # # # # # # # # # # # # # # # # # # # # # # # # # # # # # # # # # # # # # # # # # # # # # # # # # # # # # # # # # 
+    # # # # # # # # # # # # # # # # # # # # # # # # # # # # # # # # # # # # # # # # # # # # # # # # # # # # # # # # # # # # # # # 
+    # # # # # # # # # # # # # # # # # # # # # # # # # # # # # # # # # # # # # # # # # # # # # # # # # # # # # # # # # # # # # # # 
+    def fbxBrowse(self, *args):
+        path = cmds.fileDialog2(fm = 1, fileFilter = "*.fbx", okc = "Select")[0]
+        
+        #edit the text field with the above path passed in
+        cmds.textField(self.widgets["fbxImportTextField"], edit = True, text = path)
+        
+    # # # # # # # # # # # # # # # # # # # # # # # # # # # # # # # # # # # # # # # # # # # # # # # # # # # # # # # # # # # # # # # 
+    # # # # # # # # # # # # # # # # # # # # # # # # # # # # # # # # # # # # # # # # # # # # # # # # # # # # # # # # # # # # # # # 
+    # # # # # # # # # # # # # # # # # # # # # # # # # # # # # # # # # # # # # # # # # # # # # # # # # # # # # # # # # # # # # # # 
+    # # # # # # # # # # # # # # # # # # # # # # # # # # # # # # # # # # # # # # # # # # # # # # # # # # # # # # # # # # # # # # # 
+    # # # # # # # # # # # # # # # # # # # # # # # # # # # # # # # # # # # # # # # # # # # # # # # # # # # # # # # # # # # # # # # 
+    def ikHeelSolve(self, character, start, end, *args):
+        
+        lValues = []
+        rValues = []
+        
+        cmds.progressWindow(self.progWindow, edit=True, progress= 60, status= "Solving IK Foot Roll" )
+        for i in range(int(start), int(end + 1)):
+            cmds.currentTime(i)
+            
+            if cmds.objExists("ball_l"):
+                lBallVal = cmds.getAttr("ball_l.rz")
+                lValues.append(lBallVal)
+                
+            if cmds.objExists("ball_r"):
+                rBallVal = cmds.getAttr("ball_r.rz")
+                rValues.append(rBallVal)
+                
+
+
+                
+	cmds.progressWindow(self.progWindow, edit=True, progress= 80, status= "Solving IK Foot Roll" )
+	x = 0
+        for i in range(int(start), int(end + 1)):
+            cmds.currentTime(i)
+            if cmds.objExists("ball_l"):
+                if lValues[x] > 10:
+                    cmds.setAttr(character + ":ik_foot_anim_l.rx", 0)
+                    cmds.setAttr(character + ":ik_foot_anim_l.ry", 0)
+                    cmds.setAttr(character + ":ik_foot_anim_l.rz", 0)
+                    cmds.setKeyframe(character + ":ik_foot_anim_l")
+                    
+                    cmds.setAttr(character + ":heel_ctrl_l.rz", lValues[x] * -1)
+                    cmds.setKeyframe(character + ":heel_ctrl_l.rz")
+                    
+                    footPos = cmds.xform("foot_l", q = True, ws = True, t = True)
+                    ikFootPos = cmds.xform(character + ":ik_leg_foot_l", q = True, ws = True, t = True)
+                    yDiff = footPos[1] - ikFootPos[1]
+                    zDiff = footPos[2] - ikFootPos[2]
+                    
+                    cmds.xform(character + ":ik_foot_anim_l", r = True, t = [0, yDiff, zDiff])
+                    cmds.setKeyframe(character + ":ik_foot_anim_l")
+
+		else:
+		    cmds.setAttr(character + ":heel_ctrl_l.rz", 0)
+		    cmds.setKeyframe(character + ":heel_ctrl_l.rz")
+                
+                
+            if cmds.objExists("ball_r"):
+                if rValues[x] > 10:
+                    cmds.setAttr(character + ":ik_foot_anim_r.rx", 0)
+                    cmds.setAttr(character + ":ik_foot_anim_r.ry", 0)
+                    cmds.setAttr(character + ":ik_foot_anim_r.rz", 0)
+                    cmds.setKeyframe(character + ":ik_foot_anim_r")
+                    
+                    cmds.setAttr(character + ":heel_ctrl_r.rz", rValues[x] * -1)
+                    cmds.setKeyframe(character + ":heel_ctrl_r.rz")
+                    
+                    footPos = cmds.xform("foot_r", q = True, ws = True, t = True)
+                    ikFootPos = cmds.xform(character + ":ik_leg_foot_r", q = True, ws = True, t = True)
+                    yDiff = footPos[1] - ikFootPos[1]
+                    zDiff = footPos[2] - ikFootPos[2]
+                    
+                    cmds.xform(character + ":ik_foot_anim_r", r = True, t = [0, yDiff, zDiff])
+                    cmds.setKeyframe(character + ":ik_foot_anim_r")
+    
+		else:
+		    cmds.setAttr(character + ":heel_ctrl_r.rz", 0)
+		    cmds.setKeyframe(character + ":heel_ctrl_r.rz")
+		    
+	    #iterate x
+	    x = x + 1
+                
+                
+
+
+                
+
+                
+
+            
+
+        
+        
+    # # # # # # # # # # # # # # # # # # # # # # # # # # # # # # # # # # # # # # # # # # # # # # # # # # # # # # # # # # # # # # # 
+    # # # # # # # # # # # # # # # # # # # # # # # # # # # # # # # # # # # # # # # # # # # # # # # # # # # # # # # # # # # # # # # 
+    # # # # # # # # # # # # # # # # # # # # # # # # # # # # # # # # # # # # # # # # # # # # # # # # # # # # # # # # # # # # # # # 
+    # # # # # # # # # # # # # # # # # # # # # # # # # # # # # # # # # # # # # # # # # # # # # # # # # # # # # # # # # # # # # # # 
+    # # # # # # # # # # # # # # # # # # # # # # # # # # # # # # # # # # # # # # # # # # # # # # # # # # # # # # # # # # # # # # # 
+    def ikKneeSolve(self, character, start, end, *args):
+	
+	
+	#Hide all the things
+	panels = cmds.getPanel(type = 'modelPanel')
+	for panel in panels:
+	    editor = cmds.modelPanel(panel, q = True, modelEditor = True)
+	    
+	    try:
+		cmds.modelEditor(editor, edit = True, interactive = False, displayTextures = False, textures = False, allObjects = False )
+		
+	    except:
+		pass	
+	    
+	
+        startPointR = cmds.xform(character + ":ik_leg_calf_r", q = True, ws = True, t = True)
+        endPointR = cmds.xform("calf_r", q = True, ws = True, t = True)
+        distR = cmds.distanceDimension( sp=(startPointR[0],startPointR[1],startPointR[2]), ep=(endPointR[0], endPointR[1], endPointR[2]) )
+        distRParent = cmds.listRelatives(distR, parent = True)[0]
+	
+        locsR = cmds.listConnections(distR)
+        startLocR = locsR[0]
+        endLocR = locsR[1]
+        
+        cmds.pointConstraint(character + ":ik_leg_calf_r", startLocR)
+        cmds.pointConstraint("calf_r", endLocR)
+        
+        
+        startPointL = cmds.xform(character + ":ik_leg_calf_l", q = True, ws = True, t = True)
+        endPointL = cmds.xform("calf_l", q = True, ws = True, t = True)
+        distL = cmds.distanceDimension( sp=(startPointL[0],startPointL[1],startPointL[2]), ep=(endPointL[0], endPointL[1], endPointL[2]) )
+        distLParent = cmds.listRelatives(distL, parent = True)[0]
+	
+        locsL = cmds.listConnections(distL)
+        startLocL = locsL[0]
+        endLocL = locsL[1]
+        
+        cmds.pointConstraint(character + ":ik_leg_calf_l", startLocL)
+        cmds.pointConstraint("calf_l", endLocL)
+            
+
+	cmds.currentTime(int(start))
+	
+	#get distance between rig knees and mocap knees
+	for i in range(int(start), int(end) + 1):
+	    cmds.currentTime(i)
+	    
+	    distanceR =  cmds.getAttr(distR + ".distance")
+	    distanceL =  cmds.getAttr(distL + ".distance")
+	    
+	    self.checkDistance(character, distL, distanceL, distanceL, "l")
+	    self.checkDistance(character, distR, distanceR, distanceR, "r")
+
+
+	
+
+	
+        #clean up
+        cmds.delete([locsL[0], locsL[1], locsR[0], locsR[1], distL, distR, distRParent, distLParent])
+	
+	
+	#Show all the things
+	panels = cmds.getPanel(type = 'modelPanel')
+	for panel in panels:
+	    editor = cmds.modelPanel(panel, q = True, modelEditor = True)
+	    
+	    try:
+		cmds.modelEditor(editor, edit = True, interactive = True, displayTextures = True, textures = True, allObjects = True )
+		
+	    except:
+		pass		
+    # # # # # # # # # # # # # # # # # # # # # # # # # # # # # # # # # # # # # # # # # # # # # # # # # # # # # # # # # # # # # # # 
+    # # # # # # # # # # # # # # # # # # # # # # # # # # # # # # # # # # # # # # # # # # # # # # # # # # # # # # # # # # # # # # # 
+    # # # # # # # # # # # # # # # # # # # # # # # # # # # # # # # # # # # # # # # # # # # # # # # # # # # # # # # # # # # # # # # 
+    # # # # # # # # # # # # # # # # # # # # # # # # # # # # # # # # # # # # # # # # # # # # # # # # # # # # # # # # # # # # # # # 
+    # # # # # # # # # # # # # # # # # # # # # # # # # # # # # # # # # # # # # # # # # # # # # # # # # # # # # # # # # # # # # # # 
+    def checkDistance(self, character, distanceNode, distanceAttr, originalValue,  side):
+	
+	
+	if distanceAttr > 1:
+	    currentAttr = cmds.getAttr(character + ":ik_foot_anim_" + side + ".knee_twist")
+	    
+	    try:
+		cmds.setAttr(character + ":ik_foot_anim_" + side + ".knee_twist", currentAttr + 1)
+		cmds.setKeyframe(character + ":ik_foot_anim_" + side + ".knee_twist")
+		newDist = cmds.getAttr(distanceNode + ".distance")
+		
+		
+		if newDist < originalValue:
+		    self.checkDistance(character, distanceNode, newDist, newDist, side)
+		    cmds.progressWindow(self.progWindow, edit=True, progress= (cmds.progressWindow(q = True, progress = True) + 3), status= "Solving IK Pole Vectors" )
+		    
+		if newDist > originalValue:
+		    cmds.setAttr(character + ":ik_foot_anim_" + side + ".knee_twist", currentAttr - 2)
+		    cmds.setKeyframe(character + ":ik_foot_anim_" + side + ".knee_twist")
+		    newDist = cmds.getAttr(distanceNode + ".distance")
+		    
+		    self.checkDistance(character, distanceNode, newDist, newDist, side)
+		    cmds.progressWindow(self.progWindow, edit=True, progress= (cmds.progressWindow(q = True, progress = True) + 3), status= "Solving IK Pole Vectors" )
+		    
+
+	    except:
+		pass
+	    
+	#if adding 1 makes dist get smaller, continue adding 1 until the overall distance is less than 1
+	
+	#if subtracting 1 makes dist get smaller, continue subtracting 1 until the overall distance is less than 1
+	    
+		    
+
+	
+    # # # # # # # # # # # # # # # # # # # # # # # # # # # # # # # # # # # # # # # # # # # # # # # # # # # # # # # # # # # # # # # 
+    # # # # # # # # # # # # # # # # # # # # # # # # # # # # # # # # # # # # # # # # # # # # # # # # # # # # # # # # # # # # # # # 
+    # # # # # # # # # # # # # # # # # # # # # # # # # # # # # # # # # # # # # # # # # # # # # # # # # # # # # # # # # # # # # # # 
+    # # # # # # # # # # # # # # # # # # # # # # # # # # # # # # # # # # # # # # # # # # # # # # # # # # # # # # # # # # # # # # # 
+    # # # # # # # # # # # # # # # # # # # # # # # # # # # # # # # # # # # # # # # # # # # # # # # # # # # # # # # # # # # # # # # 
+    def importMocap(self, *args):
+        
+        #get the fbx file
+        filePath = cmds.textField(self.widgets["fbxImportTextField"], q = True, text = True)
+        if not os.path.exists(filePath):
+            cmds.warning("The given file path is not a valid path.")
+            return
+        
+        else:
+
+        
+            #get the active character
+            character = cmds.optionMenu(self.widgets["importMocap_characterList"], q = True, value = True)
+            
+            #duplicate that character's root
+            if cmds.objExists("root"):
+                cmds.warning("There is already a skeleton in the scene with the name \"root\". Aborting")
+                return
+            
+            newSkeleton = cmds.duplicate(character + ":root")
+            
+            #find import method
+            selectedRadioButton = cmds.radioCollection(self.widgets["importMethodRadioCollection"], q = True, select = True)
+            importMethod = cmds.radioButton(selectedRadioButton, q = True, label = True)
+            
+            #find parts that motion will be applied to
+            head = cmds.iconTextCheckBox(self.widgets["importMotionTo_HeadButton"], q = True, value = True)
+            body = cmds.iconTextCheckBox(self.widgets["importMotionTo_SpineButton"], q = True, value = True)
+            leftArm = cmds.iconTextCheckBox(self.widgets["importMotionTo_lArmButton"], q = True, value = True)
+            rightArm = cmds.iconTextCheckBox(self.widgets["importMotionTo_rArmButton"], q = True, value = True)
+            leftLeg = cmds.iconTextCheckBox(self.widgets["importMotionTo_lLegButton"], q = True, value = True)
+            rightLeg = cmds.iconTextCheckBox(self.widgets["importMotionTo_rLegButton"], q = True, value = True)
+            
+
+            if importMethod == "FK":
+                extraControls = self.importMocap_FK(character, head, body, leftArm, rightArm, leftLeg, rightLeg)
+                
+            if importMethod == "IK":
+                extraControls = self.importMocap_IK(character, head, body, leftArm, rightArm, leftLeg, rightLeg)
+	                
+            if importMethod == "Both":
+                extraControls = self.importMocap_FK(character, head, body, leftArm, rightArm, leftLeg, rightLeg)
+                self.importMocap_IK(character, head, body, leftArm, rightArm, leftLeg, rightLeg)
+                
+	    print extraControls
+	    #fingers
+	    self.importMocap_Fingers(character, leftArm, rightArm)
+	    
+	    
+	    
+            #ensure that the scene is in 30fps
+            cmds.currentUnit(time = 'ntsc')
+            cmds.playbackOptions(min = 0, max = 100, animationStartTime = 0, animationEndTime = 100)
+            cmds.currentTime(0)
+            
+            
+            #import the FBX file
+            string = "FBXImportMode -v \"exmerge\";"
+            mel.eval(string)
+            cmds.file(filePath, i = True, prompt = False, force = True)
+            
+            
+	    animLayers = cmds.ls(type = "animLayer")
+	    
+	    if animLayers != []:
+		for layer in animLayers:
+		    cmds.animLayer(layer, edit = True, selected = False)
+		cmds.animLayer("BaseAnimation", edit = True, selected = True, preferred = True)
+	    
+            #snap timeline to length of imported animation
+            firstFrame = cmds.findKeyframe("pelvis", which = 'first')
+            lastFrame = cmds.findKeyframe("pelvis", which = 'last')
+	    if lastFrame == firstFrame:
+		lastFrame = lastFrame + 1
+		
+            cmds.playbackOptions(min = firstFrame, max = lastFrame, animationStartTime = firstFrame, animationEndTime = lastFrame)
+            
+            
+            #bake the animation down onto the controls
+            cmds.select(clear = True)
+            
+            
+            if importMethod == "FK":
+                
+                if body == True:
+                    cmds.select(character + ":body_anim", add = True)
+                    cmds.select(character + ":spine_01_anim", add = True)
+                    cmds.select(character + ":spine_02_anim", add = True)
+                    
+                    if cmds.objExists(character + ":spine_03_anim"):
+                        cmds.select(character + ":spine_03_anim", add = True)
+                        
+                    if cmds.objExists(character + ":spine_04_anim"):
+                        cmds.select(character + ":spine_04_anim", add = True)
+                        
+                    if cmds.objExists(character + ":spine_05_anim"):
+                        cmds.select(character + ":spine_05_anim", add = True)
+                        
+                        
+                if head == True:
+                    cmds.select(character + ":head_fk_anim", add = True)
+                    
+                    if cmds.objExists(character + ":neck_01_fk_anim"):
+                        cmds.select(character + ":neck_01_fk_anim", add = True)
+                        
+                    if cmds.objExists(character + ":neck_02_fk_anim"):
+                        cmds.select(character + ":neck_02_fk_anim", add = True)
+                        
+                    if cmds.objExists(character + ":neck_03_fk_anim"):
+                        cmds.select(character + ":neck_03_fk_anim", add = True)
+                        
+                
+                if leftArm == True:
+                    cmds.select(character + ":clavicle_l_anim", add = True)
+                    cmds.select(character + ":fk_arm_l_anim", add = True)
+                    cmds.select(character + ":fk_elbow_l_anim", add = True)
+                    cmds.select(character + ":fk_wrist_l_anim", add = True)
+                    
+                    
+                if rightArm == True:
+                    cmds.select(character + ":clavicle_r_anim", add = True)
+                    cmds.select(character + ":fk_arm_r_anim", add = True)
+                    cmds.select(character + ":fk_elbow_r_anim", add = True)
+                    cmds.select(character + ":fk_wrist_r_anim", add = True)
+                    
+                
+                if leftLeg:
+                    cmds.select(character + ":fk_thigh_l_anim", add = True)
+                    cmds.select(character + ":fk_calf_l_anim", add = True)
+                    cmds.select(character + ":fk_foot_l_anim", add = True)
+                    
+                    if cmds.objExists("ball_l"):
+                        cmds.select(character + ":fk_ball_l_anim", add = True)
+                        
+
+                if rightLeg:
+                    cmds.select(character + ":fk_thigh_r_anim", add = True)
+                    cmds.select(character + ":fk_calf_r_anim", add = True)
+                    cmds.select(character + ":fk_foot_r_anim", add = True)
+                    
+                    if cmds.objExists("ball_r"):
+                        cmds.select(character + ":fk_ball_r_anim", add = True)
+                        
+                        
+            if importMethod == "IK":
+                
+                if body == True:
+                    cmds.select(character + ":body_anim", add = True)
+                    cmds.select(character + ":chest_ik_anim", add = True)
+                    cmds.select(character + ":mid_ik_anim", add = True)
+
+                        
+                if head == True:
+                    cmds.select(character + ":head_fk_anim", add = True)
+                    
+                    if cmds.objExists(character + ":neck_01_fk_anim"):
+                        cmds.select(character + ":neck_01_fk_anim", add = True)
+                        
+                    if cmds.objExists(character + ":neck_02_fk_anim"):
+                        cmds.select(character + ":neck_02_fk_anim", add = True)
+                        
+                    if cmds.objExists(character + ":neck_03_fk_anim"):
+                        cmds.select(character + ":neck_03_fk_anim", add = True)
+                        
+                if leftArm == True:
+                    cmds.select(character + ":ik_wrist_l_anim", add = True)
+                    cmds.select(character + ":ik_elbow_l_anim", add = True)
+
+                if rightArm == True:
+                    cmds.select(character + ":ik_wrist_r_anim", add = True)
+                    cmds.select(character + ":ik_elbow_r_anim", add = True)
+
+                if leftLeg == True:
+                    cmds.select(character + ":ik_foot_anim_l", add = True)
+                    
+                    
+                if rightLeg == True:
+                    cmds.select(character + ":ik_foot_anim_r", add = True)
+ 
+ 
+ 
+            if importMethod == "Both":
+                
+                if body == True:
+                    cmds.select(character + ":body_anim", add = True)
+                    cmds.select(character + ":chest_ik_anim", add = True)
+                    cmds.select(character + ":mid_ik_anim", add = True)
+                    cmds.select(character + ":body_anim", add = True)
+                    cmds.select(character + ":spine_01_anim", add = True)
+                    cmds.select(character + ":spine_02_anim", add = True)
+                    
+                    if cmds.objExists(character + ":spine_03_anim"):
+                        cmds.select(character + ":spine_03_anim", add = True)
+                        
+                    if cmds.objExists(character + ":spine_04_anim"):
+                        cmds.select(character + ":spine_04_anim", add = True)
+                        
+                    if cmds.objExists(character + ":spine_05_anim"):
+                        cmds.select(character + ":spine_05_anim", add = True)
+
+                        
+                if head == True:
+                    cmds.select(character + ":head_fk_anim", add = True)
+                    
+                    if cmds.objExists(character + ":neck_01_fk_anim"):
+                        cmds.select(character + ":neck_01_fk_anim", add = True)
+                        
+                    if cmds.objExists(character + ":neck_02_fk_anim"):
+                        cmds.select(character + ":neck_02_fk_anim", add = True)
+                        
+                    if cmds.objExists(character + ":neck_03_fk_anim"):
+                        cmds.select(character + ":neck_03_fk_anim", add = True)
+                        
+                        
+                if leftArm == True:
+                    cmds.select(character + ":ik_wrist_l_anim", add = True)
+                    cmds.select(character + ":ik_elbow_l_anim", add = True)
+                    cmds.select(character + ":clavicle_l_anim", add = True)
+                    cmds.select(character + ":fk_arm_l_anim", add = True)
+                    cmds.select(character + ":fk_elbow_l_anim", add = True)
+                    cmds.select(character + ":fk_wrist_l_anim", add = True)
+                    
+                if rightArm == True:
+                    cmds.select(character + ":ik_wrist_r_anim", add = True)
+                    cmds.select(character + ":ik_elbow_r_anim", add = True)
+                    cmds.select(character + ":clavicle_r_anim", add = True)
+                    cmds.select(character + ":fk_arm_r_anim", add = True)
+                    cmds.select(character + ":fk_elbow_r_anim", add = True)
+                    cmds.select(character + ":fk_wrist_r_anim", add = True)
+                    
+                if leftLeg == True:
+                    cmds.select(character + ":ik_foot_anim_l", add = True)
+                    
+                    cmds.select(character + ":fk_thigh_l_anim", add = True)
+                    cmds.select(character + ":fk_calf_l_anim", add = True)
+                    cmds.select(character + ":fk_foot_l_anim", add = True)
+                    
+                    if cmds.objExists("ball_l"):
+                        cmds.select(character + ":fk_ball_l_anim", add = True)
+                        
+                if rightLeg == True:
+                    cmds.select(character + ":ik_foot_anim_r", add = True)
+                    
+                    cmds.select(character + ":fk_thigh_r_anim", add = True)
+                    cmds.select(character + ":fk_calf_r_anim", add = True)
+                    cmds.select(character + ":fk_foot_r_anim", add = True)
+                    
+                    if cmds.objExists("ball_r"):
+                        cmds.select(character + ":fk_ball_r_anim", add = True)
+                
+	    
+	    #select fingers:
+	    for side in ["l", "r"]:
+		for finger in ["index", "middle", "ring", "pinky", "thumb"]:
+		    
+		    if cmds.objExists(character + ":" + finger + "_finger_fk_ctrl_1_" + side):
+			cmds.select(character + ":" + finger + "_finger_fk_ctrl_1_" + side, add = True)
+			
+		    if cmds.objExists(character + ":" + finger + "_finger_fk_ctrl_2_" + side):
+			cmds.select(character + ":" + finger + "_finger_fk_ctrl_2_" + side, add = True)
+			
+		    if cmds.objExists(character + ":" + finger + "_finger_fk_ctrl_3_" + side):
+			cmds.select(character + ":" + finger + "_finger_fk_ctrl_3_" + side, add = True)
+		
+		    
+	    #select exttra controls
+	    if len(extraControls) > 0:
+		for ctrl in extraControls:
+		    cmds.select(ctrl, add = True)
+		    
+            #bake simulation
+            cmds.bakeResults(simulation = True, t = (firstFrame, lastFrame))
+            
+            
+            #fix ik knee solve
+            if importMethod == "IK" or importMethod == "Both":
+		if leftLeg == True or rightLeg == True:
+		    self.progWindow = cmds.progressWindow(title = "Import Mocap", progress = 10, status = "Solving IK Pole Vectors")
+		    
+		    val = cmds.checkBox(self.widgets["kneeSolverCB"], q = True, v = True)
+		    if val == True:
+			self.ikKneeSolve(character, firstFrame, lastFrame)
+		    
+		    
+		    val = cmds.checkBox(self.widgets["heelSolverCB"], q = True, v = True)
+		    if val == True:
+			self.ikHeelSolve(character, firstFrame, lastFrame)
+			
+		    cmds.progressWindow(self.progWindow, endProgress=True)
+            
+
+            
+            #apply frame offset
+            offset = cmds.intField(self.widgets["frameOffsetField"], q = True, value = True)
+            
+            cmds.select(clear = True)
+            for control in ["head_fk_anim", "neck_01_fk_anim", "neck_02_fk_anim", "neck_03_fk_anim", "spine_01_anim", "spine_02_anim", "spine_03_anim", "spine_04_anim", "spine_05_anim", "mid_ik_anim", "chest_ik_anim",
+                            "body_anim", "hip_anim", "clavicle_l_anim", "clavicle_r_anim", "fk_arm_l_anim", "fk_arm_r_anim", "fk_elbow_l_anim", "fk_elbow_r_anim", "fk_wrist_l_anim", "fk_wrist_r_anim",
+                            "ik_elbow_l_anim", "ik_elbow_r_anim", "ik_wrist_l_anim", "ik_wrist_r_anim", "fk_thigh_l_anim", "fk_thigh_r_anim", "fk_calf_l_anim", "fk_calf_r_anim", "fk_foot_l_anim", "fk_foot_r_anim",
+                            "fk_ball_l_anim", "fk_ball_r_anim","ik_knee_anim_l", "ik_knee_anim_r", "ik_foot_anim_l", "ik_foot_anim_r", "index_finger_fk_ctrl_1_r", "index_finger_fk_ctrl_2_r", "index_finger_fk_ctrl_3_r",
+	                    "index_finger_fk_ctrl_1_l", "index_finger_fk_ctrl_2_l", "index_finger_fk_ctrl_3_l", "middle_finger_fk_ctrl_1_r", "middle_finger_fk_ctrl_2_r", "middle_finger_fk_ctrl_3_r",
+	                    "middle_finger_fk_ctrl_1_l", "middle_finger_fk_ctrl_2_l", "middle_finger_fk_ctrl_3_l", "ring_finger_fk_ctrl_1_r", "ring_finger_fk_ctrl_2_r", "ring_finger_fk_ctrl_3_r",
+	                    "ring_finger_fk_ctrl_1_l", "ring_finger_fk_ctrl_2_l", "ring_finger_fk_ctrl_3_l", "pinky_finger_fk_ctrl_1_r", "pinky_finger_fk_ctrl_2_r", "pinky_finger_fk_ctrl_3_r",
+	                    "pinky_finger_fk_ctrl_1_l", "pinky_finger_fk_ctrl_2_l", "pinky_finger_fk_ctrl_3_l", "thumb_finger_fk_ctrl_1_r", "thumb_finger_fk_ctrl_2_r", "thumb_finger_fk_ctrl_3_r", "thumb_finger_fk_ctrl_1_l", "thumb_finger_fk_ctrl_2_l", "thumb_finger_fk_ctrl_3_l"]:
+                
+                
+                if cmds.objExists(character + ":" + control):
+                    cmds.select(character + ":" + control, add = True)
+                    
+            cmds.selectKey()
+            cmds.keyframe(edit = True, r = True, tc = offset)
+            firstFrame = cmds.findKeyframe("pelvis", which = 'first')
+            lastFrame = cmds.findKeyframe("pelvis", which = 'last')
+            cmds.playbackOptions(min = firstFrame, max = lastFrame, animationStartTime = firstFrame, animationEndTime = lastFrame)
+            
+            
+            #clean up
+            cmds.select(clear = True)
+            
+            
+            #delete the old skeleton
+            cmds.delete("root")
+	    
+            
+                    
+                
+                
+                    
+            
+  
+            
+          
+    # # # # # # # # # # # # # # # # # # # # # # # # # # # # # # # # # # # # # # # # # # # # # # # # # # # # # # # # # # # # # # # 
+    # # # # # # # # # # # # # # # # # # # # # # # # # # # # # # # # # # # # # # # # # # # # # # # # # # # # # # # # # # # # # # # 
+    # # # # # # # # # # # # # # # # # # # # # # # # # # # # # # # # # # # # # # # # # # # # # # # # # # # # # # # # # # # # # # # 
+    # # # # # # # # # # # # # # # # # # # # # # # # # # # # # # # # # # # # # # # # # # # # # # # # # # # # # # # # # # # # # # # 
+    # # # # # # # # # # # # # # # # # # # # # # # # # # # # # # # # # # # # # # # # # # # # # # # # # # # # # # # # # # # # # # # 
+    def importMocap_IK(self, character, head, body, leftArm, rightArm, leftLeg, rightLeg):
+	
+        if body == True:
+            #switch to IK mode
+            cmds.setAttr(character + ":Rig_Settings.spine_ik", 1)
+            cmds.setAttr(character + ":Rig_Settings.spine_fk", 0)
+            
+            #constraints
+            cmds.parentConstraint("pelvis", character + ":body_anim")
+            
+            if cmds.objExists(character + ":chest_ik_anim"):
+                
+                #find highest spine joint
+                numSpineBones = cmds.getAttr(character + ":Skeleton_Settings.numSpineBones")
+                if numSpineBones == 5:
+                    endSpine = "spine_05"
+                    midSpine = ["spine_03"]
+                    
+                if numSpineBones == 4:
+                    endSpine = "spine_04"
+                    midSpine = ["spine_02", "spine_03"]
+                    
+                if numSpineBones == 3:
+                    endSpine = "spine_03"
+                    midSpine = ["spine_02"]
+                    
+                cmds.parentConstraint(endSpine, character + ":chest_ik_anim")
+                for each in midSpine:
+                    cmds.parentConstraint(each, character + ":mid_ik_anim")
+                    
+                    
+        if head == True:
+            cmds.orientConstraint("neck_01", character + ":neck_01_fk_anim")
+            
+            if cmds.objExists(character + ":neck_02_fk_anim"):
+                cmds.orientConstraint("neck_02", character + ":neck_02_fk_anim")
+                
+            if cmds.objExists(character + ":neck_03_fk_anim"):
+                cmds.orientConstraint("neck_03", character + ":neck_03_fk_anim")
+                
+                
+            cmds.orientConstraint("head", character + ":head_fk_anim")
+            
+        
+        if leftArm == True:
+            #switch to IK mode
+            cmds.setAttr(character + ":Rig_Settings.lArmMode", 1)
+            
+            #constraints
+            cmds.parentConstraint("hand_l", character + ":ik_wrist_l_anim", mo = True)
+            cmds.pointConstraint("lowerarm_l", character + ":ik_elbow_l_anim")
+            
+
+        if rightArm == True:
+            #switch to IK mode
+            cmds.setAttr(character + ":Rig_Settings.rArmMode", 1)
+            
+            #constraints
+            cmds.parentConstraint("hand_r", character + ":ik_wrist_r_anim", mo = True)
+            cmds.pointConstraint("lowerarm_r", character + ":ik_elbow_r_anim")
+            
+        if leftLeg == True:
+	    print "hooking up left leg"
+            #switch to IK mode
+            cmds.setAttr(character + ":Rig_Settings.lLegMode", 1)
+            
+            #constraints
+            cmds.pointConstraint("foot_l", character + ":ik_foot_anim_l")
+	    constraint = cmds.orientConstraint("foot_l", character + ":ik_foot_anim_l")[0]
+	    cmds.setAttr(constraint + ".offsetY", 90)
+
+            
+        if rightLeg == True:
+	    print "hooking up right leg"
+            #switch to IK mode
+            cmds.setAttr(character + ":Rig_Settings.rLegMode", 1)
+            
+            #constraints
+            cmds.pointConstraint("foot_r", character + ":ik_foot_anim_r", mo = True)
+	    constraint = cmds.orientConstraint("foot_r", character + ":ik_foot_anim_r")[0]
+	    cmds.setAttr(constraint + ".offsetX", 180)
+	    cmds.setAttr(constraint + ".offsetY", 90)
+	    
+	    
+	coreJoints = ["root", "pelvis", "spine_01", "spine_02", "spine_03", "spine_04", "spine_05", "neck_01", "neck_02", "neck_03", "upperarm_l", "clavicle_l", "lowerarm_l", "hand_l", "upperarm_r", "clavicle_r", "lowerarm_r", "hand_r", "thigh_l", "calf_l", "foot_l", "ball_l", "thigh_r", "calf_r", "foot_r", "ball_r"]
+	cmds.select("root", hi = True)
+	allJoints = cmds.ls(sl = True)
+	
+	extraControls = []
+	for joint in allJoints:
+	    if joint not in coreJoints:
+		if cmds.objExists(character + ":" + joint + "_anim"):
+		    try:
+			constraint = cmds.parentConstraint(joint, character + ":" + joint + "_anim")
+			extraControls.append(character + ":" + joint + "_anim")
+		    except:
+			pass
+		    
+		    
+	return extraControls
+            
+    # # # # # # # # # # # # # # # # # # # # # # # # # # # # # # # # # # # # # # # # # # # # # # # # # # # # # # # # # # # # # # # 
+    # # # # # # # # # # # # # # # # # # # # # # # # # # # # # # # # # # # # # # # # # # # # # # # # # # # # # # # # # # # # # # # 
+    # # # # # # # # # # # # # # # # # # # # # # # # # # # # # # # # # # # # # # # # # # # # # # # # # # # # # # # # # # # # # # # 
+    # # # # # # # # # # # # # # # # # # # # # # # # # # # # # # # # # # # # # # # # # # # # # # # # # # # # # # # # # # # # # # # 
+    # # # # # # # # # # # # # # # # # # # # # # # # # # # # # # # # # # # # # # # # # # # # # # # # # # # # # # # # # # # # # # # 
+    def importMocap_FK(self, character, head, body, leftArm, rightArm, leftLeg, rightLeg):
+        
+        #setup constraints to FK controls
+        
+        if body == True:
+            #switch to FK mode
+            cmds.setAttr(character + ":Rig_Settings.spine_ik", 0)
+            cmds.setAttr(character + ":Rig_Settings.spine_fk", 1)
+            
+            #constraints
+            cmds.parentConstraint("pelvis", character + ":body_anim")
+            
+            cmds.orientConstraint("spine_01", character + ":spine_01_anim")
+            cmds.orientConstraint("spine_02", character + ":spine_02_anim")
+            
+            if cmds.objExists(character + ":spine_03_anim"):
+                cmds.orientConstraint("spine_03", character + ":spine_03_anim")
+                
+            if cmds.objExists(character + ":spine_04_anim"):
+                cmds.orientConstraint("spine_04", character + ":spine_04_anim")
+                
+            if cmds.objExists(character + ":spine_05_anim"):
+                cmds.orientConstraint("spine_05", character + ":spine_05_anim")
+            
+            
+        if head == True:
+            cmds.orientConstraint("neck_01", character + ":neck_01_fk_anim")
+            
+            if cmds.objExists(character + ":neck_02_fk_anim"):
+                cmds.orientConstraint("neck_02", character + ":neck_02_fk_anim")
+                
+            if cmds.objExists(character + ":neck_03_fk_anim"):
+                cmds.orientConstraint("neck_03", character + ":neck_03_fk_anim")
+                
+                
+            cmds.orientConstraint("head", character + ":head_fk_anim")
+            
+        
+        if leftArm == True:
+            #switch to FK mode
+            cmds.setAttr(character + ":Rig_Settings.lArmMode", 0)
+            
+            cmds.pointConstraint("upperarm_l", character + ":clavicle_l_anim")
+            cmds.orientConstraint("upperarm_l", character + ":fk_arm_l_anim")
+            cmds.orientConstraint("lowerarm_l", character + ":fk_elbow_l_anim")
+            cmds.orientConstraint("hand_l", character + ":fk_wrist_l_anim")
+            
+            
+        if rightArm == True:
+            #switch to FK mode
+            cmds.setAttr(character + ":Rig_Settings.rArmMode", 0)
+            
+            cmds.pointConstraint("upperarm_r", character + ":clavicle_r_anim")
+            cmds.orientConstraint("upperarm_r", character + ":fk_arm_r_anim")
+            cmds.orientConstraint("lowerarm_r", character + ":fk_elbow_r_anim")
+            cmds.orientConstraint("hand_r", character + ":fk_wrist_r_anim")
+            
+            
+        if leftLeg == True:
+            #switch to FK mode
+            cmds.setAttr(character + ":Rig_Settings.lLegMode", 0)
+            
+            cmds.orientConstraint("thigh_l", character + ":fk_thigh_l_anim")
+            cmds.orientConstraint("calf_l", character + ":fk_calf_l_anim")
+            cmds.orientConstraint("foot_l", character + ":fk_foot_l_anim")
+            
+            if cmds.objExists("ball_l"):
+                cmds.orientConstraint("ball_l", character + ":fk_ball_l_anim")
+                
+                
+        if rightLeg == True:
+            #switch to FK mode
+            cmds.setAttr(character + ":Rig_Settings.rLegMode", 0)
+            
+            cmds.orientConstraint("thigh_r", character + ":fk_thigh_r_anim")
+            cmds.orientConstraint("calf_r", character + ":fk_calf_r_anim")
+            cmds.orientConstraint("foot_r", character + ":fk_foot_r_anim")
+            
+            if cmds.objExists("ball_r"):
+                cmds.orientConstraint("ball_r", character + ":fk_ball_r_anim")
+		
+		
+	coreJoints = ["root", "pelvis", "spine_01", "spine_02", "spine_03", "spine_04", "spine_05", "neck_01", "neck_02", "neck_03", "upperarm_l", "clavicle_l", "lowerarm_l", "hand_l", "upperarm_r", "clavicle_r", "lowerarm_r", "hand_r", "thigh_l", "calf_l", "foot_l", "ball_l", "thigh_r", "calf_r", "foot_r", "ball_r"]
+	cmds.select("root", hi = True)
+	allJoints = cmds.ls(sl = True)
+	
+	extraControls = []
+	for joint in allJoints:
+	    if joint not in coreJoints:
+		if cmds.objExists(character + ":" + joint + "_anim"):
+		    try:
+			constraint = cmds.parentConstraint(joint, character + ":" + joint + "_anim")
+			extraControls.append(character + ":" + joint + "_anim")
+		    except:
+			pass
+		    
+		    
+	return extraControls
+		
+	
+	
+            
+        
+    # # # # # # # # # # # # # # # # # # # # # # # # # # # # # # # # # # # # # # # # # # # # # # # # # # # # # # # # # # # # # # # 
+    # # # # # # # # # # # # # # # # # # # # # # # # # # # # # # # # # # # # # # # # # # # # # # # # # # # # # # # # # # # # # # # 
+    # # # # # # # # # # # # # # # # # # # # # # # # # # # # # # # # # # # # # # # # # # # # # # # # # # # # # # # # # # # # # # # 
+    # # # # # # # # # # # # # # # # # # # # # # # # # # # # # # # # # # # # # # # # # # # # # # # # # # # # # # # # # # # # # # # 
+    # # # # # # # # # # # # # # # # # # # # # # # # # # # # # # # # # # # # # # # # # # # # # # # # # # # # # # # # # # # # # # # 
+    def importMocap_Fingers(self, character, leftArm, rightArm):
+
+        if leftArm == True:
+	    for finger in ["index", "middle", "ring", "pinky", "thumb"]:
+		
+		#switch to FK mode
+		try:
+		    cmds.setAttr(character + ":" + finger + "_finger_l_mode_anim.FK_IK", 0)
+		except:
+		    pass
+		
+		
+		#setup constraints
+		
+		if cmds.objExists(character + ":" + finger + "_finger_fk_ctrl_1_l"):
+		    cmds.orientConstraint(finger + "_01_l", character + ":" + finger + "_finger_fk_ctrl_1_l")
+		    
+		if cmds.objExists(character + ":" + finger + "_finger_fk_ctrl_2_l"):
+		    cmds.orientConstraint(finger + "_02_l", character + ":" + finger + "_finger_fk_ctrl_2_l")
+		    
+		if cmds.objExists(character + ":" + finger + "_finger_fk_ctrl_3_l"):
+		    cmds.orientConstraint(finger + "_03_l", character + ":" + finger + "_finger_fk_ctrl_3_l")
+		    
+		    
+        if rightArm == True:
+	    for finger in ["index", "middle", "ring", "pinky", "thumb"]:
+		
+		#switch to FK mode
+		try:
+		    cmds.setAttr(character + ":" + finger + "_finger_r_mode_anim.FK_IK", 0)
+		except:
+		    pass
+		
+		#setup constraints
+		
+		if cmds.objExists(character + ":" + finger + "_finger_fk_ctrl_1_r"):
+		    cmds.orientConstraint(finger + "_01_r", character + ":" + finger + "_finger_fk_ctrl_1_r")
+		    
+		if cmds.objExists(character + ":" + finger + "_finger_fk_ctrl_2_r"):
+		    cmds.orientConstraint(finger + "_02_r", character + ":" + finger + "_finger_fk_ctrl_2_r")
+		    
+		if cmds.objExists(character + ":" + finger + "_finger_fk_ctrl_3_r"):
+		    cmds.orientConstraint(finger + "_03_r", character + ":" + finger + "_finger_fk_ctrl_3_r")
+		
+
+    # # # # # # # # # # # # # # # # # # # # # # # # # # # # # # # # # # # # # # # # # # # # # # # # # # # # # # # # # # # # # # # 
+    # # # # # # # # # # # # # # # # # # # # # # # # # # # # # # # # # # # # # # # # # # # # # # # # # # # # # # # # # # # # # # # 
+    # # # # # # # # # # # # # # # # # # # # # # # # # # # # # # # # # # # # # # # # # # # # # # # # # # # # # # # # # # # # # # # 
+    # # # # # # # # # # # # # # # # # # # # # # # # # # # # # # # # # # # # # # # # # # # # # # # # # # # # # # # # # # # # # # # 
+    # # # # # # # # # # # # # # # # # # # # # # # # # # # # # # # # # # # # # # # # # # # # # # # # # # # # # # # # # # # # # # # 
+    def getCharacters(self):
+        referenceNodes = []
+        references = cmds.ls(type = "reference")
+        
+        for reference in references:
+            niceName = reference.rpartition("RN")[0]
+            suffix = reference.rpartition("RN")[2]
+            if suffix != "":
+                if cmds.objExists(niceName + suffix + ":" + "Skeleton_Settings"):
+                    referenceNodes.append(niceName + suffix)
+                    
+            else:
+                if cmds.objExists(niceName + ":" + "Skeleton_Settings"):
+                    referenceNodes.append(niceName)
+                
+        for node in referenceNodes:
+            cmds.menuItem(label = node, parent = self.widgets["importMocap_characterList"])
+	    cmds.menuItem(label = node, parent = self.widgets["importAnim_characterList"])
+    
+    
+    # # # # # # # # # # # # # # # # # # # # # # # # # # # # # # # # # # # # # # # # # # # # # # # # # # # # # # # # # # # # # # # 
+    # # # # # # # # # # # # # # # # # # # # # # # # # # # # # # # # # # # # # # # # # # # # # # # # # # # # # # # # # # # # # # # 
+    # # # # # # # # # # # # # # # # # # # # # # # # # # # # # # # # # # # # # # # # # # # # # # # # # # # # # # # # # # # # # # # 
+    # # # # # # # # # # # # # # # # # # # # # # # # # # # # # # # # # # # # # # # # # # # # # # # # # # # # # # # # # # # # # # # 
+    # # # # # # # # # # # # # # # # # # # # # # # # # # # # # # # # # # # # # # # # # # # # # # # # # # # # # # # # # # # # # # # 
+    def changeActiveCharacter(self):
+        characterName = cmds.optionMenu(self.widgets["importMocap_characterList"], q = True, value = True)
+        
+        thumbnailPath = self.mayaToolsDir + "/General/Icons/ART/Thumbnails/" + self.project + "/"
+        thumbs = os.listdir(thumbnailPath)
+        for thumb in thumbs:
+            if thumb.find("_small") != -1:
+                if thumb.find(characterName) == 0:
+                    cmds.image(self.widgets["importMocap_characterThumb"], edit = True, image = thumbnailPath + thumb, ann = characterName)
+		    cmds.image(self.widgets["importAnim_characterThumb"], edit = True, image = thumbnailPath + thumb, ann = characterName)
+        
+        
