Adding ART scripts to repository

diff --git a/Engine/Extras/Maya_AnimationRiggingTools/MayaTools/General/Scripts/ART_jointMover.py b/Engine/Extras/Maya_AnimationRiggingTools/MayaTools/General/Scripts/ART_jointMover.py
new file mode 100644
--- /dev/null
+++ b/Engine/Extras/Maya_AnimationRiggingTools/MayaTools/General/Scripts/ART_jointMover.py
@@ -0,0 +1,977 @@
+#GAME ENGINE AUTO RIG
+
+import maya.cmds as cmds
+import maya.mel as mel
+from functools import partial
+import os
+
+
+class JointMover_Build():
+    
+    def __init__(self):
+        
+        
+        #get access to our maya tools
+        toolsPath = cmds.internalVar(usd = True) + "mayaTools.txt"
+        
+        if os.path.exists(toolsPath):
+            
+            f = open(toolsPath, 'r')
+            self.mayaToolsDir = f.readline()
+            f.close()
+            
+            
+        #import in the base joint mover file
+        cmds.file(self.mayaToolsDir + "/General/ART/JointMover.mb", i = True)
+        
+        #build spine controls
+        self.buildJointMover_Spine()
+        self.buildJointMover_Neck()
+        
+        #parent head mover to neck
+        neckBones = cmds.getAttr("Skeleton_Settings.numNeckBones")
+        cmds.parent("head_mover_grp", ("neck_0" + str(int(neckBones)) + "_mover"))
+        
+        #parent clavicles to spine bone
+        spineBones = cmds.getAttr("Skeleton_Settings.numSpineBones")
+        
+
+        if spineBones == 10:
+            cmds.parent("clavicle_mover_l_grp", ("spine_" + str(int(spineBones)) + "_mover"))
+            cmds.parent("clavicle_mover_r_grp", ("spine_" + str(int(spineBones)) + "_mover"))
+                
+        else:
+            cmds.parent("clavicle_mover_l_grp", ("spine_0" + str(int(spineBones)) + "_mover"))
+            cmds.parent("clavicle_mover_r_grp", ("spine_0" + str(int(spineBones)) + "_mover"))
+            
+            
+            
+        #build the hands/fingers
+        self.buildJointMover_Hands("Left")
+        self.buildJointMover_Hands("Right")
+        
+        
+        #check for leg style
+        legStyle = cmds.getAttr("Skeleton_Settings.legStyle")
+        
+        #hind legs
+        if legStyle == 1:
+            
+            #delete existing legs
+            cmds.delete("thigh_mover_l_grp")
+            cmds.delete("thigh_mover_r_grp")
+            
+            #bring in hind legs
+            cmds.file(self.mayaToolsDir + "/General/ART/JointMover_hindLeg.mb", i = True)
+            cmds.parent("thigh_mover_l_grp", "pelvis_mover")
+            cmds.parent("thigh_mover_r_grp", "pelvis_mover")
+            
+            #assign material and delete imported material
+            connections = cmds.listConnections("JointMover_hindLeg_proxy_geo_body_m")
+            for connection in connections:
+                if cmds.nodeType(connection) == "shadingEngine":
+                    subConnections = cmds.listConnections(connection)
+                    geo = []
+                    for conn in subConnections:
+                        if cmds.nodeType(conn) == "transform":
+                            geo.append(conn)
+                            
+                    for geo in geo:
+                        cmds.select(geo)
+                        cmds.select("proxy_geo_body_m", add= True)
+                        cmds.hyperShade(assign = "proxy_geo_body_m")
+                    
+            cmds.delete("JointMover_hindLeg_proxy_geo_body_m")
+            
+            
+        #build feet
+        self.buildJointMover_Feet("Left")
+        self.buildJointMover_Feet("Right")
+            
+	#build twist joints
+	if cmds.getAttr("Skeleton_Settings.leftUpperArmTwist") > 0:
+	    self.buildTwistJoints("upperarm_twist_0", "leftUpperArmTwist", "l", "upperarm_mover_l", "lowerarm_mover_l_grp", "down")
+	    
+	if cmds.getAttr("Skeleton_Settings.rightUpperArmTwist") > 0:
+	    self.buildTwistJoints("upperarm_twist_0", "rightUpperArmTwist", "r", "upperarm_mover_r", "lowerarm_mover_r_grp", "down")
+	
+	if cmds.getAttr("Skeleton_Settings.leftLowerArmTwist") > 0:
+	    self.buildTwistJoints("lowerarm_twist_0", "leftLowerArmTwist", "l", "lowerarm_mover_l", "hand_mover_l_grp", "up")
+	    
+	if cmds.getAttr("Skeleton_Settings.rightLowerArmTwist") > 0:
+	    self.buildTwistJoints("lowerarm_twist_0", "rightLowerArmTwist", "r", "lowerarm_mover_r", "hand_mover_r_grp", "up")
+	
+	if cmds.getAttr("Skeleton_Settings.leftUpperLegTwist") > 0:
+	    self.buildTwistJoints("thigh_twist_0", "leftUpperLegTwist", "l", "thigh_mover_l", "calf_mover_l_grp", "down")
+	    
+	if cmds.getAttr("Skeleton_Settings.rightUpperLegTwist") > 0:
+	    self.buildTwistJoints("thigh_twist_0", "rightUpperLegTwist", "r", "thigh_mover_r", "calf_mover_r_grp", "down")
+            
+	if cmds.getAttr("Skeleton_Settings.leftLowerLegTwist") > 0:
+	    self.buildTwistJoints("calf_twist_0", "leftLowerLegTwist", "l", "calf_mover_l", "foot_mover_l_grp", "up")
+	    
+	if cmds.getAttr("Skeleton_Settings.rightLowerLegTwist") > 0:
+	    self.buildTwistJoints("calf_twist_0", "rightLowerLegTwist", "r", "calf_mover_r", "foot_mover_r_grp", "up")
+	
+	
+        #build extra joints
+        attrs = cmds.listAttr("Skeleton_Settings")
+        self.addExtraJoints(attrs)	
+	
+	
+        #lock down all proxy geo
+        cmds.select("proxy_geo*")
+        selection = cmds.ls(sl = True, geometry = True, transforms = True)
+        
+        for each in selection:
+            if each.find("Shape") == -1:
+                cmds.setAttr(each + ".tx", keyable = False, lock = True)
+                cmds.setAttr(each + ".ty", keyable = False, lock = True)
+                cmds.setAttr(each + ".tz", keyable = False, lock = True)
+                cmds.setAttr(each + ".rx", keyable = False, lock = True)
+                cmds.setAttr(each + ".ry", keyable = False, lock = True)
+                cmds.setAttr(each + ".rz", keyable = False, lock = True)
+                cmds.setAttr(each + ".sx", keyable = False, lock = True)
+                cmds.setAttr(each + ".sy", keyable = False, lock = True)
+                cmds.setAttr(each + ".sz", keyable = False, lock = True)
+                
+                cmds.setAttr(each + ".overrideEnabled", 1)
+                cmds.setAttr(each + ".overrideDisplayType", 2)
+                
+                
+        #find all nodes in the joint mover and lock those nodes
+        cmds.select("root_mover_grp", r = True, hi = True)
+        cmds.select("Skeleton_Settings", add = True)
+        nodes = cmds.ls(sl = True, transforms = True)
+        
+        
+        for node in nodes:
+            cmds.lockNode(node, lock = True)
+        
+        
+   
+        
+    def buildSpineControls(self, pivotJoint, name):
+        
+        #create lattice
+        lattice = cmds.lattice(divisions = (2,2,2), objectCentered = True, ldv = (2,2,2))[1]
+        
+        #find lattice point positions in world space
+        
+        frontTopRight = cmds.xform((lattice + ".pt[0][0][1]"), q = True, ws = True, t = True)
+        backTopRight = cmds.xform((lattice + ".pt[0][1][1]"), q = True, ws = True, t = True)
+        frontTopLeft = cmds.xform((lattice + ".pt[1][0][1]"), q = True, ws = True, t = True)
+        backTopLeft = cmds.xform((lattice + ".pt[1][1][1]"), q = True, ws = True, t = True)
+        frontBottomRight = cmds.xform((lattice + ".pt[0][0][0]"), q = True, ws = True, t = True)
+        backBottomRight = cmds.xform((lattice + ".pt[0][1][0]"), q = True, ws = True, t = True)
+        frontBottomLeft = cmds.xform((lattice + ".pt[1][0][0]"), q = True, ws = True, t = True)
+        backBottomLeft = cmds.xform((lattice + ".pt[1][1][0]"), q = True, ws = True, t = True)
+        
+        #create a curve (cube) using the positions of the lattice points
+        globalMoverControl = cmds.curve(name = (name + "_mover"), d = 1, p=[(frontBottomRight), (frontTopRight), (frontTopLeft), (frontBottomLeft), (frontBottomRight), (frontTopRight), (backTopRight), (backBottomRight), (backBottomLeft), (frontBottomLeft), (frontTopLeft), (backTopLeft), (backBottomLeft), (backBottomRight), (backTopRight), (backTopLeft)])
+        
+        #make correct pivot for global control
+        pivot = cmds.xform(pivotJoint, q = True, ws = True, t = True)
+        cmds.xform(globalMoverControl, piv = pivot)
+        
+        #create global mover group
+        globalMoverGroup = cmds.group(empty = True, name = (name + "_mover_grp"))
+        constraint = cmds.parentConstraint(pivotJoint, globalMoverGroup)[0]
+        cmds.delete(constraint)
+        
+        cmds.parent(globalMoverControl, globalMoverGroup)
+        cmds.setAttr(globalMoverControl + ".overrideEnabled", 1)
+        cmds.setAttr(globalMoverControl + ".overrideColor", 17)
+        
+        #duplicate the globalMoverControl and scale it down
+        offsetMoverControl = cmds.duplicate(globalMoverControl, name = (name + "_mover_offset"))[0]
+        cmds.setAttr(offsetMoverControl + ".scaleX", .85)
+        cmds.setAttr(offsetMoverControl + ".scaleY", .85)
+        cmds.setAttr(offsetMoverControl + ".scaleZ", .85)
+        
+        cmds.xform(offsetMoverControl, piv = pivot)
+        
+        cmds.parent(offsetMoverControl, globalMoverControl)
+        cmds.setAttr(offsetMoverControl + ".overrideEnabled", 1)
+        cmds.setAttr(offsetMoverControl + ".overrideColor", 18)
+        
+        
+        #freeze transforms
+        cmds.makeIdentity(globalMoverControl, apply = True, t = True, r = True, s = True)
+        
+        
+        #create the geo mover control
+        geoMover = cmds.duplicate("geo_mover", name = (name + "_geo_mover"))[0]
+        shapeNode = cmds.listRelatives(geoMover, children = True, shapes = True)[0]
+        cmds.setAttr(geoMover + "|" + shapeNode + ".visibility", 0)
+        
+        constraint = cmds.pointConstraint(lattice, geoMover)[0]
+        cmds.delete(constraint)
+        
+        cmds.parent(geoMover, offsetMoverControl)
+        cmds.makeIdentity(geoMover, apply = True, t = True, r = True, s = True)
+    
+        
+        #create lra controls
+        lraGroup = cmds.group(empty = True, name = (name + "_lra_grp"))
+        constraint = cmds.parentConstraint(pivotJoint, lraGroup)[0]
+        cmds.delete(constraint)
+        cmds.parent(lraGroup, offsetMoverControl)
+        
+        lraCtrl = cmds.duplicate("lra", name = (name + "_lra"))[0]
+        cmds.setAttr(lraCtrl + ".v", 1)
+        constraint = cmds.parentConstraint(lraGroup, lraCtrl)[0]
+        cmds.delete(constraint)
+        cmds.parent(lraCtrl, lraGroup)
+        cmds.makeIdentity(lraCtrl, apply = True, t = True, r = True, s = True)
+        
+        #lock translations and rotations on lra controls
+        cmds.setAttr(lraCtrl + ".tx", lock = True, k = False)
+        cmds.setAttr(lraCtrl + ".ty", lock = True, k = False)
+        cmds.setAttr(lraCtrl + ".tz", lock = True, k = False)
+        cmds.setAttr(lraCtrl + ".rx", lock = True, k = False)
+        cmds.setAttr(lraCtrl + ".ry", lock = True, k = False)
+        cmds.setAttr(lraCtrl + ".rz", lock = True, k = False)
+        
+        
+        #delete lattice
+        cmds.delete(lattice)
+        
+        cmds.container("JointMover", edit = True, addNode = globalMoverGroup, includeNetwork=True,includeHierarchyBelow=True)
+        return [globalMoverGroup, globalMoverControl, offsetMoverControl, geoMover]
+        
+        
+
+    
+    def buildJointMover_Spine(self):
+        
+        #get number of spine bones
+        spineBones = cmds.getAttr("Skeleton_Settings.numSpineBones")
+        
+        
+        if spineBones == 2:
+            
+            #select proxy geo and build controls
+            cmds.select(["proxy_geo_spine_01", "proxy_geo_spine_02"], r = True)
+            spine1Controls = self.buildSpineControls("spine_01_pivot_joint", "spine_01")
+            cmds.parent(["proxy_geo_spine_01", "proxy_geo_spine_02"], spine1Controls[3])
+            
+            cmds.select(["proxy_geo_spine_03", "proxy_geo_spine_04", "proxy_geo_spine_05"], r = True)
+            spine2Controls = self.buildSpineControls("spine_05_pivot_joint", "spine_02")
+            cmds.parent(["proxy_geo_spine_03", "proxy_geo_spine_04", "proxy_geo_spine_05"], spine2Controls[3])
+            
+            cmds.parent(spine2Controls[0], spine1Controls[1])
+            cmds.parent(spine1Controls[0], "pelvis_mover")
+            
+            
+        if spineBones == 3:
+            
+            #select proxy geo and build controls
+            cmds.select(["proxy_geo_spine_01"], r = True)
+            spine1Controls = self.buildSpineControls("spine_01_pivot_joint", "spine_01")
+            cmds.parent(["proxy_geo_spine_01"], spine1Controls[3])
+            
+            cmds.select(["proxy_geo_spine_02", "proxy_geo_spine_03"], r = True)
+            spine2Controls = self.buildSpineControls("spine_04_pivot_joint", "spine_02")
+            cmds.parent(["proxy_geo_spine_02", "proxy_geo_spine_03"], spine2Controls[3])
+            
+            cmds.select(["proxy_geo_spine_04", "proxy_geo_spine_05"], r = True)
+            spine3Controls = self.buildSpineControls("spine_07_pivot_joint", "spine_03")
+            cmds.parent(["proxy_geo_spine_04", "proxy_geo_spine_05"], spine3Controls[3])
+            
+            cmds.parent(spine3Controls[0], spine2Controls[1])
+            cmds.parent(spine2Controls[0], spine1Controls[1])
+            cmds.parent(spine1Controls[0], "pelvis_mover")
+            
+            
+
+        if spineBones == 4:
+            
+            #select proxy geo and build controls
+            cmds.select(["proxy_geo_spine_01"], r = True)
+            spine1Controls = self.buildSpineControls("spine_01_pivot_joint", "spine_01")
+            cmds.parent(["proxy_geo_spine_01"], spine1Controls[3])
+            
+            cmds.select(["proxy_geo_spine_02"], r = True)
+            spine2Controls = self.buildSpineControls("spine_03_pivot_joint", "spine_02")
+            cmds.parent(["proxy_geo_spine_02"], spine2Controls[3])
+            
+            cmds.select(["proxy_geo_spine_03"], r = True)
+            spine3Controls = self.buildSpineControls("spine_05_pivot_joint", "spine_03")
+            cmds.parent(["proxy_geo_spine_03"], spine3Controls[3])
+            
+            cmds.select(["proxy_geo_spine_04", "proxy_geo_spine_05"], r = True)
+            spine4Controls = self.buildSpineControls("spine_07_pivot_joint", "spine_04")
+            cmds.parent(["proxy_geo_spine_04", "proxy_geo_spine_05"], spine4Controls[3])
+            
+            cmds.parent(spine4Controls[0], spine3Controls[1])
+            cmds.parent(spine3Controls[0], spine2Controls[1])
+            cmds.parent(spine2Controls[0], spine1Controls[1])
+            cmds.parent(spine1Controls[0], "pelvis_mover")
+            
+            
+            
+        if spineBones == 5:
+            
+            #select proxy geo and build controls
+            cmds.select(["proxy_geo_spine_01"], r = True)
+            spine1Controls = self.buildSpineControls("spine_01_pivot_joint", "spine_01")
+            cmds.parent(["proxy_geo_spine_01"], spine1Controls[3])
+            
+            cmds.select(["proxy_geo_spine_02"], r = True)
+            spine2Controls = self.buildSpineControls("spine_03_pivot_joint", "spine_02")
+            cmds.parent(["proxy_geo_spine_02"], spine2Controls[3])
+            
+            cmds.select(["proxy_geo_spine_03"], r = True)
+            spine3Controls = self.buildSpineControls("spine_05_pivot_joint", "spine_03")
+            cmds.parent(["proxy_geo_spine_03"], spine3Controls[3])
+            
+            cmds.select(["proxy_geo_spine_04"], r = True)
+            spine4Controls = self.buildSpineControls("spine_07_pivot_joint", "spine_04")
+            cmds.parent(["proxy_geo_spine_04"], spine4Controls[3])
+            
+            cmds.select(["proxy_geo_spine_05"], r = True)
+            spine5Controls = self.buildSpineControls("spine_09_pivot_joint", "spine_05")
+            cmds.parent(["proxy_geo_spine_05"], spine5Controls[3])
+            
+            cmds.parent(spine5Controls[0], spine4Controls[1])
+            cmds.parent(spine4Controls[0], spine3Controls[1])
+            cmds.parent(spine3Controls[0], spine2Controls[1])
+            cmds.parent(spine2Controls[0], spine1Controls[1])
+            cmds.parent(spine1Controls[0], "pelvis_mover")
+            
+            
+
+            
+    def buildJointMover_Neck(self):
+        
+        #get number of neck bones
+        neckBones = cmds.getAttr("Skeleton_Settings.numNeckBones")
+        spineBones = cmds.getAttr("Skeleton_Settings.numSpineBones")
+        
+        
+        if neckBones == 1:
+            
+            #select proxy geo and build controls
+            cmds.select(["proxy_geo_neck_01", "proxy_geo_neck_02", "proxy_geo_neck_03"], r = True)
+            neck1Controls = self.buildSpineControls("neck_01_pivot_joint", "neck_01")
+            cmds.parent(["proxy_geo_neck_01", "proxy_geo_neck_02", "proxy_geo_neck_03"], neck1Controls[3])
+            
+            if spineBones == 10:
+                cmds.parent(neck1Controls[0], ("spine_" + str(int(spineBones)) + "_mover"))
+                
+            else:
+                cmds.parent(neck1Controls[0], ("spine_0" + str(int(spineBones)) + "_mover"))
+                
+                
+                
+        if neckBones == 2:
+            
+            #select proxy geo and build controls
+            cmds.select(["proxy_geo_neck_01", "proxy_geo_neck_02"], r = True)
+            neck1Controls = self.buildSpineControls("neck_01_pivot_joint", "neck_01")
+            cmds.parent(["proxy_geo_neck_01", "proxy_geo_neck_02"], neck1Controls[3])
+            
+            #select proxy geo and build controls
+            cmds.select(["proxy_geo_neck_03"], r = True)
+            neck2Controls = self.buildSpineControls("neck_03_pivot_joint", "neck_02")
+            cmds.parent(["proxy_geo_neck_03"], neck2Controls[3])
+            
+            
+            cmds.parent(neck2Controls[0], neck1Controls[1])
+            
+            if spineBones == 10:
+                cmds.parent(neck1Controls[0], ("spine_" + str(int(spineBones)) + "_mover"))
+                
+            else:
+                cmds.parent(neck1Controls[0], ("spine_0" + str(int(spineBones)) + "_mover"))
+                
+                
+                
+        if neckBones == 3:
+            
+            #select proxy geo and build controls
+            cmds.select(["proxy_geo_neck_01"], r = True)
+            neck1Controls = self.buildSpineControls("neck_01_pivot_joint", "neck_01")
+            cmds.parent(["proxy_geo_neck_01"], neck1Controls[3])
+            
+            #select proxy geo and build controls
+            cmds.select(["proxy_geo_neck_02"], r = True)
+            neck2Controls = self.buildSpineControls("neck_02_pivot_joint", "neck_02")
+            cmds.parent(["proxy_geo_neck_02"], neck2Controls[3])
+            
+            #select proxy geo and build controls
+            cmds.select(["proxy_geo_neck_03"], r = True)
+            neck3Controls = self.buildSpineControls("neck_03_pivot_joint", "neck_03")
+            cmds.parent(["proxy_geo_neck_03"], neck3Controls[3])
+            
+            
+            cmds.parent(neck3Controls[0], neck2Controls[1])
+            cmds.parent(neck2Controls[0], neck1Controls[1])
+            
+            if spineBones == 10:
+                cmds.parent(neck1Controls[0], ("spine_" + str(int(spineBones)) + "_mover"))
+                
+            else:
+                cmds.parent(neck1Controls[0], ("spine_0" + str(int(spineBones)) + "_mover"))
+                
+                
+    def buildJointMover_Hands(self, side):
+        
+        foot = side.lower()
+        
+        #check each attr, if attr is false, parent proxy to next in chain.
+        
+        #BIG TOE
+        if cmds.getAttr("Skeleton_Settings." + foot + "thumb3") == False:
+            cmds.parent("proxy_geo_thumb_03_" + foot[0], "thumb_02_geo_mover_" + foot[0])
+            shape = cmds.listRelatives("thumb_03_mover_" + foot[0], shapes = True)[0]
+            cmds.setAttr("thumb_03_mover_" + foot[0] + ".v", 0)
+            cmds.setAttr("thumb_03_mover_" + foot[0] + ".v", lock = True, keyable = False)
+            cmds.setAttr("thumb_03_mover_" + foot[0] + "|" + shape + ".v", 0)
+            cmds.setAttr("thumb_03_mover_" + foot[0] + "|" + shape + ".v", lock = True, keyable = False)
+            
+        if cmds.getAttr("Skeleton_Settings." + foot + "thumb2") == False:
+            children = cmds.listRelatives("thumb_02_geo_mover_" + foot[0], children = True, type = "transform")
+            cmds.parent(children, "thumb_01_geo_mover_" + foot[0])
+            shape = cmds.listRelatives("thumb_02_mover_" + foot[0], shapes = True)[0]
+            cmds.setAttr("thumb_02_mover_" + foot[0] + ".v", 0)
+            cmds.setAttr("thumb_02_mover_" + foot[0] + ".v", lock = True, keyable = False)
+            cmds.setAttr("thumb_02_mover_" + foot[0] + "|" + shape + ".v", 0)
+            cmds.setAttr("thumb_02_mover_" + foot[0] + "|" + shape + ".v", lock = True, keyable = False)
+            
+        if cmds.getAttr("Skeleton_Settings." + foot + "thumb1") == False:
+            children = cmds.listRelatives("thumb_01_geo_mover_" + foot[0], children = True, type = "transform")
+            cmds.parent(children, "hand_geo_mover_" + foot[0])
+            
+            
+            shape = cmds.listRelatives("thumb_01_mover_" + foot[0], shapes = True)[0]
+            cmds.setAttr("thumb_01_mover_" + foot[0] + "|" + shape + ".v", 0)
+            cmds.setAttr("thumb_01_mover_" + foot[0] + "|" + shape + ".v", lock = True, keyable = False)
+            
+            shape = cmds.listRelatives("thumb_01_mover_offset_" + foot[0], shapes = True)[0]
+            cmds.setAttr("thumb_01_mover_offset_" + foot[0] + "|" + shape + ".v", 0)
+            cmds.setAttr("thumb_01_mover_offset_" + foot[0] + "|" + shape + ".v", lock = True, keyable = False)
+            
+            cmds.parent("thumb_02_mover_" + foot[0] + "_grp", "hand_mover_" + foot[0])
+            
+        #INDEX TOE
+        
+        for toe in ["index", "middle", "ring", "pinky"]:
+            
+            if cmds.getAttr("Skeleton_Settings." + foot + toe + "3") == False:
+                cmds.parent("proxy_geo_" + toe + "_03_" + foot[0], toe + "_02_geo_mover_" + foot[0])
+                shape = cmds.listRelatives(toe + "_03_mover_" + foot[0], shapes = True)[0]
+                cmds.setAttr(toe + "_03_mover_" + foot[0] + ".v", 0)
+                cmds.setAttr(toe + "_03_mover_" + foot[0] + ".v", lock = True, keyable = False)
+                cmds.setAttr(toe + "_03_mover_" + foot[0] + "|" + shape + ".v", 0)
+                cmds.setAttr(toe + "_03_mover_" + foot[0] + "|" + shape + ".v", lock = True, keyable = False)
+                
+            if cmds.getAttr("Skeleton_Settings." + foot + toe + "2") == False:
+                children = cmds.listRelatives(toe + "_02_geo_mover_" + foot[0], children = True, type = "transform")
+                
+                cmds.parent(children, "" + toe + "_01_geo_mover_" + foot[0])
+                shape = cmds.listRelatives(toe + "_02_mover_" + foot[0], shapes = True)[0]
+                cmds.setAttr(toe + "_02_mover_" + foot[0] + ".v", 0)
+                cmds.setAttr(toe + "_02_mover_" + foot[0] + ".v", lock = True, keyable = False)
+                cmds.setAttr(toe + "_02_mover_" + foot[0] + "|" + shape + ".v", 0)
+                cmds.setAttr(toe + "_02_mover_" + foot[0] + "|" + shape + ".v", lock = True, keyable = False)
+                
+            if cmds.getAttr("Skeleton_Settings." + foot + toe + "1") == False:
+                children = cmds.listRelatives(toe + "_01_geo_mover_" + foot[0], children = True, type = "transform")
+                cmds.parent(children, "" + toe + "_metacarpal_geo_mover_" + foot[0])
+                shape = cmds.listRelatives(toe + "_01_mover_" + foot[0], shapes = True)[0]
+                cmds.setAttr(toe + "_01_mover_" + foot[0] + ".v", 0)
+                cmds.setAttr(toe + "_01_mover_" + foot[0] + ".v", lock = True, keyable = False)
+                cmds.setAttr(toe + "_01_mover_" + foot[0] + "|" + shape + ".v", 0)
+                cmds.setAttr(toe + "_01_mover_" + foot[0] + "|" + shape + ".v", lock = True, keyable = False)
+                
+            if cmds.getAttr("Skeleton_Settings." + foot + toe + "meta") == False:
+                children = cmds.listRelatives("" + toe + "_metacarpal_geo_mover_" + foot[0], children = True, type = "transform")
+                cmds.parent(children, "hand_geo_mover_" + foot[0])
+                
+                shape = cmds.listRelatives(toe + "_metacarpal_mover_" + foot[0], shapes = True)[0]
+                cmds.setAttr(toe + "_metacarpal_mover_" + foot[0] + "|" + shape + ".v", 0)
+                cmds.setAttr(toe + "_metacarpal_mover_" + foot[0] + "|" + shape + ".v", lock = True, keyable = False)
+                
+                shape = cmds.listRelatives(toe + "_metacarpal_mover_offset_" + foot[0], shapes = True)[0]
+                cmds.setAttr(toe + "_metacarpal_mover_offset_" + foot[0] + "|" + shape + ".v", 0)
+                cmds.setAttr(toe + "_metacarpal_mover_offset_" + foot[0] + "|" + shape + ".v", lock = True, keyable = False)
+
+                
+                cmds.parent(toe + "_01_mover_" + foot[0] + "_grp", "hand_mover_" + foot[0])
+                
+                
+
+            
+    def buildJointMover_Feet(self, side):
+        
+        foot = side.lower()
+        
+        #check each attr, if attr is false, parent proxy to next in chain.
+        
+        #BIG TOE
+        if cmds.getAttr("Skeleton_Settings." + foot + "Footbigtoe2") == False:
+            children = cmds.listRelatives("bigtoe_distal_phalange_geo_mover_" + foot[0], children = True, type = "transform")
+            cmds.parent(children, "bigtoe_proximal_phalange_geo_mover_" + foot[0])
+            shape = cmds.listRelatives("bigtoe_distal_phalange_mover_" + foot[0], shapes = True)[0]
+            cmds.setAttr("bigtoe_distal_phalange_mover_" + foot[0] + ".v", 0)
+            cmds.setAttr("bigtoe_distal_phalange_mover_" + foot[0] + ".v", lock = True, keyable = False)
+            cmds.setAttr("bigtoe_distal_phalange_mover_" + foot[0] + "|" + shape + ".v", 0)
+            cmds.setAttr("bigtoe_distal_phalange_mover_" + foot[0] + "|" + shape + ".v", lock = True, keyable = False)
+            
+        if cmds.getAttr("Skeleton_Settings." + foot + "Footbigtoe1") == False:
+            children = cmds.listRelatives("bigtoe_proximal_phalange_geo_mover_" + foot[0], children = True, type = "transform")
+            cmds.parent(children, "bigtoe_metatarsal_geo_mover_" + foot[0])
+            shape = cmds.listRelatives("bigtoe_proximal_phalange_mover_" + foot[0], shapes = True)[0]
+            cmds.setAttr("bigtoe_proximal_phalange_mover_" + foot[0] + ".v", 0)
+            cmds.setAttr("bigtoe_proximal_phalange_mover_" + foot[0] + ".v", lock = True, keyable = False)
+            cmds.setAttr("bigtoe_proximal_phalange_mover_" + foot[0] + "|" + shape + ".v", 0)
+            cmds.setAttr("bigtoe_proximal_phalange_mover_" + foot[0] + "|" + shape + ".v", lock = True, keyable = False)
+            
+        if cmds.getAttr("Skeleton_Settings." + foot + "Footbigtoemeta") == False:
+            children = cmds.listRelatives("bigtoe_metatarsal_geo_mover_" + foot[0], children = True, type = "transform")
+            cmds.parent(children, "ball_geo_mover_" + foot[0])
+            
+            
+            shape = cmds.listRelatives("bigtoe_metatarsal_mover_" + foot[0], shapes = True)[0]
+            cmds.setAttr("bigtoe_metatarsal_mover_" + foot[0] + "|" + shape + ".v", 0)
+            cmds.setAttr("bigtoe_metatarsal_mover_" + foot[0] + "|" + shape + ".v", lock = True, keyable = False)
+            
+            shape = cmds.listRelatives("bigtoe_metatarsal_mover_offset_" + foot[0], shapes = True)[0]
+            cmds.setAttr("bigtoe_metatarsal_mover_offset_" + foot[0] + "|" + shape + ".v", 0)
+            cmds.setAttr("bigtoe_metatarsal_mover_offset_" + foot[0] + "|" + shape + ".v", lock = True, keyable = False)
+            
+            cmds.parent("bigtoe_proximal_phalange_mover_" + foot[0] + "_grp", "ball_mover_" + foot[0])
+            
+        #INDEX TOE
+        
+        for toe in ["index", "middle", "ring", "pinky"]:
+            
+            if cmds.getAttr("Skeleton_Settings." + foot + "Foot" + toe + "3") == False:
+                children = cmds.listRelatives("" + toe + "_distal_phalange_geo_mover_" + foot[0], children = True, type = "transform")
+                cmds.parent(children, "" + toe + "_middle_phalange_geo_mover_" + foot[0])
+                shape = cmds.listRelatives(toe + "_distal_phalange_mover_" + foot[0], shapes = True)[0]
+                cmds.setAttr(toe + "_distal_phalange_mover_" + foot[0] + ".v", 0)
+                cmds.setAttr(toe + "_distal_phalange_mover_" + foot[0] + ".v", lock = True, keyable = False)
+                cmds.setAttr(toe + "_distal_phalange_mover_" + foot[0] + "|" + shape + ".v", 0)
+                cmds.setAttr(toe + "_distal_phalange_mover_" + foot[0] + "|" + shape + ".v", lock = True, keyable = False)
+                
+            if cmds.getAttr("Skeleton_Settings." + foot + "Foot" + toe + "2") == False:
+                children = cmds.listRelatives("" + toe + "_middle_phalange_geo_mover_" + foot[0], children = True, type = "transform")
+                cmds.parent(children, "" + toe + "_proximal_phalange_geo_mover_" + foot[0])
+                shape = cmds.listRelatives(toe + "_middle_phalange_mover_" + foot[0], shapes = True)[0]
+                cmds.setAttr(toe + "_middle_phalange_mover_" + foot[0] + ".v", 0)
+                cmds.setAttr(toe + "_middle_phalange_mover_" + foot[0] + ".v", lock = True, keyable = False)
+                cmds.setAttr(toe + "_middle_phalange_mover_" + foot[0] + "|" + shape + ".v", 0)
+                cmds.setAttr(toe + "_middle_phalange_mover_" + foot[0] + "|" + shape + ".v", lock = True, keyable = False)
+                
+            if cmds.getAttr("Skeleton_Settings." + foot + "Foot" + toe + "1") == False:
+                children = cmds.listRelatives("" + toe + "_proximal_phalange_geo_mover_" + foot[0], children = True, type = "transform")
+                cmds.parent(children, "" + toe + "_metatarsal_geo_mover_" + foot[0])
+                shape = cmds.listRelatives(toe + "_proximal_phalange_mover_" + foot[0], shapes = True)[0]
+                cmds.setAttr(toe + "_proximal_phalange_mover_" + foot[0] + ".v", 0)
+                cmds.setAttr(toe + "_proximal_phalange_mover_" + foot[0] + ".v", lock = True, keyable = False)
+                cmds.setAttr(toe + "_proximal_phalange_mover_" + foot[0] + "|" + shape + ".v", 0)
+                cmds.setAttr(toe + "_proximal_phalange_mover_" + foot[0] + "|" + shape + ".v", lock = True, keyable = False)
+                
+            if cmds.getAttr("Skeleton_Settings." + foot + "Foot" + toe + "meta") == False:
+                children = cmds.listRelatives("" + toe + "_metatarsal_geo_mover_" + foot[0], children = True, type = "transform")
+                cmds.parent(children, "ball_geo_mover_" + foot[0])
+                
+                shape = cmds.listRelatives(toe + "_metatarsal_mover_" + foot[0], shapes = True)[0]
+                cmds.setAttr(toe + "_metatarsal_mover_" + foot[0] + "|" + shape + ".v", 0)
+                cmds.setAttr(toe + "_metatarsal_mover_" + foot[0] + "|" + shape + ".v", lock = True, keyable = False)
+                
+                shape = cmds.listRelatives(toe + "_metatarsal_mover_offset_" + foot[0], shapes = True)[0]
+                cmds.setAttr(toe + "_metatarsal_mover_offset_" + foot[0] + "|" + shape + ".v", 0)
+                cmds.setAttr(toe + "_metatarsal_mover_offset_" + foot[0] + "|" + shape + ".v", lock = True, keyable = False)
+                
+                
+                cmds.parent(toe + "_proximal_phalange_mover_" + foot[0] + "_grp", "ball_mover_" + foot[0])
+                
+        if cmds.getAttr("Skeleton_Settings." + foot + "ball") == False:
+            children = cmds.listRelatives("ball_geo_mover_" + foot[0], children = True, type = "transform")
+            cmds.parent(children, "foot_geo_mover_" + foot[0])
+            cmds.parent(toe + "_proximal_phalange_mover_" + foot[0] + "_grp", "foot_mover_" + foot[0])
+            shape = cmds.listRelatives("ball_mover_" + foot[0], shapes = True)[0]
+            cmds.setAttr("ball_mover_" + foot[0] + ".v", 0)
+            cmds.setAttr("ball_mover_" + foot[0] + ".v", lock = True, keyable = False)
+            cmds.setAttr("ball_mover_" + foot[0] + "|" + shape + ".v", 0)
+            cmds.setAttr("ball_mover_" + foot[0] + "|" + shape + ".v", lock = True, keyable = False)
+                
+            
+                
+
+                
+            
+            
+            
+            
+
+            
+            
+    def addExtraJoints(self, attrs):
+        skippedAttrs = []
+        
+        #right side joints needs to have -180 on rotate X to be mirrored
+        #print attrs
+        
+        for attr in attrs:
+            if attr.find("extraJoint") == 0:
+                attribute = cmds.getAttr("Skeleton_Settings." + attr, asString = True)
+                parent = attribute.partition("/")[0]
+                jointType = attribute.partition("/")[2].partition("/")[0]
+                name = attribute.rpartition("/")[2]
+                
+                
+                if parent.find("_") != -1:
+                    
+                    if parent.rpartition("_")[2] == "l":
+                        parent = parent.rpartition("_l")[0]
+                        parent = parent + "_mover_l"
+                        
+                        if cmds.objExists(parent) == False:
+                            parent = parent.partition("_mover_l")[0] + "_l_mover"
+                        
+                    if parent.rpartition("_")[2] == "r":
+                        parent = parent.rpartition("_r")[0]
+                        parent = parent + "_mover_r"
+                        
+                        if cmds.objExists(parent) == False:
+                            parent = parent.partition("_mover_r")[0] + "_r_mover"
+                        
+                    if parent.rpartition("_")[2] != "l":
+                        if parent.rpartition("_")[2] != "r":
+                            if parent.rpartition("_")[2] != "mover":
+                                parent = parent + "_mover"
+
+                    
+                if parent.find("_") == -1:
+                    parent = parent + "_mover"
+                
+                #check to see if the parent is a leaf joint and modify parent name if so
+                if parent.find("(") != -1:
+                    parent = parent.partition(" (")[0] + "_mover"
+                
+                
+                if cmds.objExists(parent):
+                    
+                    if jointType != "chain":
+                        #for leaf joints, find the suffix TR, T, or R
+                        if name.find("(") != -1:
+                            name = name.partition(" (")[0]
+                            
+                        group = cmds.group(empty = True, name = (name + "_mover_grp"))
+                        constraint = cmds.pointConstraint((parent), group)[0]
+                        cmds.delete(constraint)
+                        
+                        if name.rpartition("_")[2] == "r":
+                            cmds.setAttr(group + ".rx", -180)
+                            
+                        if name.find("r_") == 0:
+                            cmds.setAttr(group + ".rx", -180)
+                            
+                        if name.find("_r_") != -1:
+                            cmds.setAttr(group + ".rx", -180)
+
+                        #force update
+                        cmds.select(group)
+                        cmds.setToolTo( 'moveSuperContext' )
+                        cmds.refresh(force = True)
+                        cmds.select(clear = True)     
+                        
+                        cmds.parent(group, parent)
+                    
+    
+                        
+                        #position the global mover (need to dupe template)
+                        duplicate = cmds.duplicate("extra_joints_template", name = (name + "_mover"))[0]
+                        cmds.setAttr(duplicate + ".visibility", True)
+                        constraint = cmds.parentConstraint(group, duplicate)
+                        cmds.delete(constraint)
+                        cmds.parent(duplicate, group)
+                        cmds.makeIdentity(duplicate, t = 1, r = 1, s = 1, apply = True)
+                        
+                        #create the offset control
+                        offsetCtrl = cmds.duplicate(duplicate, name = (name + "_mover_offset"))[0]
+                        cmds.setAttr(offsetCtrl + ".scaleX",.85)
+                        cmds.setAttr(offsetCtrl + ".scaleY",.85)
+                        cmds.setAttr(offsetCtrl + ".scaleZ",.85)
+                        
+                        cmds.parent(offsetCtrl, duplicate)
+                        cmds.makeIdentity(offsetCtrl, t = 1, r = 1, s = 1, apply = True)
+                        cmds.setAttr(offsetCtrl + ".overrideColor", 18)
+                        
+                        #create the geo mover
+                        geoMover = cmds.duplicate("geo_mover", name = (name + "_geo_mover"))[0]
+                        constraint = cmds.parentConstraint(offsetCtrl, geoMover)[0]
+                        cmds.delete(constraint)
+                        cmds.parent(geoMover, offsetCtrl)
+                        shape = cmds.listRelatives(geoMover, shapes = True)[0]
+                        cmds.setAttr(geoMover + "|" + shape + ".v", 0)
+                        cmds.makeIdentity(geoMover, t = 1, r = 1, s = 1, apply = True)
+    
+                        #create the proxy geo
+                        proxy = cmds.duplicate("extra_joints_template_geo", name = "proxy_geo_" + name)[0]
+                        constraint = cmds.parentConstraint(offsetCtrl, proxy)[0]
+                        cmds.delete(constraint)
+                        cmds.parent(proxy, geoMover)
+                        cmds.setAttr(proxy + ".visibility", 1)
+                        cmds.makeIdentity(proxy, t = 1, r = 1, s = 1, apply = True)
+                        
+                        
+                    
+                        #create the lra
+                        lraGroup = cmds.group(empty = True, name = name + "_lra_grp")
+                        constraint = cmds.parentConstraint(group, lraGroup)[0]
+                        cmds.delete(constraint)
+                        cmds.parent(lraGroup, offsetCtrl)
+                        
+                        lra = cmds.duplicate("lra", name = name + "_lra")[0]
+                        cmds.setAttr(lra + ".v", 1)
+
+                            
+                        constraint = cmds.parentConstraint(lraGroup, lra)[0]
+                        cmds.delete(constraint)
+
+                        cmds.parent(lra, lraGroup)
+                        cmds.makeIdentity(lra, t = 1, r = 1, s = 1, apply = True)
+                        
+                        attrs = cmds.listAttr(lra, keyable = True)
+                        for attr in attrs:
+                            cmds.setAttr(lra + "." + attr, lock = True)
+
+                        
+                    #otherwise, it's dynamic
+                    else:
+                        numJoints = name.partition(" (")[2].partition(")")[0]
+                        name = name.partition(" (")[0]
+                                                
+
+                        space = 20
+            
+                        for i in range(int(numJoints)):
+                            
+                            group = cmds.group(empty = True, name = (name + "_0" + str(i + 1) + "_mover_grp"))
+                            constraint = cmds.pointConstraint((parent), group)[0]
+                            cmds.delete(constraint)
+                            
+                             
+                            #need to mirror world group for right side
+                            if name.rpartition("_")[2] == "r":
+                                cmds.setAttr(group + ".rx", -180)
+                                space = -20
+                                
+                            if name.find("r_") == 0:
+                                cmds.setAttr(group + ".rx", -180)
+                                space = -20
+                                
+                            if name.find("_r_") != -1:
+                                cmds.setAttr(group + ".rx", -180)
+                                space = -20
+
+
+                            cmds.parent(group, parent)
+                            if i != 0:
+                                cmds.setAttr(group + ".ty", space)
+                            
+                            
+                            
+                            #position the global mover (need to dupe template)
+                            duplicate = cmds.duplicate("extra_joints_dynamic", name = (name + "_0" + str(i + 1) + "_mover"))[0]
+                            cmds.setAttr(duplicate + ".visibility", True)
+                            
+                            #need to mirror world group for right side
+                            if name.rpartition("_")[2] == "r":
+                                cmds.xform(duplicate, relative = True, piv = [0, 20, 0])
+                                
+                            if name.find("r_") == 0:
+                                cmds.xform(duplicate, relative = True, piv = [0, 20, 0])
+                                
+                            if name.find("_r_") != -1:
+                                cmds.xform(duplicate, relative = True, piv = [0, 20, 0])
+                                
+                                
+                                
+                            constraint = cmds.parentConstraint(group, duplicate)
+                            cmds.delete(constraint)
+                            cmds.parent(duplicate, group)
+                            cmds.makeIdentity(duplicate, t = 1, r = 1, s = 1, apply = True)
+                            
+                            
+                            #create the offset control
+                            offsetCtrl = cmds.duplicate(duplicate, name = (name + "_0" + str(i + 1) + "_mover_offset"))[0]
+                            cmds.setAttr(offsetCtrl + ".scaleX",.85)
+                            cmds.setAttr(offsetCtrl + ".scaleY",.85)
+                            cmds.setAttr(offsetCtrl + ".scaleZ",.85)
+                            
+                            cmds.parent(offsetCtrl, duplicate)
+                            cmds.makeIdentity(offsetCtrl, t = 1, r = 1, s = 1, apply = True)
+                            cmds.setAttr(offsetCtrl + ".overrideColor", 18)
+                            
+                            
+                            #create the geo mover
+                            geoMover = cmds.duplicate("geo_mover", name = (name + "_0" + str(i + 1) + "_geo_mover"))[0]
+                            constraint = cmds.parentConstraint(offsetCtrl, geoMover)[0]
+                            cmds.delete(constraint)
+                            cmds.parent(geoMover, offsetCtrl)
+                            shape = cmds.listRelatives(geoMover, shapes = True)[0]
+                            cmds.setAttr(geoMover + "|" + shape + ".v", 0)
+                            cmds.makeIdentity(geoMover, t = 1, r = 1, s = 1, apply = True)
+    
+    
+    
+                            #create the proxy geo
+                            proxy = cmds.duplicate("extra_joints_dynamic_geo", name = "proxy_geo_" + name + "_0" + str(i + 1))[0]
+                            constraint = cmds.pointConstraint(offsetCtrl, proxy)[0]
+                            cmds.delete(constraint)
+                            
+                            cmds.parent(proxy, geoMover)
+                            cmds.setAttr(proxy + ".visibility", 1)
+                            cmds.makeIdentity(proxy, t = 1, r = 1, s = 1, apply = True)
+                            
+                            
+                            #create the lra
+                            lraGroup = cmds.group(empty = True, name = name + "_0" + str(i + 1) + "_lra_grp")
+                            constraint = cmds.parentConstraint(group, lraGroup)[0]
+                            cmds.delete(constraint)
+                            cmds.parent(lraGroup, offsetCtrl)
+                            
+                            lra = cmds.duplicate("lra", name = name + "_0" + str(i + 1) + "_lra")[0]
+                            cmds.setAttr(lra + ".v", 1)
+
+                            constraint = cmds.parentConstraint(lraGroup, lra)[0]
+                            cmds.delete(constraint)
+
+                            cmds.parent(lra, lraGroup)
+                            cmds.makeIdentity(lra, t = 1, r = 1, s = 1, apply = True)
+                            
+                            attrs = cmds.listAttr(lra, keyable = True)
+                            for attr in attrs:
+                                cmds.setAttr(lra + "." + attr, lock = True)
+                            
+                            
+                            parent = duplicate
+                            
+                                
+                else:
+                    skippedAttrs.append(attr)
+                    
+        if skippedAttrs:
+            self.addExtraJoints(skippedAttrs)
+            
+            
+            
+    def buildTwistJoints(self, moverName, attr, side, parent, child, direction):
+	
+	#get length of bone
+	dist = (abs(cmds.getAttr(child + ".tx"))) / (cmds.getAttr("Skeleton_Settings." + attr))
+
+	#Twist Attr
+	if cmds.getAttr("Skeleton_Settings." + attr) > 0:
+	    
+	    for i in range(int(cmds.getAttr("Skeleton_Settings." + attr))):
+		name = moverName + str(i + 1)
+		group = cmds.group(empty = True, name = (name + "_mover_" + side + "_grp"))
+		lraTarget = parent.partition("_mover")[0] + "_" + side + "_lra"
+		constraint = cmds.parentConstraint(lraTarget, group)[0]
+		cmds.delete(constraint)	
+		cmds.parent(group, parent)
+		
+		
+		#position group
+		tx = cmds.getAttr(child + ".tx")
+		if  i == 0:
+		    if direction == "down":
+			if tx > 0:
+			    tx = 1
+			else:
+			    tx = -1
+			cmds.setAttr(group + ".tx", tx)
+			
+		    if direction == "up":
+			cmds.setAttr(group + ".tx", cmds.getAttr(child + ".tx"))			
+		    
+		    
+		else:
+		    if direction == "down":
+			if cmds.getAttr(child + ".tx") > 0:
+			    tx = dist * i
+			else:
+			    tx = (dist * -1) * i
+			    
+			cmds.setAttr(group + ".tx", tx)
+			
+		    if direction == "up":
+			if cmds.getAttr(child + ".tx") > 0:
+			    tx = dist * i
+			    tx = cmds.getAttr(child + ".tx") - tx
+			    
+			else:
+			    tx = dist * i
+			    tx = cmds.getAttr(child + ".tx") + tx
+			    
+			cmds.setAttr(group + ".tx", tx)
+			
+		    
+		    
+		
+		#position the global mover (need to dupe template)
+		duplicate = cmds.duplicate("extra_joints_template", name = (name + "_mover_" + side))[0]
+		cmds.setAttr(duplicate + ".visibility", True)
+		constraint = cmds.parentConstraint(group, duplicate)
+		cmds.delete(constraint)
+		cmds.parent(duplicate, group)
+		cmds.makeIdentity(duplicate, t = 1, r = 1, s = 1, apply = True)		
+
+		
+		#create the offset control
+		offsetCtrl = cmds.duplicate(duplicate, name = (name + "_mover_offset_" + side))[0]
+		cmds.setAttr(offsetCtrl + ".scaleX",.85)
+		cmds.setAttr(offsetCtrl + ".scaleY",.85)
+		cmds.setAttr(offsetCtrl + ".scaleZ",.85)
+		
+		cmds.parent(offsetCtrl, duplicate)
+		cmds.makeIdentity(offsetCtrl, t = 1, r = 1, s = 1, apply = True)
+		cmds.setAttr(offsetCtrl + ".overrideColor", 18)		
+		
+		#create the geo mover
+		geoMover = cmds.duplicate("geo_mover", name = (name + "_geo_mover_" + side))[0]
+		constraint = cmds.parentConstraint(offsetCtrl, geoMover)[0]
+		cmds.delete(constraint)
+		cmds.parent(geoMover, offsetCtrl)
+		shape = cmds.listRelatives(geoMover, shapes = True)[0]
+		cmds.setAttr(geoMover + "|" + shape + ".v", 0)
+		cmds.makeIdentity(geoMover, t = 1, r = 1, s = 1, apply = True)
+	
+		#create the proxy geo
+		proxy = cmds.duplicate("extra_joints_template_geo", name = "proxy_geo_" + name + "_" + side)[0]
+		constraint = cmds.parentConstraint(offsetCtrl, proxy)[0]
+		cmds.delete(constraint)
+		cmds.parent(proxy, geoMover)
+		cmds.setAttr(proxy + ".visibility", 1)
+		cmds.makeIdentity(proxy, t = 1, r = 1, s = 1, apply = True)		
+		
+		
+		#create the lra
+		lraGroup = cmds.group(empty = True, name = name + "_" + side + "_lra_grp")
+		constraint = cmds.parentConstraint(group, lraGroup)[0]
+		cmds.delete(constraint)
+		cmds.parent(lraGroup, offsetCtrl)
+		
+		lra = cmds.duplicate("lra", name = name + "_" + side + "_lra")[0]
+		cmds.setAttr(lra + ".v", 1)
+	
+		    
+		constraint = cmds.parentConstraint(lraGroup, lra)[0]
+		cmds.delete(constraint)
+	
+		cmds.parent(lra, lraGroup)
+		cmds.makeIdentity(lra, t = 1, r = 1, s = 1, apply = True)
+		
+		attrs = cmds.listAttr(lra, keyable = True)
+		for attr in attrs:
+		    cmds.setAttr(lra + "." + attr, lock = True)   		
+
+     
