Adding ART scripts to repository

diff --git a/Engine/Extras/Maya_AnimationRiggingTools/MayaTools/General/Scripts/ART_poseEditor.py b/Engine/Extras/Maya_AnimationRiggingTools/MayaTools/General/Scripts/ART_poseEditor.py
new file mode 100644
--- /dev/null
+++ b/Engine/Extras/Maya_AnimationRiggingTools/MayaTools/General/Scripts/ART_poseEditor.py
@@ -0,0 +1,2236 @@
+import maya.cmds as cmds
+import maya.mel as mel
+from functools import partial
+import os, cPickle
+import subprocess
+
+class PoseEditor_UI():
+    
+    def __init__(self, character, inst):
+        
+        #create class variables
+        self.widgets = {}
+        self.character = character
+	self.animUIInst = inst
+        
+        #get access to our maya tools
+        toolsPath = cmds.internalVar(usd = True) + "mayaTools.txt"
+        if os.path.exists(toolsPath):
+            
+            f = open(toolsPath, 'r')
+            self.mayaToolsDir = f.readline()
+            f.close()
+            
+            
+            
+	#get all controls
+        self.controls = []
+	self.customModCtrls = []
+        for control in ["head_fk_anim", "neck_01_fk_anim", "neck_02_fk_anim", "neck_03_fk_anim", "spine_01_anim", "spine_02_anim", "spine_03_anim", "spine_04_anim", "spine_05_anim", "mid_ik_anim", "chest_ik_anim",
+                        "body_anim", "hip_anim", "clavicle_l_anim", "clavicle_r_anim", "fk_arm_l_anim", "fk_arm_r_anim", "fk_elbow_l_anim", "fk_elbow_r_anim", "fk_wrist_l_anim", "fk_wrist_r_anim",
+                        "ik_elbow_l_anim", "ik_elbow_r_anim", "ik_wrist_l_anim", "ik_wrist_r_anim", "fk_thigh_l_anim", "fk_thigh_r_anim", "fk_calf_l_anim", "fk_calf_r_anim", "fk_foot_l_anim", "fk_foot_r_anim",
+                        "fk_ball_l_anim", "fk_ball_r_anim", "ik_foot_anim_l", "ik_foot_anim_r", "heel_ctrl_l", "heel_ctrl_r", "toe_wiggle_ctrl_l", "toe_wiggle_ctrl_r",
+                        "toe_tip_ctrl_l", "toe_tip_ctrl_r", "master_anim", "offset_anim", "root_anim", "upperarm_l_twist_anim", "upperarm_r_twist_anim", "thigh_l_twist_anim", "thigh_r_twist_anim",
+                        "pinky_metacarpal_ctrl_l", "pinky_metacarpal_ctrl_r", "pinky_finger_fk_ctrl_1_l", "pinky_finger_fk_ctrl_1_r", "pinky_finger_fk_ctrl_2_l", "pinky_finger_fk_ctrl_2_r", "pinky_finger_fk_ctrl_3_l", "pinky_finger_fk_ctrl_3_r",
+                        "ring_metacarpal_ctrl_l", "ring_metacarpal_ctrl_r", "ring_finger_fk_ctrl_1_l", "ring_finger_fk_ctrl_1_r", "ring_finger_fk_ctrl_2_l", "ring_finger_fk_ctrl_2_r", "ring_finger_fk_ctrl_3_l", "ring_finger_fk_ctrl_3_r",
+                        "middle_metacarpal_ctrl_l", "middle_metacarpal_ctrl_r", "middle_finger_fk_ctrl_1_l", "middle_finger_fk_ctrl_1_r", "middle_finger_fk_ctrl_2_l", "middle_finger_fk_ctrl_2_r", "middle_finger_fk_ctrl_3_l", "middle_finger_fk_ctrl_3_r",
+                        "index_metacarpal_ctrl_l", "index_metacarpal_ctrl_r", "index_finger_fk_ctrl_1_l", "index_finger_fk_ctrl_1_r", "index_finger_fk_ctrl_2_l", "index_finger_fk_ctrl_2_r", "index_finger_fk_ctrl_3_l", "index_finger_fk_ctrl_3_r",
+                        "thumb_finger_fk_ctrl_1_l", "thumb_finger_fk_ctrl_1_r", "thumb_finger_fk_ctrl_2_l", "thumb_finger_fk_ctrl_2_r", "thumb_finger_fk_ctrl_3_l", "thumb_finger_fk_ctrl_3_r",
+                        "index_l_ik_anim", "index_r_ik_anim", "middle_l_ik_anim", "middle_r_ik_anim", "ring_l_ik_anim", "ring_r_ik_anim", "pinky_l_ik_anim", "pinky_r_ik_anim", "thumb_l_ik_anim", "thumb_r_ik_anim",
+                        "index_l_poleVector", "index_r_poleVector", "middle_l_poleVector", "middle_r_poleVector", "ring_l_poleVector", "ring_r_poleVector", "pinky_l_poleVector", "pinky_r_poleVector", "thumb_l_poleVector", "thumb_r_poleVector",
+                        "l_global_ik_anim", "r_global_ik_anim", "Rig_Settings"]:
+	    self.controls.append(control)
+	    
+	#find custom joints
+        customJoints = []
+        attrs = cmds.listAttr(self.character + ":" + "Skeleton_Settings")
+        for attr in attrs:
+            if attr.find("extraJoint") == 0:
+                customJoints.append(attr)
+                
+        for joint in customJoints:
+            attribute = cmds.getAttr(self.character + ":" + "Skeleton_Settings." + joint, asString = True)
+            jointType = attribute.partition("/")[2].partition("/")[0]
+            label = attribute.rpartition("/")[2]
+            
+            if jointType == "leaf":
+                label = label.partition(" (")[0]
+		control = label + "_anim"
+		self.controls.append(control)
+		self.customModCtrls.append(control)
+		
+            if jointType == "jiggle":
+		control = label + "_anim"
+		self.controls.append(control)
+		self.customModCtrls.append(control)
+		
+	    if jointType == "chain" or jointType == "dynamic":
+		numJointsInChain = label.partition("(")[2].partition(")")[0]
+		label = label.partition(" (")[0]
+		self.controls.append(label + "_dyn_anim")
+		self.customModCtrls.append(label + "_dyn_anim")
+		
+		cmds.select(self.character + ":" + label + "_ik_*_anim")
+		selection = cmds.ls(sl = True)
+		for each in selection:
+		    niceName = each.partition(":")[2]
+		    self.controls.append(niceName)
+		    self.customModCtrls.append(niceName)
+		    
+		
+		for i in range(int(numJointsInChain)):
+		    self.controls.append("fk_" + label + "_0" + str(i + 1) + "_anim")
+		    self.controls.append(label + "_cv_" + str(i) + "_anim")	
+		    self.customModCtrls.append(label + "_cv_" + str(i) + "_anim")
+		    self.customModCtrls.append("fk_" + label + "_0" + str(i + 1) + "_anim")
+		
+		
+	#add space switch nodes
+	cmds.select(self.character + ":*_space_switcher_follow")
+	nodes = cmds.ls(sl = True)
+	spaceSwitchers = []
+	for node in nodes:
+	    if node.find("invis") == -1:
+		spaceSwitchers.append(node)
+		
+	
+        for control in spaceSwitchers:
+            spaceSwitchNode = control.rpartition("_follow")[0].partition(":")[2]
+	    self.controls.append(spaceSwitchNode)
+	    
+
+	    
+        #check to see if window exists, if so, delete
+        if cmds.window("poseEditorUI", exists = True):
+            cmds.deleteUI("poseEditorUI")
+            
+        #build window
+        self.widgets["window"] = cmds.window("poseEditorUI", w = 700, h = 400, title = self.character + "_Pose Editor", sizeable = False)
+            
+            
+        #create the main layout
+        self.widgets["topLevelLayout"] = cmds.columnLayout()
+	
+	#banner 
+	cmds.image(image = self.mayaToolsDir + "/General/Icons/ART/poseEditorBanner.bmp", w = 700, h = 50, parent = self.widgets["topLevelLayout"] )
+        
+        #create the rowColumnLayout 
+        self.widgets["rowColumnLayout"] = cmds.rowColumnLayout(w = 700, h = 400, nc = 2, cw = [(1, 150), (2, 550)], parent = self.widgets["topLevelLayout"])
+        
+        #create the columnLayout for the left side
+        self.widgets["leftSideButtonColumn"] = cmds.columnLayout(w = 150, h = 400, parent = self.widgets["rowColumnLayout"], cat = ["both", 5], rs = 5)
+        
+        
+        #and create the scroll layout for the right side
+        self.widgets["rightSideScrollLayout"] = cmds.scrollLayout(w = 550, h = 400, hst = 0, parent = self.widgets["rowColumnLayout"])
+        self.widgets["poseEditor_loadPoseProjFilter"] = cmds.rowColumnLayout(nc = 3, cat = [(1, "both", 10), (2, "both", 10), (3, "both", 10)], parent = self.widgets["rightSideScrollLayout"], rs = [1, 10])
+        
+        #create the projects optionMenu and the filter field
+	cmds.text(label = "Projects:", align = 'right')	
+        self.widgets["loadPose_projects"] = cmds.optionMenu(w = 250, parent = self.widgets["poseEditor_loadPoseProjFilter"])
+        self.widgets["loadPose_filter"] = cmds.textField(w = 180, text = "Search", parent = self.widgets["poseEditor_loadPoseProjFilter"], cc = self.filterResults)
+	
+	cmds.text(label = "Categories:",  align = 'right')
+	self.widgets["loadPose_categories"] = cmds.optionMenu(w = 250, parent = self.widgets["poseEditor_loadPoseProjFilter"], cc = self.loadPoses)
+        cmds.optionMenu(self.widgets["loadPose_projects"], edit = True,cc = partial(self.getCategories, self.widgets["loadPose_categories"]))
+	self.widgets["loadAllPosesButton"] = cmds.button(label = "Load All Poses", c = self.loadPoses_All)
+	
+        
+        cmds.separator(w = 540, h = 20, style = "out",  parent = self.widgets["rightSideScrollLayout"])
+	
+	self.widgets["loadPose_poses"] = cmds.rowColumnLayout(nc = 5, rs = [1, 5], cs = [(1, 7), (2, 7), (3, 7), (4, 7), (5, 7)], parent = self.widgets["rightSideScrollLayout"])
+        
+        #create save pose button
+        self.widgets["poseEditor_savePoseButton"] = cmds.symbolButton(w = 140, h = 50, image = self.mayaToolsDir + "/General/Icons/ART/savePose.bmp", parent = self.widgets["leftSideButtonColumn"], c = partial(self.savePose, False))
+	self.widgets["poseEditor_saveSelectionPoseButton"] = cmds.symbolButton(w = 140, h = 50, image = self.mayaToolsDir + "/General/Icons/ART/savePoseSelection.bmp", parent = self.widgets["leftSideButtonColumn"], c = partial(self.savePose, True))
+	self.widgets["poseEditor_syncStatusButton"] = cmds.button("poseEditor_syncStatusButton", w = 150, h = 50, label = "Out Of Sync!", bgc = [1, 0, 0], visible = False, parent = self.widgets["leftSideButtonColumn"], c = self.relaunch)
+
+        
+        #populate projects option menu
+        self.getProjects(self.widgets["loadPose_projects"])
+	self.getCategories(self.widgets["loadPose_categories"])
+        
+
+            
+        #show the window
+        cmds.showWindow(self.widgets["window"])
+	
+	#load poses
+	self.loadPoses()
+        
+        
+    # # # # # # # # # # # # # # # # # # # # # # # # # # # # # # # # # # # # # # # # # # # # # # # # # # # # # # # # # # # # # # # 
+    # # # # # # # # # # # # # # # # # # # # # # # # # # # # # # # # # # # # # # # # # # # # # # # # # # # # # # # # # # # # # # # 
+    # # # # # # # # # # # # # # # # # # # # # # # # # # # # # # # # # # # # # # # # # # # # # # # # # # # # # # # # # # # # # # # 
+    # # # # # # # # # # # # # # # # # # # # # # # # # # # # # # # # # # # # # # # # # # # # # # # # # # # # # # # # # # # # # # # 
+    # # # # # # # # # # # # # # # # # # # # # # # # # # # # # # # # # # # # # # # # # # # # # # # # # # # # # # # # # # # # # # # 
+    def filterResults(self, *args):
+	
+	#clear out any existing children of the row column layout first
+	children = cmds.rowColumnLayout(self.widgets["loadPose_poses"], q = True, childArray = True)
+	
+	#search through children labels, looking for search text. add matches to list and delete others from the UI
+	searchText = cmds.textField(self.widgets["loadPose_filter"], q = True, text = True)
+	if searchText == "" or searchText == "Search":
+	    self.loadPoses()
+	    
+	else:
+	    filtered = []
+	    for child in children:
+		label = cmds.iconTextButton(child, q=True, label = True)
+		if label.find(searchText) != -1:
+		    filtered.append(child)
+		    
+	    for child in children:
+		if child not in filtered:
+		    cmds.deleteUI(child)
+		
+    # # # # # # # # # # # # # # # # # # # # # # # # # # # # # # # # # # # # # # # # # # # # # # # # # # # # # # # # # # # # # # # 
+    # # # # # # # # # # # # # # # # # # # # # # # # # # # # # # # # # # # # # # # # # # # # # # # # # # # # # # # # # # # # # # # 
+    # # # # # # # # # # # # # # # # # # # # # # # # # # # # # # # # # # # # # # # # # # # # # # # # # # # # # # # # # # # # # # # 
+    # # # # # # # # # # # # # # # # # # # # # # # # # # # # # # # # # # # # # # # # # # # # # # # # # # # # # # # # # # # # # # # 
+    # # # # # # # # # # # # # # # # # # # # # # # # # # # # # # # # # # # # # # # # # # # # # # # # # # # # # # # # # # # # # # # 
+    def addPoseCategory(self, *args):
+	
+	selectedItem = cmds.optionMenu(self.widgets["savePose_categories"], q= True, v = True)
+	if selectedItem == "Add New Category":
+	    result = cmds.promptDialog(title = "Add Category", message = "Category Name:", button = ["Accept", "Cancel"])
+	    if result == "Accept":
+		text = cmds.promptDialog(q = True, text = True)
+		
+		#get categories from option menu. if this category exists, warn and return. else, add to the list
+		children = cmds.optionMenu(self.widgets["savePose_categories"], q = True, itemListLong = True)
+		categories = []
+		if children != None:
+		    for child in children:
+			label = cmds.menuItem(child, q = True, label = True)
+			categories.append(label)
+
+		if text not in categories:
+		    #create the directory
+		    project = cmds.optionMenu(self.widgets["savePose_projects"], q = True, value = True)
+		    path = self.mayaToolsDir + "/General/ART/Projects/" + project + "/Poses/" + text
+		    if not os.path.exists(path):
+			os.makedirs(path)
+			
+		    self.getCategories(self.widgets["savePose_categories"])
+		    cmds.optionMenu(self.widgets["savePose_categories"], edit = True, v = text)
+
+		else:
+		    cmds.warning("Category already exists.")
+		    return
+
+
+    # # # # # # # # # # # # # # # # # # # # # # # # # # # # # # # # # # # # # # # # # # # # # # # # # # # # # # # # # # # # # # # 
+    # # # # # # # # # # # # # # # # # # # # # # # # # # # # # # # # # # # # # # # # # # # # # # # # # # # # # # # # # # # # # # # 
+    # # # # # # # # # # # # # # # # # # # # # # # # # # # # # # # # # # # # # # # # # # # # # # # # # # # # # # # # # # # # # # # 
+    # # # # # # # # # # # # # # # # # # # # # # # # # # # # # # # # # # # # # # # # # # # # # # # # # # # # # # # # # # # # # # # 
+    # # # # # # # # # # # # # # # # # # # # # # # # # # # # # # # # # # # # # # # # # # # # # # # # # # # # # # # # # # # # # # # 
+    def savePose(self, selectionOnly, *args):
+        
+        selection = cmds.ls(sl = True)
+	
+        #check to see what spaces the character is in
+	warningNodes = []
+
+	safeAttrs = ["space_world", "space_chest_ik_world_aligned", "space_body_anim", "space_head_fk_anim", "space_chest_ik_anim"]
+	
+	#need to find all space switch nodes for the current character
+	cmds.select(self.character + ":*_space_switcher_follow")
+	nodes = cmds.ls(sl = True)
+	spaceSwitchers = []
+	for node in nodes:
+	    if node.find("invis") == -1:
+		spaceSwitchers.append(node)
+		
+	
+	selectNodes = []
+        for control in spaceSwitchers:
+            spaceSwitchNode = control.rpartition("_follow")[0]
+	    selectNodes.append(spaceSwitchNode)
+	    
+	cmds.select(clear = True)
+	for node in selectNodes:
+	    attrs = cmds.listAttr(node, string = "space_*")
+	    
+	    for attr in attrs:
+		if cmds.getAttr(node + "." + attr) != 0:
+		    if not attr in safeAttrs:
+			warningNodes.append(node + "." + attr)
+		    
+	#if length of warning nodes is greater than 0, warn user that cannot save poses with spaces
+	if len(warningNodes) > 0:
+	    for node in warningNodes:
+		print node
+	    cmds.warning("Cannot save poses if all controls are not in default spaces. Please switch spaces back to defaults and try again")
+
+	    
+	    return
+	
+	    
+	    
+        #bring up a UI that has inputs for project and name of pose
+        if cmds.window("savePoseWindow", exists = True):
+            cmds.deleteUI("savePoseWindow")
+            
+        self.widgets["savePose_window"] = cmds.window("savePoseWindow", w = 300, h = 600, title = "Save Pose", sizeable = False)
+        
+        self.widgets["savePose_mainLayout"] = cmds.columnLayout(w =300, h = 600, rs = 10, columnOffset = ("both", 5))
+        
+        
+        #create the projects field, the name text field, and the button
+        cmds.separator(h = 5)
+        cmds.text(label = "Project:", font = "boldLabelFont", parent = self.widgets["savePose_mainLayout"])
+        self.widgets["savePose_projects"] = cmds.optionMenu(w = 300, h = 30, parent = self.widgets["savePose_mainLayout"])
+	
+        cmds.text(label = "Category:", font = "boldLabelFont", parent = self.widgets["savePose_mainLayout"])
+        self.widgets["savePose_categories"] = cmds.optionMenu(w = 300, h = 30, parent = self.widgets["savePose_mainLayout"], cc = self.addPoseCategory)
+	cmds.menuItem(label = "None", parent = self.widgets["savePose_categories"])
+	cmds.menuItem(label = "Add New Category", parent = self.widgets["savePose_categories"])
+	
+	
+	
+        self.widgets["savePose_name"] = cmds.textField(w = 300, h = 30, text = "Pose Name", parent = self.widgets["savePose_mainLayout"])
+	
+	
+	#add change command to project option menu
+	cmds.optionMenu(self.widgets["savePose_projects"], edit = True, cc = partial(self.getCategories, self.widgets["savePose_categories"]))
+
+
+        
+        #Create a model editor area at the button for the thumbnail
+	self.widgets["savePose_modelForm"] = cmds.formLayout(w = 300, h = 300, parent = self.widgets["savePose_mainLayout"])
+	self.widgets["savePose_modelEditorBorder"] = cmds.image(bgc = [0,0,0], parent = self.widgets["savePose_modelForm"], w = 300, h = 300)
+        self.widgets["savePose_modelEditor"] = cmds.modelEditor(parent = self.widgets["savePose_modelForm"])
+	cmds.formLayout(self.widgets["savePose_modelForm"], edit = True, af = [(self.widgets["savePose_modelEditorBorder"], 'top', 0), (self.widgets["savePose_modelEditorBorder"], 'bottom', 0), (self.widgets["savePose_modelEditorBorder"], 'right', 0), (self.widgets["savePose_modelEditorBorder"], 'left', 0)])	
+	cmds.formLayout(self.widgets["savePose_modelForm"], edit = True, af = [(self.widgets["savePose_modelEditor"], 'top', 5), (self.widgets["savePose_modelEditor"], 'bottom', 5), (self.widgets["savePose_modelEditor"], 'right', 5), (self.widgets["savePose_modelEditor"], 'left', 5)])
+
+
+
+	#Create a camera for the editor
+	self.screenshot_camera = cmds.camera(name = "save_pose_cam", worldUp = (0, 0, 1), rotation = (90, 0, 0))
+	cmds.setAttr(self.screenshot_camera[1] + ".focalLength", 70)
+	cmds.viewFit(self.screenshot_camera[0])
+	
+	
+	#Attach the camera to the model editor.
+	cmds.modelEditor(self.widgets["savePose_modelEditor"], edit = True, camera = self.screenshot_camera[0] )
+	
+	
+	
+
+
+	#create the save pose button
+        self.widgets["saveButtonSaveButtonLayout"] = cmds.rowColumnLayout(nc = 3, cw = [(1, 80), (2, 140), (3, 80)], parent = self.widgets["savePose_mainLayout"])
+        cmds.text(label = "", parent = self.widgets["saveButtonSaveButtonLayout"])
+        self.widgets["savePose_saveButton"] = cmds.symbolButton(w = 140, h = 50, image = self.mayaToolsDir + "/General/Icons/ART/savePose.bmp", c = partial(self.savePose_execute, selectionOnly), parent = self.widgets["saveButtonSaveButtonLayout"])
+	
+	
+        #show the window
+        cmds.showWindow(self.widgets["savePose_window"])
+        
+        #populate projects
+        self.getProjects(self.widgets["savePose_projects"])
+	
+	
+	#set to favorite project
+	settingsLocation = self.mayaToolsDir + "/General/Scripts/projectSettings.txt"
+	if os.path.exists(settingsLocation):
+	    f = open(settingsLocation, 'r')
+	    settings = cPickle.load(f)
+	    favoriteProject = settings.get("FavoriteProject")
+	    
+	    try:
+		cmds.optionMenu(self.widgets["savePose_projects"], edit = True, v = favoriteProject)
+	    except:
+		pass
+	    
+	
+	#turn on smooth shading & textures
+	cmds.modelEditor(self.widgets["savePose_modelEditor"], edit = True, nurbsCurves = False, displayAppearance = "smoothShaded", displayTextures = True, headsUpDisplay = False, cameras = False, grid = False, joints = False, textures = True )
+
+	#if there was a selection, re-select it
+	if selection:
+	    cmds.select(selection)
+	    
+	#get categories
+	self.getCategories(self.widgets["savePose_categories"])
+    
+
+    # # # # # # # # # # # # # # # # # # # # # # # # # # # # # # # # # # # # # # # # # # # # # # # # # # # # # # # # # # # # # # # 
+    # # # # # # # # # # # # # # # # # # # # # # # # # # # # # # # # # # # # # # # # # # # # # # # # # # # # # # # # # # # # # # # 
+    # # # # # # # # # # # # # # # # # # # # # # # # # # # # # # # # # # # # # # # # # # # # # # # # # # # # # # # # # # # # # # # 
+    # # # # # # # # # # # # # # # # # # # # # # # # # # # # # # # # # # # # # # # # # # # # # # # # # # # # # # # # # # # # # # # 
+    # # # # # # # # # # # # # # # # # # # # # # # # # # # # # # # # # # # # # # # # # # # # # # # # # # # # # # # # # # # # # # # 
+    def savePose_execute(self, selectionOnly, *args):
+        
+        #get project name
+        project = cmds.optionMenu(self.widgets["savePose_projects"], q = True, value = True)
+	category = cmds.optionMenu(self.widgets["savePose_categories"], q = True, value = True)
+        sourceControl = None
+	
+        #get pose name
+        pose = cmds.textField(self.widgets["savePose_name"], q = True, text = True)
+        
+        if pose == "" or pose == "Pose Name":
+            cmds.warning("Invalid Pose Name")
+            return
+        
+        #sort out category
+        if category == "None" or category == "Add New Category":
+	    category = "No Category"
+	    path = self.mayaToolsDir + "/General/ART/Projects/" + project + "/Poses/" + category
+	    if not os.path.exists(path):
+		os.makedirs(path)
+	    
+        #generate export path
+        exportPath = self.mayaToolsDir + "/General/ART/Projects/" + project + "/Poses/" + category + "/"
+
+        if not os.path.exists(exportPath):
+            os.makedirs(exportPath)
+            
+            
+        #generate pose data, passing in export path
+        self.createPose(exportPath, pose, selectionOnly)
+	
+	
+	
+	#check if source control is on
+	settingsLocation = self.mayaToolsDir + "/General/Scripts/projectSettings.txt"
+	if os.path.exists(settingsLocation):
+	    f = open(settingsLocation, 'r')
+	    settings = cPickle.load(f)
+	    f.close()	    
+	    sourceControl = settings.get("UseSourceControl")
+	    
+	    
+	
+	
+	#save thumbnail
+	sucess = False
+	currentTime = cmds.currentTime(q = True)
+	
+	#selection only poses
+	if selectionOnly:
+	    if cmds.playblast(activeEditor = True) != self.widgets["savePose_modelEditor"]:
+		cmds.modelEditor(self.widgets["savePose_modelEditor"], edit = True, activeView = True)
+		
+	    try:
+		f = open(self.mayaToolsDir + "/General/ART/Projects/" + project + "/Poses/" + category + "/" + pose + "__SelectionOnly.bmp", 'w')
+		f.close()
+		success = True
+		
+	    
+	    except:
+		if sourceControl == False:
+		    cmds.confirmDialog(title = "Save Pose", icon = "critical", message = "Could not create thumbnail. Make sure the file is not read only. Aborting operation.")
+		    return 
+		
+		else:
+		    #if using source control, check to see if we can check out the file
+		    result = cmds.confirmDialog(title = "Save Pose", icon = "critical", message = "Could not save thumbnail. File may exist already and be marked as read only.", button = ["Check Out File", "Cancel"])
+				    
+		    if result == "Check Out File":
+			import perforceUtils
+			reload(perforceUtils)
+			writeable = perforceUtils.p4_checkOutCurrentFile(self.mayaToolsDir + "/General/ART/Projects/" + project + "/Poses/" + category + "/" + pose + "__SelectionOnly.bmp")
+			
+			if writeable:
+			    #now that it is checked out, try saving again
+			    try:
+				f = open(self.mayaToolsDir + "/General/ART/Projects/" + project + "/Poses/" + category + "/" + pose + "__SelectionOnly.bmp", 'w')
+				f.close()				
+				success = True
+				
+				
+			    except:
+				cmds.confirmDialog(title = "Save Pose", icon = "critical", message = "Perforce operation unsucessful. Could not save file. Aborting operation.")
+				return 
+			    
+			else:
+			    cmds.confirmDialog(title = "Save Pose", icon = "critical", message = "Perforce operation unsucessful. Could not save file. Aborting operation.")
+			    return 
+			
+		    else:
+			cmds.warning("Operation Aborted.")
+			return 
+		    
+		    
+	#full body poses
+	else:
+	    
+	    if cmds.playblast(activeEditor = True) != self.widgets["savePose_modelEditor"]:
+		cmds.modelEditor(self.widgets["savePose_modelEditor"], edit = True, activeView = True)
+		
+		
+	    try:
+		f = open(self.mayaToolsDir + "/General/ART/Projects/" + project + "/Poses/" + category + "/" + pose + ".bmp", 'w')
+		success = True
+		f.close()
+		
+	    except:
+		if sourceControl == False:
+		    cmds.confirmDialog(title = "Save Pose", icon = "critical", message = "Could not create thumbnail. Make sure the file is not read only. Aborting operation.")
+		    return 
+		
+		else:
+		    #if using source control, check to see if we can check out the file
+		    result = cmds.confirmDialog(title = "Save Pose", icon = "critical", message = "Could not save thumbnail. File may exist already and be marked as read only.", button = ["Check Out File", "Cancel"])
+				    
+		    if result == "Check Out File":
+			import perforceUtils
+			reload(perforceUtils)
+			writeable = perforceUtils.p4_checkOutCurrentFile(self.mayaToolsDir + "/General/ART/Projects/" + project + "/Poses/" + category + "/" + pose + ".bmp")
+			
+			if writeable:
+			    #now that it is checked out, try saving again
+			    try:
+				f = open(self.mayaToolsDir + "/General/ART/Projects/" + project + "/Poses/" + category + "/" + pose + ".bmp", 'w')
+				f.close()
+				success = True
+				
+			    except:
+				cmds.confirmDialog(title = "Save Pose", icon = "critical", message = "Perforce operation unsucessful. Could not save file. Aborting operation.")
+				return 
+			    
+			else:
+			    cmds.confirmDialog(title = "Save Pose", icon = "critical", message = "Perforce operation unsucessful. Could not save file. Aborting operation.")
+			    return 
+			
+		    else:
+			cmds.confirmDialog(title = "Save Pose", icon = "information", message = "Aborting operation.")
+			return 	    
+	
+	
+	
+	if success == True:
+	    
+	    if selectionOnly:
+		cmds.playblast(frame = currentTime, format = "image", cf = self.mayaToolsDir + "/General/ART/Projects/" + project + "/Poses/" + category + "/" + pose + "__SelectionOnly.bmp", v = False, wh = [100, 100], p = 100)
+	    
+	    else:
+		cmds.playblast(frame = currentTime, format = "image", cf = self.mayaToolsDir + "/General/ART/Projects/" + project + "/Poses/" + category + "/" + pose + ".bmp", v = False, wh = [100, 100], p = 100)
+		
+	    #delete the UI
+	    cmds.deleteUI("savePoseWindow")
+    
+	    #delete the camera
+	    cmds.delete(self.screenshot_camera[0])
+	    
+	    self.getCategories(self.widgets["loadPose_categories"])
+	    self.loadPoses()
+	    
+	    if sourceControl:
+		result = cmds.confirmDialog(icon = "information", title = "Operation Complete", message = "Pose has been successfully created!", button = ["Submit Created Files", "Continue Without Submitting"])
+			    
+		if result == "Submit Created Files":
+		    
+		    submittedFiles = []
+		    
+		    #import perforce utils
+		    import perforceUtils
+		    reload(perforceUtils)
+		    
+		    #get a description for the changelist
+		    cmds.promptDialog(title = "Perforce", message = "Please Enter a Description..", button = ["Submit"])
+		    desc = cmds.promptDialog(q = True, text = True)
+		    
+		    #submit the export file
+		    if selectionOnly:
+			
+			#THUMBNAIL
+			submitted = perforceUtils.p4_submitCurrentFile(self.mayaToolsDir + "/General/ART/Projects/" + project + "/Poses/" + category + "/" + pose + "__SelectionOnly.bmp", desc)
+			
+			#if the submit failed because the file is not in a changelist, add the file now
+			if submitted == False:
+			    added = perforceUtils.p4_addAndSubmitCurrentFile(self.mayaToolsDir + "/General/ART/Projects/" + project + "/Poses/" + category + "/" + pose + "__SelectionOnly.bmp", desc)
+			    
+			    if added:
+				submittedFiles.append(self.mayaToolsDir + "/General/ART/Projects/" + project + "/Poses/" + category + "/" + pose + "__SelectionOnly.bmp")
+				
+			if submitted == True:
+			    submittedFiles.append(self.mayaToolsDir + "/General/ART/Projects/" + project + "/Poses/" + category + "/" + pose + "__SelectionOnly.bmp")
+			    
+			    
+			#DATA FILE
+			submitted = perforceUtils.p4_submitCurrentFile(exportPath + pose + "__SelectionOnly.txt", desc)
+			
+			#if the submit failed because the file is not in a changelist, add the file now
+			if submitted == False:
+			    added = perforceUtils.p4_addAndSubmitCurrentFile(exportPath + pose + "__SelectionOnly.txt", desc)
+			    
+			    if added:
+				submittedFiles.append(exportPath + pose + "__SelectionOnly.txt",)
+				
+			if submitted == True:
+			    submittedFiles.append(exportPath + pose + "__SelectionOnly.txt",)
+			    
+			    
+			    
+			
+		    #full pose
+		    else:
+			#THUMBNAIL
+			submitted = perforceUtils.p4_submitCurrentFile(self.mayaToolsDir + "/General/ART/Projects/" + project + "/Poses/" + category + "/" + pose + ".bmp", desc)
+			
+			#if the submit failed because the file is not in a changelist, add the file now
+			if submitted == False:
+			    added = perforceUtils.p4_addAndSubmitCurrentFile(self.mayaToolsDir + "/General/ART/Projects/" + project + "/Poses/" + category + "/" + pose + ".bmp", desc)
+			    
+			    if added:
+				submittedFiles.append(self.mayaToolsDir + "/General/ART/Projects/" + project + "/Poses/" + category + "/" + pose + ".bmp")
+				
+			if submitted == True:
+			    submittedFiles.append(self.mayaToolsDir + "/General/ART/Projects/" + project + "/Poses/" + category + "/" + pose + ".bmp")			
+			
+			
+			#DATA FILE
+			submitted = perforceUtils.p4_submitCurrentFile(exportPath + pose + ".txt", desc)
+			
+			#if the submit failed because the file is not in a changelist, add the file now
+			if submitted == False:
+			    added = perforceUtils.p4_addAndSubmitCurrentFile(exportPath + pose + ".txt", desc)
+			    
+			    if added:
+				submittedFiles.append(exportPath + pose + ".txt",)
+				
+			if submitted == True:
+			    submittedFiles.append(exportPath + pose + ".txt",)			
+			
+			
+
+			
+			
+			
+	    
+	else:
+	    cmds.warning("Operation Failed")
+    # # # # # # # # # # # # # # # # # # # # # # # # # # # # # # # # # # # # # # # # # # # # # # # # # # # # # # # # # # # # # # # 
+    # # # # # # # # # # # # # # # # # # # # # # # # # # # # # # # # # # # # # # # # # # # # # # # # # # # # # # # # # # # # # # # 
+    # # # # # # # # # # # # # # # # # # # # # # # # # # # # # # # # # # # # # # # # # # # # # # # # # # # # # # # # # # # # # # # 
+    # # # # # # # # # # # # # # # # # # # # # # # # # # # # # # # # # # # # # # # # # # # # # # # # # # # # # # # # # # # # # # # 
+    # # # # # # # # # # # # # # # # # # # # # # # # # # # # # # # # # # # # # # # # # # # # # # # # # # # # # # # # # # # # # # # 
+    def createPose(self, path, poseName, selectionOnly, *args):
+        sourceControl = None
+        data = []
+        
+	#check if source control is on
+	settingsLocation = self.mayaToolsDir + "/General/Scripts/projectSettings.txt"
+	if os.path.exists(settingsLocation):
+	    f = open(settingsLocation, 'r')
+	    settings = cPickle.load(f)
+	    f.close()	    
+	    sourceControl = settings.get("UseSourceControl")
+	    
+	    
+        if selectionOnly:
+	    #for every control, get the attributes of the control and write to file
+	    selection = cmds.ls(sl = True)
+	    for control in selection:
+		control = control.partition(":")[2]
+		
+		
+		#if control exists, get data from attrs, and add control to data list
+		if cmds.objExists(self.character + ":" + control):
+		    attrs = cmds.listAttr(self.character + ":" + control, keyable = True, unlocked = True)
+		    datalist = [control, attrs]
+		    
+		    
+		    #list attrs
+		    attrs = cmds.listAttr(self.character + ":" + control, keyable = True, unlocked = True)
+		
+		    for attr in attrs:
+			value = cmds.getAttr(self.character + ":" + control + "." + attr)
+			datalist.append(value)
+			
+		    #add data list to data
+		    data.append(datalist)
+		    
+	    #write data to file
+	    if os.path.exists(path + poseName + "__SelectionOnly.txt"):
+		result = cmds.confirmDialog(title = "Save Pose", message = "Pose with that name already exists.", button = ["Overwrite", "Cancel"], defaultButton = "Overwrite", cancelButton = "Cancel", dismissString = "Cancel")
+		if result == "Cancel":
+		    return
+		
+		
+	    try:
+		f = open(path + poseName + "__SelectionOnly.txt", 'w')
+		
+	    
+	    except:
+		if sourceControl == False:
+		    cmds.confirmDialog(title = "Save Pose", icon = "critical", message = "Cannot complete operation as file is not writeable.")
+		    return
+		
+		else:
+		    result = cmds.confirmDialog(title = "Save Pose", icon = "critical", message = "Could not save the pose file. File may exist already and be marked as read only.", button = ["Check Out File", "Cancel"])
+		    
+		    if result == "Check Out File":
+			import perforceUtils
+			reload(perforceUtils)
+			writeable = perforceUtils.p4_checkOutCurrentFile(path + poseName + "__SelectionOnly.txt")
+			
+			if writeable:
+			    try:
+				f = open(path + poseName + "__SelectionOnly.txt", 'w')
+				
+			    except:
+				cmds.confirmDialog(title = "Publish", icon = "critical", message = "Perforce operation unsucessful. Could not save file. Aborting operation.")
+				return False
+			    
+			else:
+			    cmds.confirmDialog(title = "Publish", icon = "critical", message = "Perforce operation unsucessful. Could not save file. Aborting operation.")
+			    return False
+			
+		    else:
+			cmds.warning("Operation Aborted")
+			return False
+			
+			
+	    cPickle.dump(data, f)
+	    f.close()
+	    
+	    
+	    
+	    
+	
+	else:
+	    #TOES AND CUSTOM JOINTS NOT IN THIS LIST YET!!
+	    #for every control, get the attributes of the control and write to file
+	    for control in self.controls:
+		
+		
+		#if control exists, get data from attrs, and add control to data list
+		if cmds.objExists(self.character + ":" + control):
+		    attrs = cmds.listAttr(self.character + ":" + control, keyable = True, unlocked = True)
+		    datalist = [control, attrs]
+		    
+		    
+		    
+		    #list attrs
+		    attrs = cmds.listAttr(self.character + ":" + control, keyable = True, unlocked = True)
+		    if attrs != None:
+			for attr in attrs:
+			    value = cmds.getAttr(self.character + ":" + control + "." + attr)
+			    datalist.append(value)
+
+		    #add data list to data
+		    data.append(datalist)
+		    
+	    #write data to file
+	    if os.path.exists(path + poseName + ".txt"):
+		result = cmds.confirmDialog(title = "Save Pose", message = "Pose with that name already exists.", button = ["Overwrite", "Cancel"], defaultButton = "Overwrite", cancelButton = "Cancel", dismissString = "Cancel")
+		if result == "Cancel":
+		    return
+		
+		
+	    try:
+		f = open(path + poseName + ".txt", 'w')
+		
+	    except:
+		if sourceControl == False:
+		    cmds.confirmDialog(title = "Save Pose", icon = "critical", message = "Cannot complete operation as file is not writeable.")
+		    return
+		
+		else:
+		    result = cmds.confirmDialog(title = "Save Pose", icon = "critical", message = "Could not save the pose file. File may exist already and be marked as read only.", button = ["Check Out File", "Cancel"])
+		    
+		    if result == "Check Out File":
+			import perforceUtils
+			reload(perforceUtils)
+			writeable = perforceUtils.p4_checkOutCurrentFile(path + poseName + ".txt")
+			
+			if writeable:
+			    try:
+				f = open(path + poseName + ".txt", 'w')
+				
+			    except:
+				cmds.confirmDialog(title = "Publish", icon = "critical", message = "Perforce operation unsucessful. Could not save file. Aborting operation.")
+				return False
+			    
+			else:
+			    cmds.confirmDialog(title = "Publish", icon = "critical", message = "Perforce operation unsucessful. Could not save file. Aborting operation.")
+			    return False
+			
+		    else:
+			cmds.warning("Operation Aborted")
+			return False
+		    
+		
+	    cPickle.dump(data, f)
+	    f.close()
+	
+	
+            
+    # # # # # # # # # # # # # # # # # # # # # # # # # # # # # # # # # # # # # # # # # # # # # # # # # # # # # # # # # # # # # # # 
+    # # # # # # # # # # # # # # # # # # # # # # # # # # # # # # # # # # # # # # # # # # # # # # # # # # # # # # # # # # # # # # # 
+    # # # # # # # # # # # # # # # # # # # # # # # # # # # # # # # # # # # # # # # # # # # # # # # # # # # # # # # # # # # # # # # 
+    # # # # # # # # # # # # # # # # # # # # # # # # # # # # # # # # # # # # # # # # # # # # # # # # # # # # # # # # # # # # # # # 
+    # # # # # # # # # # # # # # # # # # # # # # # # # # # # # # # # # # # # # # # # # # # # # # # # # # # # # # # # # # # # # # # 
+    def loadPoses_All(self, *args):
+	
+	#clear out any existing children of the row column layout first
+	children = cmds.rowColumnLayout(self.widgets["loadPose_poses"], q = True, childArray = True)
+	
+	if children != None:
+	    for child in children:
+		cmds.deleteUI(child)
+	
+	#when the UI first launches, get the project in the optionMenu and find all poses in that project
+	project = cmds.optionMenu(self.widgets["loadPose_projects"], q = True, value = True)
+	
+	#get all categories in this project
+	categories = os.listdir(self.mayaToolsDir + "/General/ART/Projects/" + project + "/Poses/")
+	
+	for category in categories:
+	    try:
+		path = self.mayaToolsDir + "/General/ART/Projects/" + project + "/Poses/" + category + "/"
+		
+		poses = os.listdir(path)
+		thumbs = []
+		partialThumbs = []
+		
+		for pose in poses:
+		    if pose.rpartition(".")[2] == "bmp":
+			if pose.find("__SelectionOnly") == -1:
+			    thumbs.append(pose)
+			else:
+			    partialThumbs.append(pose)
+			
+		for thumb in thumbs:
+		    button = cmds.iconTextButton(image = path + thumb, w= 100, h = 120, parent = self.widgets["loadPose_poses"], label = thumb.partition(".")[0], style = "iconAndTextVertical", c = partial(self.loadPose_exact, path, thumb.partition(".")[0]))
+		
+		    #create the popup menu for the button
+		    menu = cmds.popupMenu(b = 3, parent =button)
+		    cmds.menuItem(label = "Load Pose on Selected Controls", parent = menu, c = partial(self.loadPose_selection, path, thumb.partition(".")[0]))
+		    cmds.menuItem(label = "Load Mirrored Pose", parent = menu, c = partial(self.loadPose_mirrored, path, thumb.partition(".")[0]))
+		    cmds.menuItem(label = "Load Ghost Pose", parent = menu, c = partial(self.loadPose_ghost, path, thumb.partition(".")[0]))
+		    cmds.menuItem(label = "Load Mirrored Ghost Pose", parent = menu, c = partial(self.loadPose_ghost_mirrored, path, thumb.partition(".")[0]))
+		    
+		    cmds.menuItem(divider = True, parent = menu)
+		    cmds.menuItem(label = "Delete Pose", parent = menu, c = partial(self.deletePose, path, thumb.partition(".")[0]))
+		    cmds.menuItem(label = "Find on Disk", parent = menu, c = partial(self.findPoseFileOnDisk, path, thumb.partition(".")[0]))
+		    
+		for thumb in partialThumbs:
+		    button = cmds.iconTextButton(image = path + thumb, imageOverlayLabel = "*PART", w= 100, h = 120, parent = self.widgets["loadPose_poses"], label = thumb.rpartition("__SelectionOnly")[0] + " (partial)", style = "iconAndTextVertical", c = partial(self.loadPose_exact, path, thumb.partition(".")[0]))
+		
+		#create the popup menu for the button
+		menu = cmds.popupMenu(b = 3, parent =button)
+		cmds.menuItem(label = "Load Mirrored Pose", parent = menu, c = partial(self.pastePoseOpposite, path, thumb.partition(".")[0]))
+
+		cmds.menuItem(divider = True, parent = menu)
+		cmds.menuItem(label = "Delete Pose", parent = menu, c = partial(self.deletePose, path, thumb.partition(".")[0]))
+		cmds.menuItem(label = "Find on Disk", parent = menu, c = partial(self.findPoseFileOnDisk, path, thumb.partition(".")[0]))
+		
+		
+	    except:
+		pass
+	
+	
+    # # # # # # # # # # # # # # # # # # # # # # # # # # # # # # # # # # # # # # # # # # # # # # # # # # # # # # # # # # # # # # # 
+    # # # # # # # # # # # # # # # # # # # # # # # # # # # # # # # # # # # # # # # # # # # # # # # # # # # # # # # # # # # # # # # 
+    # # # # # # # # # # # # # # # # # # # # # # # # # # # # # # # # # # # # # # # # # # # # # # # # # # # # # # # # # # # # # # # 
+    # # # # # # # # # # # # # # # # # # # # # # # # # # # # # # # # # # # # # # # # # # # # # # # # # # # # # # # # # # # # # # # 
+    # # # # # # # # # # # # # # # # # # # # # # # # # # # # # # # # # # # # # # # # # # # # # # # # # # # # # # # # # # # # # # # 
+    def loadPoses(self, *args):
+	
+	#clear out any existing children of the row column layout first
+	children = cmds.rowColumnLayout(self.widgets["loadPose_poses"], q = True, childArray = True)
+	
+	if children != None:
+	    for child in children:
+		cmds.deleteUI(child)
+	
+	#when the UI first launches, get the project in the optionMenu and find all poses in that project
+	project = cmds.optionMenu(self.widgets["loadPose_projects"], q = True, value = True)
+	category = cmds.optionMenu(self.widgets["loadPose_categories"], q = True, value = True)
+	
+	try:
+	    path = self.mayaToolsDir + "/General/ART/Projects/" + project + "/Poses/" + category + "/"
+	    
+	    poses = os.listdir(path)
+	    thumbs = []
+	    partialThumbs = []
+	    
+	    for pose in poses:
+		if pose.rpartition(".")[2] == "bmp":
+		    if pose.find("__SelectionOnly") == -1:
+			thumbs.append(pose)
+		    else:
+			partialThumbs.append(pose)
+		    
+	    for thumb in thumbs:
+		button = cmds.iconTextButton(image = path + thumb, w= 100, h = 120, parent = self.widgets["loadPose_poses"], label = thumb.partition(".")[0], style = "iconAndTextVertical", c = partial(self.loadPose_exact, path, thumb.partition(".")[0]))
+	    
+		#create the popup menu for the button
+		menu = cmds.popupMenu(b = 3, parent =button)
+		cmds.menuItem(label = "Load Pose on Selected Controls", parent = menu, c = partial(self.loadPose_selection, path, thumb.partition(".")[0]))
+		cmds.menuItem(label = "Load Mirrored Pose", parent = menu, c = partial(self.loadPose_mirrored, path, thumb.partition(".")[0]))
+		cmds.menuItem(label = "Load Ghost Pose", parent = menu, c = partial(self.loadPose_ghost, path, thumb.partition(".")[0]))
+		cmds.menuItem(label = "Load Mirrored Ghost Pose", parent = menu, c = partial(self.loadPose_ghost_mirrored, path, thumb.partition(".")[0]))
+		
+		cmds.menuItem(divider = True, parent = menu)
+		cmds.menuItem(label = "Delete Pose", parent = menu, c = partial(self.deletePose, path, thumb.partition(".")[0]))
+		cmds.menuItem(label = "Find on Disk", parent = menu, c = partial(self.findPoseFileOnDisk, path, thumb.partition(".")[0]))
+		
+		
+	    for thumb in partialThumbs:
+		button = cmds.iconTextButton(image = path + thumb, imageOverlayLabel = "*PART", w= 100, h = 120, parent = self.widgets["loadPose_poses"], label = thumb.rpartition("__SelectionOnly")[0] + " (partial)", style = "iconAndTextVertical", c = partial(self.loadPose_partial, path, thumb.partition(".")[0]))
+	    
+		#create the popup menu for the button
+		menu = cmds.popupMenu(b = 3, parent =button)
+		cmds.menuItem(label = "Load Mirrored Pose", parent = menu, c = partial(self.pastePoseOpposite, path, thumb.partition(".")[0]))
+
+		cmds.menuItem(divider = True, parent = menu)
+		cmds.menuItem(label = "Delete Pose", parent = menu, c = partial(self.deletePose, path, thumb.partition(".")[0]))
+		cmds.menuItem(label = "Find on Disk", parent = menu, c = partial(self.findPoseFileOnDisk, path, thumb.partition(".")[0]))
+		
+	except:
+	    pass
+	    
+	    
+    # # # # # # # # # # # # # # # # # # # # # # # # # # # # # # # # # # # # # # # # # # # # # # # # # # # # # # # # # # # # # # # 
+    # # # # # # # # # # # # # # # # # # # # # # # # # # # # # # # # # # # # # # # # # # # # # # # # # # # # # # # # # # # # # # # 
+    # # # # # # # # # # # # # # # # # # # # # # # # # # # # # # # # # # # # # # # # # # # # # # # # # # # # # # # # # # # # # # # 
+    # # # # # # # # # # # # # # # # # # # # # # # # # # # # # # # # # # # # # # # # # # # # # # # # # # # # # # # # # # # # # # # 
+    # # # # # # # # # # # # # # # # # # # # # # # # # # # # # # # # # # # # # # # # # # # # # # # # # # # # # # # # # # # # # # # 
+    def deletePose(self, path, poseName, *args):
+	
+	dataFile = path + poseName + ".txt"
+	thumbnail = path + poseName + ".bmp"
+	
+
+	try:    
+	    os.remove(dataFile)
+	    os.remove(thumbnail)
+	    cmds.button("poseEditor_syncStatusButton", edit = True, visible = True)
+	    
+	except:
+	    cmds.confirmDialog(title = "Delete Pose", message = "Pose Data is marked as Read Only. Cannot delete file")
+	    return
+	    
+	    
+	    
+
+	
+    # # # # # # # # # # # # # # # # # # # # # # # # # # # # # # # # # # # # # # # # # # # # # # # # # # # # # # # # # # # # # # # 
+    # # # # # # # # # # # # # # # # # # # # # # # # # # # # # # # # # # # # # # # # # # # # # # # # # # # # # # # # # # # # # # # 
+    # # # # # # # # # # # # # # # # # # # # # # # # # # # # # # # # # # # # # # # # # # # # # # # # # # # # # # # # # # # # # # # 
+    # # # # # # # # # # # # # # # # # # # # # # # # # # # # # # # # # # # # # # # # # # # # # # # # # # # # # # # # # # # # # # # 
+    # # # # # # # # # # # # # # # # # # # # # # # # # # # # # # # # # # # # # # # # # # # # # # # # # # # # # # # # # # # # # # # 
+    def findPoseFileOnDisk(self, path, poseName, *args):
+	drive = path.partition(":")[0] + ":\\"
+	restOfPath = path.partition(":")[2]
+	restOfPath = os.path.normpath(restOfPath)
+	fullPath = drive + restOfPath + "\\" + poseName + ".bmp"
+	
+	
+	Popen_arg = "explorer /select, \"" + os.path.normpath(fullPath) + "\""
+	subprocess.Popen(Popen_arg)
+	
+	
+	
+    # # # # # # # # # # # # # # # # # # # # # # # # # # # # # # # # # # # # # # # # # # # # # # # # # # # # # # # # # # # # # # # 
+    # # # # # # # # # # # # # # # # # # # # # # # # # # # # # # # # # # # # # # # # # # # # # # # # # # # # # # # # # # # # # # # 
+    # # # # # # # # # # # # # # # # # # # # # # # # # # # # # # # # # # # # # # # # # # # # # # # # # # # # # # # # # # # # # # # 
+    # # # # # # # # # # # # # # # # # # # # # # # # # # # # # # # # # # # # # # # # # # # # # # # # # # # # # # # # # # # # # # # 
+    # # # # # # # # # # # # # # # # # # # # # # # # # # # # # # # # # # # # # # # # # # # # # # # # # # # # # # # # # # # # # # # 
+    def loadPose_partial(self, path, poseName, *args):
+	currentTime = cmds.currentTime(q = True)
+	
+	#load pose
+	if os.path.exists(path + poseName + ".txt"):
+	    f = open(path + poseName + ".txt", 'r')
+	    poseData = cPickle.load(f)
+	    for data in poseData:
+		control = self.character + ":" + data[0]
+		originalAttrs = data[1]
+		
+		if cmds.objExists(control):
+		    newData = []
+		    for i in range(2, int(len(data))):
+			newData.append(data[i])
+		    
+		    attrs = cmds.listAttr(control, keyable = True, unlocked = True)
+		    for i in range(int(len(attrs))):
+			
+			try:
+			    for attr in originalAttrs:
+				if attrs[i] in originalAttrs:
+				    dataIndex = originalAttrs.index(attrs[i])
+				    cmds.setAttr(control + "." + attrs[i], newData[dataIndex])
+				
+			except :
+			    pass
+			
+			
+    # # # # # # # # # # # # # # # # # # # # # # # # # # # # # # # # # # # # # # # # # # # # # # # # # # # # # # # # # # # # # # # 
+    # # # # # # # # # # # # # # # # # # # # # # # # # # # # # # # # # # # # # # # # # # # # # # # # # # # # # # # # # # # # # # # 
+    # # # # # # # # # # # # # # # # # # # # # # # # # # # # # # # # # # # # # # # # # # # # # # # # # # # # # # # # # # # # # # # 
+    # # # # # # # # # # # # # # # # # # # # # # # # # # # # # # # # # # # # # # # # # # # # # # # # # # # # # # # # # # # # # # # 
+    # # # # # # # # # # # # # # # # # # # # # # # # # # # # # # # # # # # # # # # # # # # # # # # # # # # # # # # # # # # # # # # 
+    def loadPose_exact(self, path, poseName, *args):
+	
+	currentTime = cmds.currentTime(q = True)
+	
+	#capture the character's current rig settings
+	spineMode = cmds.getAttr(self.character + ":Rig_Settings.spine_ik") #if 0, fk, if 1, ik
+	lArmMode = cmds.getAttr(self.character + ":Rig_Settings.lArmMode")
+	rArmMode = cmds.getAttr(self.character + ":Rig_Settings.rArmMode")
+	lLegMode = cmds.getAttr(self.character + ":Rig_Settings.lLegMode")
+	rLegMode = cmds.getAttr(self.character + ":Rig_Settings.rLegMode")
+	
+	
+
+	if poseName.find("__SelectionOnly") == -1:
+	    self.animUIInst.resetAll()
+	    
+	    
+	    #reset space switch nodes
+	    spaceNodes= []
+	    cmds.select(self.character + ":*_space_switcher_follow")
+	    nodes = cmds.ls(sl = True)
+	    spaceSwitchers = []
+	    for node in nodes:
+		if node.find("invis") == -1:
+		    spaceSwitchers.append(node)
+		    
+	    
+	    for control in spaceSwitchers:
+		spaceSwitchNode = control.rpartition("_follow")[0]
+		spaceNodes.append(spaceSwitchNode)
+		
+	    avoidAttrs = ["scaleX", "scaleY", "scaleZ", "visibility"]
+	    for node in spaceNodes:
+		attrs = cmds.listAttr(node, keyable = True, unlocked = True)
+		for attr in attrs:
+		    if attr not in avoidAttrs:
+			cmds.setAttr(node + "." + attr, 0)
+		    
+	    
+	#load pose
+	if os.path.exists(path + poseName + ".txt"):
+	    f = open(path + poseName + ".txt", 'r')
+	    poseData = cPickle.load(f)
+	    for data in poseData:
+		control = self.character + ":" + data[0]
+		originalAttrs = data[1]
+		
+		if cmds.objExists(control):
+		    newData = []
+		    for i in range(2, int(len(data))):
+			newData.append(data[i])
+		    
+		    attrs = cmds.listAttr(control, keyable = True, unlocked = True)
+		    if attrs != None:
+			for i in range(int(len(attrs))):
+			    
+			    try:
+				for attr in originalAttrs:
+				    if attrs[i] in originalAttrs:
+					dataIndex = originalAttrs.index(attrs[i])
+					cmds.setAttr(control + "." + attrs[i], newData[dataIndex])
+				    
+			    except :
+				pass
+			
+	    
+	    #key all
+	    self.animUIInst.selectEverything()
+	    cmds.setKeyframe()
+	    
+	    
+	    
+	    
+	    #refresh
+	    currentFrame = cmds.currentTime(q = True)
+	    cmds.currentTime(currentFrame - 1)
+	    cmds.currentTime(currentFrame)
+	    
+	    
+	    
+	    
+	    #reset/match rig modes back to what they were before operation
+	    
+	    if spineMode == 0:
+		#then we were in FK spine. Check to see if we still are
+		currentSpineMode = cmds.getAttr(self.character + ":Rig_Settings.spine_ik")
+		#if we are currently not in fk, then switch to FK and match
+		if currentSpineMode == 1:
+		    self.animUIInst.match_singleFrame("spine", None, "FK", "IK")
+		    cmds.setAttr(self.character + ":Rig_Settings.spine_ik", 0)
+		    cmds.setAttr(self.character + ":Rig_Settings.spine_fk", 1)
+		    cmds.setKeyframe(self.character + ":Rig_Settings.spine_ik")
+		    cmds.setKeyframe(self.character + ":Rig_Settings.spine_fk")
+		    
+	    if spineMode == 1:
+		#then we were in IK spine. Check to see if we still are
+		currentSpineMode = cmds.getAttr(self.character + ":Rig_Settings.spine_ik")
+		#if we are currently not in ik, then switch to iK and match
+		if currentSpineMode == 0:
+		    self.animUIInst.match_singleFrame("spine", None, "IK", "FK")
+		    cmds.setAttr(self.character  + ":Rig_Settings.spine_ik", 1)
+		    cmds.setAttr(self.character  + ":Rig_Settings.spine_fk", 0)
+		    cmds.setKeyframe(self.character  + ":Rig_Settings.spine_ik")
+		    cmds.setKeyframe(self.character  + ":Rig_Settings.spine_fk")
+		    
+	    
+	    if lArmMode == 0:
+		#then we were in fk arm. check against the current mode
+		currentMode = cmds.getAttr(self.character + ":Rig_Settings.lArmMode")
+		#if we are not in the original mode, switch and match
+		if currentMode == 1:
+		    self.animUIInst.match_singleFrame("arm", "l", "FK", "IK")
+		    cmds.setAttr(self.character + ":Rig_Settings.lArmMode", 0)
+		    cmds.setKeyframe(self.character + ":Rig_Settings.lArmMode")
+		    
+	    if lArmMode == 1:
+		#then we were in ik arm. check against the current mode
+		currentMode = cmds.getAttr(self.character + ":Rig_Settings.lArmMode")
+		#if we are not in the original mode, switch and match
+		if currentMode == 0:
+		    self.animUIInst.match_singleFrame("arm", "l", "IK", "FK")
+		    cmds.setAttr(self.character + ":Rig_Settings.lArmMode", 1)
+		    cmds.setKeyframe(self.character + ":Rig_Settings.lArmMode")
+		    
+		    
+	    if rArmMode == 0:
+		#then we were in fk arm. check against the current mode
+		currentMode = cmds.getAttr(self.character + ":Rig_Settings.rArmMode")
+		#if we are not in the original mode, switch and match
+		if currentMode == 1:
+		    self.animUIInst.match_singleFrame("arm", "r", "FK", "IK")
+		    cmds.setAttr(self.character + ":Rig_Settings.rArmMode", 0)
+		    cmds.setKeyframe(self.character + ":Rig_Settings.rArmMode")
+		    
+	    if rArmMode == 1:
+		#then we were in ik arm. check against the current mode
+		currentMode = cmds.getAttr(self.character + ":Rig_Settings.rArmMode")
+		#if we are not in the original mode, switch and match
+		if currentMode == 0:
+		    self.animUIInst.match_singleFrame("arm", "r", "IK", "FK")
+		    cmds.setAttr(self.character + ":Rig_Settings.rArmMode", 1)
+		    cmds.setKeyframe(self.character + ":Rig_Settings.rArmMode")
+		    
+	    
+		    
+	    if lLegMode == 0:
+		#then we were in fk arm. check against the current mode
+		currentMode = cmds.getAttr(self.character + ":Rig_Settings.lLegMode")
+		#if we are not in the original mode, switch and match
+		if currentMode == 1:
+		    self.animUIInst.match_singleFrame("leg", "l", "FK", "IK")
+		    cmds.setAttr(self.character + ":Rig_Settings.lLegMode", 0)
+		    cmds.setKeyframe(self.character + ":Rig_Settings.lLegMode")
+		    
+	    if lLegMode == 1:
+		#then we were in ik arm. check against the current mode
+		currentMode = cmds.getAttr(self.character + ":Rig_Settings.lLegMode")
+		#if we are not in the original mode, switch and match
+		if currentMode == 0:
+		    self.animUIInst.match_singleFrame("leg", "l", "IK", "FK")
+		    cmds.setAttr(self.character + ":Rig_Settings.lLegMode", 1)
+		    cmds.setKeyframe(self.character + ":Rig_Settings.lLegMode")
+		    
+	    if rLegMode == 0:
+		#then we were in fk arm. check against the current mode
+		currentMode = cmds.getAttr(self.character + ":Rig_Settings.rLegMode")
+		#if we are not in the original mode, switch and match
+		if currentMode == 1:
+		    self.animUIInst.match_singleFrame("leg", "r", "FK", "IK")
+		    cmds.setAttr(self.character + ":Rig_Settings.rLegMode", 0)
+		    cmds.setKeyframe(self.character + ":Rig_Settings.rLegMode")
+		    
+	    if rLegMode == 1:
+		#then we were in ik arm. check against the current mode
+		currentMode = cmds.getAttr(self.character + ":Rig_Settings.rLegMode")
+		#if we are not in the original mode, switch and match
+		if currentMode == 0:
+		    self.animUIInst.match_singleFrame("leg", "r", "IK", "FK")
+		    cmds.setAttr(self.character + ":Rig_Settings.rLegMode", 1)
+		    cmds.setKeyframe(self.character + ":Rig_Settings.rLegMode")
+		    
+		    
+		    
+	    
+	    #refresh
+	    currentFrame = cmds.currentTime(q = True)
+	    cmds.currentTime(currentFrame - 1)
+	    cmds.currentTime(currentFrame)
+		    
+		    
+
+
+    # # # # # # # # # # # # # # # # # # # # # # # # # # # # # # # # # # # # # # # # # # # # # # # # # # # # # # # # # # # # # # # 
+    # # # # # # # # # # # # # # # # # # # # # # # # # # # # # # # # # # # # # # # # # # # # # # # # # # # # # # # # # # # # # # # 
+    # # # # # # # # # # # # # # # # # # # # # # # # # # # # # # # # # # # # # # # # # # # # # # # # # # # # # # # # # # # # # # # 
+    # # # # # # # # # # # # # # # # # # # # # # # # # # # # # # # # # # # # # # # # # # # # # # # # # # # # # # # # # # # # # # # 
+    # # # # # # # # # # # # # # # # # # # # # # # # # # # # # # # # # # # # # # # # # # # # # # # # # # # # # # # # # # # # # # # 
+    def loadPose_selection(self, path, poseName, *args):
+	
+	selection = cmds.ls(sl = True)
+	
+	if os.path.exists(path + poseName + ".txt"):
+	    f = open(path + poseName + ".txt", 'r')
+	    poseData = cPickle.load(f)
+	    
+	    for data in poseData:
+		control = self.character + ":" + data[0]
+		originalAttrs = data[1]
+		
+		if cmds.objExists(control):
+		    newData = []
+		    for i in range(2, int(len(data))):
+			newData.append(data[i])
+		    
+		    attrs = cmds.listAttr(control, keyable = True, unlocked = True)
+		    for i in range(int(len(attrs))):
+			
+			try:
+			    for attr in originalAttrs:
+				if attrs[i] in originalAttrs:
+				    dataIndex = originalAttrs.index(attrs[i])
+				    if control in selection:
+					cmds.setAttr(control + "." + attrs[i], newData[dataIndex])
+				
+			except :
+			    pass
+    # # # # # # # # # # # # # # # # # # # # # # # # # # # # # # # # # # # # # # # # # # # # # # # # # # # # # # # # # # # # # # # 
+    # # # # # # # # # # # # # # # # # # # # # # # # # # # # # # # # # # # # # # # # # # # # # # # # # # # # # # # # # # # # # # # 
+    # # # # # # # # # # # # # # # # # # # # # # # # # # # # # # # # # # # # # # # # # # # # # # # # # # # # # # # # # # # # # # # 
+    # # # # # # # # # # # # # # # # # # # # # # # # # # # # # # # # # # # # # # # # # # # # # # # # # # # # # # # # # # # # # # # 
+    # # # # # # # # # # # # # # # # # # # # # # # # # # # # # # # # # # # # # # # # # # # # # # # # # # # # # # # # # # # # # # # 
+    def pastePoseOpposite(self, path, poseName, *args):
+        
+
+	mirrorAllTransControls =  []
+        mirrorxTransControls = ["ik_elbow_l_anim", "ik_elbow_r_anim", "clavicle_l_anim", "clavicle_r_anim","ik_foot_anim_l", "ik_foot_anim_r", "ik_wrist_l_anim", "ik_wrist_r_anim"]
+        mirrorRotateZandY = ["ik_foot_anim_l", "ik_foot_anim_r", "ik_wrist_l_anim", "ik_wrist_r_anim"]
+        
+        if os.path.exists(path + poseName + ".txt"):
+            f  = open(path + poseName + ".txt", 'r')
+            poseData = cPickle.load(f)
+            f.close()
+            
+            
+            #sort through pose data, finding control, and values
+	    for data in poseData:
+		control = self.character + ":" + data[0]
+		originalAttrs = data[1]
+		
+		if cmds.objExists(control):
+		    newData = []
+		    for i in range(2, int(len(data))):
+			newData.append(data[i])
+		    
+		    attrs = cmds.listAttr(control, keyable = True, unlocked = True)
+		    for i in range(int(len(attrs))):
+			
+			try:
+			    for attr in originalAttrs:
+				if attrs[i] in originalAttrs:
+				    dataIndex = originalAttrs.index(attrs[i])
+				    
+
+				    if control.find("_l") != -1:
+					if control.rpartition("_l")[2] == "":
+					    ctrl = control.rpartition("_l")[0] + "_r"
+					    
+					    if ctrl in mirrorAllTransControls:
+						for attr in attrs:
+						    if attr.find("translateX") == 0:
+							index =  attrs.index(attr)
+							newData[index] = newData[index] * -1
+						    if attr.find("translateY") == 0:
+							index =  attrs.index(attr)
+							newData[index] = newData[index] * -1
+						    if attr.find("translateZ") == 0:
+							index =  attrs.index(attr)
+							newData[index] = newData[index] * -1
+							
+					    if ctrl in mirrorxTransControls:
+						for attr in attrs:
+						    if attr.find("translateX") == 0:
+							index =  attrs.index(attr)
+							newData[index] = newData[index] * -1
+							
+					    if ctrl in mirrorRotateZandY:
+						for attr in attrs:
+						    if attr.find("rotateY") == 0:
+							index =  attrs.index(attr)
+							newData[index] = newData[index] * -1
+						    if attr.find("rotateZ") == 0:
+							index =  attrs.index(attr)
+							newData[index] = newData[index] * -1
+		    
+							
+					    cmds.setAttr(ctrl + "." + attrs[i], newData[dataIndex])
+							
+					    
+						
+					    
+				    if control.find("_l_") != -1:
+					prefix = control.partition("_l_")[0]
+					suffix = control.partition("_l_")[2]
+					ctrl = prefix + "_r_" + suffix
+		    
+					if ctrl in mirrorAllTransControls:
+					    for attr in attrs:
+						if attr.find("translateX") == 0:
+						    index =  attrs.index(attr)
+						    newData[index] = newData[index] * -1
+						if attr.find("translateY") == 0:
+						    index =  attrs.index(attr)
+						    newData[index] = newData[index] * -1
+						if attr.find("translateZ") == 0:
+						    index =  attrs.index(attr)
+						    newData[index] = newData[index] * -1
+						    
+					if ctrl in mirrorxTransControls:
+					    for attr in attrs:
+						if attr.find("translateX") == 0:
+						    index =  attrs.index(attr)
+						    newData[index] = newData[index] * -1
+						    
+					if ctrl in mirrorRotateZandY:
+					    for attr in attrs:
+						if attr.find("rotateY") == 0:
+						    index =  attrs.index(attr)
+						    newData[index] = newData[index] * -1
+						if attr.find("rotateZ") == 0:
+						    index =  attrs.index(attr)
+						    newData[index] = newData[index] * -1
+						    
+		    
+					cmds.setAttr(ctrl + "." + attrs[i], newData[dataIndex])
+							
+					
+					
+				    if control.find("_r") != -1:
+					if control.rpartition("_r")[2] == "":
+					    ctrl = control.rpartition("_r")[0] + "_l"
+		    
+					    if ctrl in mirrorAllTransControls:
+						for attr in attrs:
+						    if attr.find("translateX") == 0:
+							index =  attrs.index(attr)
+							newData[index] = newData[index] * -1
+						    if attr.find("translateY") == 0:
+							index =  attrs.index(attr)
+							newData[index] = newData[index] * -1
+						    if attr.find("translateZ") == 0:
+							index =  attrs.index(attr)
+							newData[index] = newData[index] * -1
+							
+					    if ctrl in mirrorxTransControls:
+						for attr in attrs:
+						    if attr.find("translateX") == 0:
+							index =  attrs.index(attr)
+							newData[index] = newData[index] * -1
+							
+					    if ctrl in mirrorRotateZandY:
+						for attr in attrs:
+						    if attr.find("rotateY") == 0:
+							index =  attrs.index(attr)
+							newData[index] = newData[index] * -1
+						    if attr.find("rotateZ") == 0:
+							index =  attrs.index(attr)
+							newData[index] = newData[index] * -1
+							
+		    
+					    cmds.setAttr(ctrl + "." + attrs[i], newData[dataIndex])
+							
+					    
+				    if control.find("_r_") != -1:
+					prefix = control.partition("_r_")[0]
+					suffix = control.partition("_r_")[2]
+					ctrl = prefix + "_l_" + suffix
+		    
+					if ctrl in mirrorAllTransControls:
+					    for attr in attrs:
+						if attr.find("translateX") == 0:
+						    index =  attrs.index(attr)
+						    newData[index] = newData[index] * -1
+						if attr.find("translateY") == 0:
+						    index =  attrs.index(attr)
+						    newData[index] = newData[index] * -1
+						if attr.find("translateZ") == 0:
+						    index =  attrs.index(attr)
+						    newData[index] = newData[index] * -1
+						    
+					if ctrl in mirrorxTransControls:
+					    for attr in attrs:
+						if attr.find("translateX") == 0:
+						    index =  attrs.index(attr)
+						    newData[index] = newData[index] * -1
+						    
+					if ctrl in mirrorRotateZandY:
+					    for attr in attrs:
+						if attr.find("rotateY") == 0:
+						    index =  attrs.index(attr)
+						    newData[index] = newData[index] * -1
+						if attr.find("rotateZ") == 0:
+						    index =  attrs.index(attr)
+						    newData[index] = newData[index] * -1
+		    
+		    
+						    
+					cmds.setAttr(ctrl + "." + attrs[i], newData[dataIndex])
+					
+			except:
+			    pass
+			
+    # # # # # # # # # # # # # # # # # # # # # # # # # # # # # # # # # # # # # # # # # # # # # # # # # # # # # # # # # # # # # # # 
+    # # # # # # # # # # # # # # # # # # # # # # # # # # # # # # # # # # # # # # # # # # # # # # # # # # # # # # # # # # # # # # # 
+    # # # # # # # # # # # # # # # # # # # # # # # # # # # # # # # # # # # # # # # # # # # # # # # # # # # # # # # # # # # # # # # 
+    # # # # # # # # # # # # # # # # # # # # # # # # # # # # # # # # # # # # # # # # # # # # # # # # # # # # # # # # # # # # # # # 
+    # # # # # # # # # # # # # # # # # # # # # # # # # # # # # # # # # # # # # # # # # # # # # # # # # # # # # # # # # # # # # # # 
+    def loadPose_mirrored(self, path, poseName, *args):
+	
+	#create a progress window
+	window = cmds.window("loadPose_mirrored_window", title = "Creating Mirror of Pose")
+	cmds.columnLayout()
+	progressControl = cmds.progressBar(maxValue=1000, width=300)
+	cmds.showWindow( window )	
+
+
+	#load the pose first
+	self.loadPose_exact(path, poseName)
+	
+	#check for fk arm orientation settings
+	cmds.setAttr(self.character + ":Rig_Settings.lFkArmOrient", 0)
+	cmds.setAttr(self.character + ":Rig_Settings.rFkArmOrient", 0)
+	
+	group = cmds.group(empty = True, name = "mirror_grp")
+	locators = []
+	leftLocators = []
+	rightLocators = []
+	
+	#constrain controls
+        for control in ["master_anim", "offset_anim", "root_anim", "body_anim", "hip_anim", "chest_ik_anim", "neck_01_fk_anim", "neck_02_fk_anim",
+	                "neck_03_fk_anim", "head_fk_anim", "spine_01_anim", "spine_02_anim", "spine_03_anim", "spine_04_anim", "spine_05_anim"]:
+
+	    if cmds.objExists(self.character + ":" + control):
+		
+		#create a locator for the control
+		locator = cmds.spaceLocator(name = control + "_loc")[0]
+		locators.append(locator)
+		#position the locator in the same space as the control
+		constraint = cmds.parentConstraint(self.character + ":" + control, locator)[0]
+		cmds.delete(constraint)
+		
+		#parent the locator under the group
+		cmds.parent(locator, group)
+		
+	
+        for control in self.controls:
+	    
+	    
+	    if cmds.objExists(self.character + ":" + control):
+		cmds.progressBar(progressControl, edit = True, step = 1)
+		
+		#create a locator for the control
+		name = control
+                if control.find("_l") != -1:
+                    if control.rpartition("_l")[2] == "":
+                        name = control.rpartition("_l")[0] + "_r"
+			
+                if control.find("_l_") != -1:
+                    prefix = control.partition("_l_")[0]
+                    suffix = control.partition("_l_")[2]
+                    name = prefix + "_r_" + suffix
+
+                if control.find("_r") != -1:
+                    if control.rpartition("_r")[2] == "":
+                        name = control.rpartition("_r")[0] + "_l"
+			
+                if control.find("_r_") != -1:
+                    prefix = control.partition("_r_")[0]
+                    suffix = control.partition("_r_")[2]
+                    name = prefix + "_l_" + suffix
+		    
+		    
+		if cmds.objExists(name + "_loc") == False:
+		    locator = cmds.spaceLocator(name = name + "_loc")[0]
+		    locators.append(locator)
+		    
+		    #position the locator in the same space as the control
+		    constraint = cmds.parentConstraint(self.character + ":" + control, locator)[0]
+		    cmds.delete(constraint)
+		    
+		    #set rotation values
+		    cmds.select(locator)
+		    cmds.addAttr(ln = "rxData")
+		    cmds.addAttr(ln = "ryData")
+		    cmds.addAttr(ln = "rzData")
+		    cmds.setAttr(locator + ".rxData",(cmds.getAttr(self.character + ":" + control + ".rx")))
+		    cmds.setAttr(locator + ".ryData",(cmds.getAttr(self.character + ":" + control + ".ry")))
+		    cmds.setAttr(locator + ".rzData",(cmds.getAttr(self.character + ":" + control + ".rz")))
+		    
+		    try:
+			cmds.addAttr(ln = "txData")
+			cmds.addAttr(ln = "tyData")
+			cmds.addAttr(ln = "tzData")
+			cmds.setAttr(locator + ".txData",(cmds.getAttr(self.character + ":" + control + ".tx")))
+			cmds.setAttr(locator + ".tyData",(cmds.getAttr(self.character + ":" + control + ".ty")))
+			cmds.setAttr(locator + ".tzData",(cmds.getAttr(self.character + ":" + control + ".tz")))		    
+    
+		    except:
+			pass
+		    
+		    
+		    #parent the locator under the group
+		    cmds.parent(locator, group)
+	    
+	    
+	    
+	    
+	    
+	#MIRROR THE LOCATORS
+	
+	#mirror the group
+	cmds.setAttr(group + ".sx", -1)
+	
+	#unparent all locators
+	for loc in locators:
+	    cmds.parent(loc, world = True)
+
+	    
+	#SET THE MIRRORED POSE
+
+	for loc in locators:
+	    cmds.progressBar(progressControl, edit = True, step = 1)
+	    
+	    
+	    control = self.character + ":" + loc.rpartition("_loc")[0]
+	    
+	    if loc.find("body_anim") != -1:
+		constraint = cmds.pointConstraint(loc, control)[0]
+		cmds.setKeyframe(control)
+		cmds.delete(constraint)
+		
+		cmds.setAttr(control + ".rx", (cmds.getAttr(control + ".rx") * -1))
+		cmds.setAttr(control + ".ry", (cmds.getAttr(control + ".ry") * -1))
+		
+	    if loc.find("hip_anim") != -1:
+		constraint = cmds.pointConstraint(loc, control)[0]
+		cmds.setKeyframe(control)
+		cmds.delete(constraint)
+		
+		cmds.setAttr(control + ".rx", (cmds.getAttr(control + ".rx") * -1))
+		cmds.setAttr(control + ".ry", (cmds.getAttr(control + ".ry") * -1))
+		
+	    if loc.find("chest_ik") != -1:
+		constraint = cmds.pointConstraint(loc, control)[0]
+		cmds.setKeyframe(control)
+		cmds.delete(constraint)
+		
+		cmds.setAttr(control + ".rx", (cmds.getAttr(control + ".rx") * -1))
+		cmds.setAttr(control + ".ry", (cmds.getAttr(control + ".ry") * -1))
+		
+	    if loc.find("mid_ik") != -1:
+		
+		cmds.setAttr(control + ".tx", (cmds.getAttr(loc + ".txData")))
+		cmds.setAttr(control + ".ty", (cmds.getAttr(loc + ".tyData") * -1))
+		cmds.setAttr(control + ".tz", (cmds.getAttr(loc + ".tzData") * -1))
+		cmds.setKeyframe(control)		
+
+		#constraint = cmds.pointConstraint(loc, control)[0]
+		#cmds.setKeyframe(control)
+		#cmds.delete(constraint)
+		
+		cmds.setAttr(control + ".rx", (cmds.getAttr(control + ".rx") * -1))
+		cmds.setAttr(control + ".ry", (cmds.getAttr(control + ".ry") * -1))
+		
+	    if loc.find("spine") != -1:
+		cmds.setAttr(control + ".rx", (cmds.getAttr(control + ".rx") * -1))
+		cmds.setAttr(control + ".ry", (cmds.getAttr(control + ".ry") * -1))
+		
+	    if loc.find("neck") != -1:
+		cmds.setAttr(control + ".rx", (cmds.getAttr(control + ".rx") * -1))
+		cmds.setAttr(control + ".ry", (cmds.getAttr(control + ".ry") * -1))
+		
+	    if loc.find("head") != -1:
+		cmds.setAttr(control + ".rx", (cmds.getAttr(control + ".rx") * -1))
+		cmds.setAttr(control + ".ry", (cmds.getAttr(control + ".ry") * -1))
+		
+	    if loc.find("clavicle") != -1:
+		constraint = cmds.pointConstraint(loc, control)[0]
+		cmds.setKeyframe(control)
+		cmds.delete(constraint)
+		
+	    if loc.find("ik_wrist") != -1:
+		pConstraint = cmds.pointConstraint(loc, control)[0]
+		cmds.setKeyframe(control)
+		cmds.delete(pConstraint)
+		cmds.setAttr(control + ".rx", (cmds.getAttr(loc + ".rxData")))
+		cmds.setAttr(control + ".ry", (cmds.getAttr(loc + ".ryData") * -1))
+		cmds.setAttr(control + ".rz", (cmds.getAttr(loc + ".rzData") * -1))
+		cmds.setKeyframe(control)
+		
+	    if loc.find("ik_elbow") != -1:
+		constraint = cmds.pointConstraint(loc, control)[0]
+		cmds.setKeyframe(control)
+		cmds.delete(constraint)
+		
+	    if loc.find("ik_foot") != -1:
+		constraint = cmds.pointConstraint(loc, control)[0]
+		cmds.setKeyframe(control)
+		cmds.delete(constraint)
+		cmds.setAttr(control + ".rx", (cmds.getAttr(loc + ".rxData")))
+		cmds.setAttr(control + ".ry", (cmds.getAttr(loc + ".ryData") * -1))
+		cmds.setAttr(control + ".rz", (cmds.getAttr(loc + ".rzData") * -1))
+		cmds.setKeyframe(control)
+		
+	    if loc.find("ik_knee") != -1:
+		cmds.setAttr(control + ".rx", (cmds.getAttr(loc + ".rxData")))
+		cmds.setAttr(control + ".ry", (cmds.getAttr(loc + ".ryData")))
+		cmds.setAttr(control + ".rz", (cmds.getAttr(loc + ".rzData")))
+		cmds.setKeyframe(control)
+		
+	    if loc.find("index_metacarpal") != -1 or loc.find("index_finger") != -1:
+		cmds.setAttr(control + ".rx", (cmds.getAttr(loc + ".rxData")))
+		cmds.setAttr(control + ".ry", (cmds.getAttr(loc + ".ryData")))
+		cmds.setAttr(control + ".rz", (cmds.getAttr(loc + ".rzData")))
+		cmds.setKeyframe(control)
+
+		
+	    if loc.find("middle_metacarpal") != -1 or loc.find("middle_finger") != -1:
+		cmds.setAttr(control + ".rx", (cmds.getAttr(loc + ".rxData")))
+		cmds.setAttr(control + ".ry", (cmds.getAttr(loc + ".ryData")))
+		cmds.setAttr(control + ".rz", (cmds.getAttr(loc + ".rzData")))
+		cmds.setKeyframe(control)
+		
+	    if loc.find("ring_metacarpal") != -1 or loc.find("ring_finger") != -1:
+		cmds.setAttr(control + ".rx", (cmds.getAttr(loc + ".rxData")))
+		cmds.setAttr(control + ".ry", (cmds.getAttr(loc + ".ryData")))
+		cmds.setAttr(control + ".rz", (cmds.getAttr(loc + ".rzData")))
+		cmds.setKeyframe(control)
+		
+    	    if loc.find("pinky_metacarpal") != -1 or loc.find("pinky_finger") != -1:
+		cmds.setAttr(control + ".rx", (cmds.getAttr(loc + ".rxData")))
+		cmds.setAttr(control + ".ry", (cmds.getAttr(loc + ".ryData")))
+		cmds.setAttr(control + ".rz", (cmds.getAttr(loc + ".rzData")))
+		cmds.setKeyframe(control)
+		
+	    if loc.find("thumb_finger") != -1:
+		cmds.setAttr(control + ".rx", (cmds.getAttr(loc + ".rxData")))
+		cmds.setAttr(control + ".ry", (cmds.getAttr(loc + ".ryData")))
+		cmds.setAttr(control + ".rz", (cmds.getAttr(loc + ".rzData")))
+		cmds.setKeyframe(control)
+		
+
+	    if loc.find("heel") != -1:
+		cmds.setAttr(control + ".rx", (cmds.getAttr(loc + ".rxData")))
+		cmds.setAttr(control + ".ry", (cmds.getAttr(loc + ".ryData")))
+		cmds.setAttr(control + ".rz", (cmds.getAttr(loc + ".rzData")))
+		cmds.setKeyframe(control)
+		
+	    if loc.find("toe_wiggle") != -1:
+		cmds.setAttr(control + ".rx", (cmds.getAttr(loc + ".rxData")))
+		cmds.setAttr(control + ".ry", (cmds.getAttr(loc + ".ryData")))
+		cmds.setAttr(control + ".rz", (cmds.getAttr(loc + ".rzData")))
+		cmds.setKeyframe(control)
+		
+	    if loc.find("toe_tip") != -1:
+		cmds.setAttr(control + ".rx", (cmds.getAttr(loc + ".rxData")))
+		cmds.setAttr(control + ".ry", (cmds.getAttr(loc + ".ryData")))
+		cmds.setAttr(control + ".rz", (cmds.getAttr(loc + ".rzData")))
+		cmds.setKeyframe(control)
+		
+	    if loc.find("fk_") == 0:
+		cmds.setAttr(control + ".rx", (cmds.getAttr(loc + ".rxData")))
+		cmds.setAttr(control + ".ry", (cmds.getAttr(loc + ".ryData")))
+		cmds.setAttr(control + ".rz", (cmds.getAttr(loc + ".rzData")))
+		cmds.setKeyframe(control)
+		
+		
+	    #custom rig modules
+	    if (loc.rpartition("_loc")[0]) in self.customModCtrls:
+		
+		if control.find("_cv_") == -1:
+		    
+		    try:
+			constraint = cmds.pointConstraint(loc, control)[0]
+			cmds.setKeyframe(control)
+			cmds.delete(constraint)
+			
+		    except:
+			pass
+		    
+		    try:
+			cmds.setAttr(control + ".rx", (cmds.getAttr(loc + ".rxData")))
+			cmds.setAttr(control + ".ry", (cmds.getAttr(loc + ".ryData")))
+			cmds.setAttr(control + ".rz", (cmds.getAttr(loc + ".rzData")))			
+			cmds.setKeyframe(control)
+		    except:
+			pass
+
+	    
+	#get custom attrs from each side and set on the other
+	leftAttrs = []
+	rightAttrs = []
+	for attr in ["mid_bend", "mid_swivel", "tip_pivot", "tip_swivel"]:
+
+	    
+	    val = cmds.getAttr(self.character + ":ik_wrist_l_anim." + attr)
+	    leftAttrs.append(val)
+	    
+	    val = cmds.getAttr(self.character + ":ik_wrist_r_anim." + attr)
+	    rightAttrs.append(val)
+	    
+	i = 0
+	for attr in ["mid_bend", "mid_swivel", "tip_pivot", "tip_swivel"]:
+	    cmds.setAttr(self.character + ":ik_wrist_l_anim." + attr, rightAttrs[i])
+	    cmds.setAttr(self.character + ":ik_wrist_r_anim." + attr, leftAttrs[i])
+	    i = i + 1
+	    
+
+	for attr in ["knee_twist"]:
+	    leftAttrs = []
+	    rightAttrs = []
+	    
+	    val = cmds.getAttr(self.character + ":ik_foot_anim_l." + attr)
+	    leftAttrs.append(val)
+	    
+	    val = cmds.getAttr(self.character + ":ik_foot_anim_r." + attr)
+	    rightAttrs.append(val)
+	    
+	i = 0
+	for attr in ["knee_twist"]:
+	    cmds.setAttr(self.character + ":ik_foot_anim_l." + attr, rightAttrs[i])
+	    cmds.setAttr(self.character + ":ik_foot_anim_r." + attr, leftAttrs[i])
+	    i = i + 1
+	    
+	   
+	#one last check to ensure things mirrored properly
+	for loc in locators:
+	    cmds.progressBar(progressControl, edit = True, step = 1)
+	    
+	    control = self.character + ":" + loc.rpartition("_loc")[0]
+	    #custom rig modules
+	    if (loc.rpartition("_loc")[0]) in self.customModCtrls:
+		
+		if control.find("_cv_") == -1:
+		    
+		    try:
+			constraint = cmds.pointConstraint(loc, control)[0]
+			cmds.setKeyframe(control)
+			cmds.delete(constraint)
+			
+		    except:
+			pass
+		    
+		    try:
+			cmds.setAttr(control + ".rx", (cmds.getAttr(loc + ".rxData")))
+			cmds.setAttr(control + ".ry", (cmds.getAttr(loc + ".ryData")))
+			cmds.setAttr(control + ".rz", (cmds.getAttr(loc + ".rzData")))			
+			cmds.setKeyframe(control)
+		    except:
+			pass
+		    
+
+		
+		
+	#Delete mirror nodes
+	cmds.progressBar(progressControl, edit = True, pr = 5000)
+	cmds.delete(group)
+	cmds.delete(locators)
+	
+	cmds.deleteUI(window)
+	
+
+
+
+
+    # # # # # # # # # # # # # # # # # # # # # # # # # # # # # # # # # # # # # # # # # # # # # # # # # # # # # # # # # # # # # # # 
+    # # # # # # # # # # # # # # # # # # # # # # # # # # # # # # # # # # # # # # # # # # # # # # # # # # # # # # # # # # # # # # # 
+    # # # # # # # # # # # # # # # # # # # # # # # # # # # # # # # # # # # # # # # # # # # # # # # # # # # # # # # # # # # # # # # 
+    # # # # # # # # # # # # # # # # # # # # # # # # # # # # # # # # # # # # # # # # # # # # # # # # # # # # # # # # # # # # # # # 
+    # # # # # # # # # # # # # # # # # # # # # # # # # # # # # # # # # # # # # # # # # # # # # # # # # # # # # # # # # # # # # # # 
+    def loadPose_ghost(self, path, poseName, *args):
+	
+	#with the character name, find all skin cluster nodes that belong to the character. This will get us our character geo
+	skinClusters = cmds.ls(type = 'skinCluster')
+	
+	characterGeo = []
+	for cluster in skinClusters:
+	    connections = cmds.listConnections(cluster, type = 'mesh')
+	    print connections
+	    if connections != None:
+		characterGeo.append(connections)
+	    
+	#load the exact pose first
+	self.loadPose_exact(path, poseName)
+	
+	#then duplicate off the geo, parent to world, combine meshes, and assign material
+	dupeGeo = []
+	for geo in characterGeo:
+	    print geo
+	    dupe = cmds.duplicate(geo)[0]
+	    dupeGeo.append(dupe)
+	    
+	    #unlock transforms
+	    for attr in [".tx", ".ty", ".tz", ".rx", ".ry", ".rz", ".sx", ".sy", ".sz"]:
+		cmds.setAttr(dupe + attr, lock = False, keyable = True)
+		
+	    parent = cmds.listRelatives(dupe, parent = True)
+	    if parent != None:
+		cmds.parent(dupe, world = True)
+		
+	    #freeze xforms
+	    cmds.makeIdentity(dupe, t = 1, r = 1, s = 1, apply = True)
+		
+	cmds.select(clear = True)
+	for geo in dupeGeo:
+	    cmds.select(geo, add = True)
+	    
+	if len(dupeGeo) > 1:
+	    ghostMesh = cmds.polyUnite((cmds.ls(sl = True)), n = poseName + "_ghost" )[0]
+	    cmds.delete(dupeGeo)
+	if len(dupeGeo) == 1:
+	    ghostMesh = cmds.rename(dupeGeo[0], poseName + "_ghost")
+	
+	cmds.delete(ghostMesh, ch = True)
+	
+	
+	
+	#add material to ghost mesh
+	shader = cmds.shadingNode("lambert", asShader = True, name = ghostMesh + "_material")
+	cmds.setAttr(shader + ".color", 0, 0.327258, 1, type = 'double3')
+	cmds.setAttr(shader + ".incandescence", .266011, .266011, .266011, type = 'double3')
+	cmds.select(ghostMesh)
+	cmds.hyperShade(assign = shader)
+	cmds.select(clear = True)
+
+	
+	#create locator with rotate values matching control rotate values, and pointConstrain to each control
+	locators = []
+        for control in self.controls:
+	    if control.find("space_switcher") == -1:
+		if cmds.objExists(self.character + ":" + control):
+		    loc = cmds.spaceLocator(name = control + "_loc")[0]
+		    cmds.setAttr(loc + ".v", 0)
+		    cmds.select(loc)
+		    cmds.addAttr(ln = "rxData")
+		    cmds.addAttr(ln = "ryData")
+		    cmds.addAttr(ln = "rzData")
+		    
+		    cmds.setAttr(loc + ".rxData", (cmds.getAttr(self.character + ":" + control + ".rx")))
+		    cmds.setAttr(loc + ".ryData", (cmds.getAttr(self.character + ":" + control + ".ry")))
+		    cmds.setAttr(loc + ".rzData", (cmds.getAttr(self.character + ":" + control + ".rz")))
+		    
+		    constraint = cmds.pointConstraint(self.character + ":" + control, loc)[0]
+		    cmds.delete(constraint)
+		    locators.append(loc)
+	
+	#parent locators under the mesh
+	
+	for loc in locators:
+	    cmds.parent(loc, ghostMesh)
+	    
+	#select ghost mesh
+	cmds.select(ghostMesh)
+	
+	#launch snap UI
+	self.ghostPose_UI(ghostMesh, locators, False)
+	
+	
+	    
+    # # # # # # # # # # # # # # # # # # # # # # # # # # # # # # # # # # # # # # # # # # # # # # # # # # # # # # # # # # # # # # # 
+    # # # # # # # # # # # # # # # # # # # # # # # # # # # # # # # # # # # # # # # # # # # # # # # # # # # # # # # # # # # # # # # 
+    # # # # # # # # # # # # # # # # # # # # # # # # # # # # # # # # # # # # # # # # # # # # # # # # # # # # # # # # # # # # # # # 
+    # # # # # # # # # # # # # # # # # # # # # # # # # # # # # # # # # # # # # # # # # # # # # # # # # # # # # # # # # # # # # # # 
+    # # # # # # # # # # # # # # # # # # # # # # # # # # # # # # # # # # # # # # # # # # # # # # # # # # # # # # # # # # # # # # # 
+    def loadPose_ghost_mirrored(self, path, poseName, *args):
+	
+	#with the character name, find all skin cluster nodes that belong to the character. This will get us our character geo
+	skinClusters = cmds.ls(type = 'skinCluster')
+	
+	characterGeo = []
+	for cluster in skinClusters:
+	    connections = cmds.listConnections(cluster, type = 'mesh')
+	    if connections != None:
+		characterGeo.append(connections)
+	    
+	#load the exact pose first
+	self.loadPose_mirrored(path, poseName)
+	
+	
+	#set spaces to default
+	#spaceSwitchers = []
+	#for control in self.controls:
+	    #if control.find("space_switcher") != -1:
+		#spaceSwitchers.append(control)
+		
+	#for node in spaceSwitchers:
+	    #self.setControlSpace(node, None)
+		
+	
+
+	
+	#then duplicate off the geo, parent to world, combine meshes, and assign material
+	dupeGeo = []
+	for geo in characterGeo:
+	    dupe = cmds.duplicate(geo)[0]
+	    dupeGeo.append(dupe)
+	    
+	    #unlock transforms
+	    for attr in [".tx", ".ty", ".tz", ".rx", ".ry", ".rz", ".sx", ".sy", ".sz"]:
+		cmds.setAttr(dupe + attr, lock = False, keyable = True)
+		
+	    #parent to world
+	    parent = cmds.listRelatives(dupe, parent = True)
+	    if parent != None:
+		cmds.parent(dupe, world = True)
+		
+	    
+	cmds.select(clear = True)
+	for geo in dupeGeo:
+	    cmds.select(geo, add = True)
+	    
+	if len(dupeGeo) > 1:
+	    ghostMesh = cmds.polyUnite((cmds.ls(sl = True)), n = poseName + "_ghost" )[0]
+	    cmds.delete(dupeGeo)
+	if len(dupeGeo) == 1:
+	    ghostMesh = cmds.rename(dupeGeo[0], poseName + "_ghost")
+	
+	cmds.delete(ghostMesh, ch = True)
+	
+	
+	#add material to ghost mesh
+	shader = cmds.shadingNode("lambert", asShader = True, name = ghostMesh + "_material")
+	cmds.setAttr(shader + ".color", 0, 0.327258, 1, type = 'double3')
+	cmds.setAttr(shader + ".incandescence", .266011, .266011, .266011, type = 'double3')
+	cmds.select(ghostMesh)
+	cmds.hyperShade(assign = shader)
+	cmds.select(clear = True)
+
+	#create locator with rotate values matching control rotate values, and pointConstrain to each control
+	locators = []
+        for control in self.controls:
+	    if control.find("space_switcher") == -1:
+		if cmds.objExists(self.character + ":" + control):
+		    loc = cmds.spaceLocator(name = control + "_loc")[0]
+		    cmds.setAttr(loc + ".v", 0)
+		    cmds.select(loc)
+		    cmds.addAttr(ln = "rxData")
+		    cmds.addAttr(ln = "ryData")
+		    cmds.addAttr(ln = "rzData")
+		    
+		    cmds.setAttr(loc + ".rxData", (cmds.getAttr(self.character + ":" + control + ".rx")))
+		    cmds.setAttr(loc + ".ryData", (cmds.getAttr(self.character + ":" + control + ".ry")))
+		    cmds.setAttr(loc + ".rzData", (cmds.getAttr(self.character + ":" + control + ".rz")))
+		    
+		    try:
+			cmds.addAttr(ln = "txData")
+			cmds.addAttr(ln = "tyData")
+			cmds.addAttr(ln = "tzData")
+			cmds.setAttr(locator + ".txData",(cmds.getAttr(self.character + ":" + control + ".tx")))
+			cmds.setAttr(locator + ".tyData",(cmds.getAttr(self.character + ":" + control + ".ty")))
+			cmds.setAttr(locator + ".tzData",(cmds.getAttr(self.character + ":" + control + ".tz")))		    
+    
+		    except:
+			pass
+		    
+		    constraint = cmds.pointConstraint(self.character + ":" + control, loc)[0]
+		    cmds.delete(constraint)
+		    locators.append(loc)
+	
+	
+		    #parent the locator under the group
+		    cmds.parent(loc, ghostMesh)	
+	
+
+	#mirror the ghost mesh and freeze transforms
+	cmds.select(ghostMesh)
+	
+	#launch snap UI
+	self.ghostPose_UI(ghostMesh, locators, True)
+	
+	
+
+	
+
+    # # # # # # # # # # # # # # # # # # # # # # # # # # # # # # # # # # # # # # # # # # # # # # # # # # # # # # # # # # # # # # # 
+    # # # # # # # # # # # # # # # # # # # # # # # # # # # # # # # # # # # # # # # # # # # # # # # # # # # # # # # # # # # # # # # 
+    # # # # # # # # # # # # # # # # # # # # # # # # # # # # # # # # # # # # # # # # # # # # # # # # # # # # # # # # # # # # # # # 
+    # # # # # # # # # # # # # # # # # # # # # # # # # # # # # # # # # # # # # # # # # # # # # # # # # # # # # # # # # # # # # # # 
+    # # # # # # # # # # # # # # # # # # # # # # # # # # # # # # # # # # # # # # # # # # # # # # # # # # # # # # # # # # # # # # # 
+    def ghostPose_UI(self, mesh, locators, mirror, *args):
+	
+	
+        #check to see if window exists, if so, delete
+        if cmds.window("ghostPoseUI", exists = True):
+            cmds.deleteUI("ghostPoseUI")
+            
+        #build window
+        self.widgets["ghostPoseUI_window"] = cmds.window("ghostPoseUI", w = 200, h = 100, title = "Ghost Pose", sizeable = False, titleBar = True, titleBarMenu = False)
+            
+            
+        #create the main layout
+        self.widgets["ghostPoseUI_layout"] = cmds.columnLayout(co = ["both", 20], rs = 5)
+        
+        #create save pose button
+	cmds.separator(h = 10)
+	if mirror:
+	    self.widgets["ghostPoseUI_snapButton"] = cmds.symbolButton(w = 140, h = 50, image = self.mayaToolsDir + "/General/Icons/ART/snap.bmp", c= partial(self.ghostPose_snap, mesh, locators))
+	else:
+	    self.widgets["ghostPoseUI_snapButton"] = cmds.symbolButton(w = 140, h = 50, image = self.mayaToolsDir + "/General/Icons/ART/snap.bmp", c= partial(self.ghostPose_snap, mesh, locators))
+	    
+	self.widgets["ghostPoseUI_snapButton"] = cmds.symbolButton(w = 140, h = 50, image = self.mayaToolsDir + "/General/Icons/ART/close.bmp", c = partial(self.close, mesh))
+	cmds.separator(h = 10)
+
+        
+
+            
+        #show the window
+        cmds.showWindow(self.widgets["ghostPoseUI_window"])
+	
+    # # # # # # # # # # # # # # # # # # # # # # # # # # # # # # # # # # # # # # # # # # # # # # # # # # # # # # # # # # # # # # # 
+    # # # # # # # # # # # # # # # # # # # # # # # # # # # # # # # # # # # # # # # # # # # # # # # # # # # # # # # # # # # # # # # 
+    # # # # # # # # # # # # # # # # # # # # # # # # # # # # # # # # # # # # # # # # # # # # # # # # # # # # # # # # # # # # # # # 
+    # # # # # # # # # # # # # # # # # # # # # # # # # # # # # # # # # # # # # # # # # # # # # # # # # # # # # # # # # # # # # # # 
+    # # # # # # # # # # # # # # # # # # # # # # # # # # # # # # # # # # # # # # # # # # # # # # # # # # # # # # # # # # # # # # # 
+    def close(self, mesh, *args):
+	
+	cmds.deleteUI(self.widgets["ghostPoseUI_window"])
+	cmds.delete(mesh)
+	cmds.delete(mesh + "_material")
+    
+ 
+    # # # # # # # # # # # # # # # # # # # # # # # # # # # # # # # # # # # # # # # # # # # # # # # # # # # # # # # # # # # # # # # 
+    # # # # # # # # # # # # # # # # # # # # # # # # # # # # # # # # # # # # # # # # # # # # # # # # # # # # # # # # # # # # # # # 
+    # # # # # # # # # # # # # # # # # # # # # # # # # # # # # # # # # # # # # # # # # # # # # # # # # # # # # # # # # # # # # # # 
+    # # # # # # # # # # # # # # # # # # # # # # # # # # # # # # # # # # # # # # # # # # # # # # # # # # # # # # # # # # # # # # # 
+    # # # # # # # # # # # # # # # # # # # # # # # # # # # # # # # # # # # # # # # # # # # # # # # # # # # # # # # # # # # # # # # 
+    def ghostPose_snap(self, mesh, locators, *args):
+	constraints = []
+	for loc in locators:
+	    suffix = loc.rpartition("_loc")[2]
+	    if suffix == "":
+		control = loc.rpartition("_loc")[0]
+		if control not in self.customModCtrls:
+		    if cmds.objExists(self.character + ":" + control):
+			
+			
+			if cmds.getAttr(self.character + ":" + control + ".tx", keyable = True):
+			    if cmds.getAttr(self.character + ":" + control + ".ty", keyable = True):
+				if cmds.getAttr(self.character + ":" + control + ".tz", keyable = True):
+				    constraint = cmds.pointConstraint(loc, self.character + ":" + control)[0]
+				    cmds.setKeyframe(self.character + ":" + control)
+				    #cmds.delete(constraint)
+    
+			if control == "master_anim":
+			    constraint = cmds.orientConstraint(loc, self.character + ":" + control)[0]
+			    cmds.setKeyframe(self.character + ":" + control)
+			    #cmds.delete(constraint)
+			    
+			try:
+			    if control == "mid_ik_anim":
+				cmds.setAttr(control + ".tx", (cmds.getAttr(loc + ".txData")))
+				cmds.setAttr(control + ".ty", (cmds.getAttr(loc + ".tyData") * -1))
+				cmds.setAttr(control + ".tz", (cmds.getAttr(loc + ".tzData") * -1))
+				cmds.setKeyframe(self.character + ":" + control)	
+				
+			except:
+			    pass
+			
+			    
+
+
+			    
+			if control != "master_anim":
+			    if cmds.getAttr(self.character + ":" + control + ".rx", keyable = True):
+				try:
+				    cmds.setAttr(self.character + ":" + control + ".rx", (cmds.getAttr(loc + ".rxData")))
+				except:
+				    pass
+				
+			    if cmds.getAttr(self.character + ":" + control + ".ry", keyable = True):
+				try:
+				    cmds.setAttr(self.character + ":" + control + ".ry", (cmds.getAttr(loc + ".ryData")))
+				except:
+				    pass
+				
+				
+			    if cmds.getAttr(self.character + ":" + control + ".rz", keyable = True):
+				try:
+				    cmds.setAttr(self.character + ":" + control + ".rz", (cmds.getAttr(loc + ".rzData")))
+				except:
+				    pass
+				
+	
+	
+
+	
+	#run select all on the character
+	import ART_animationUI
+	reload(ART_animationUI)
+	self.animUIInst.selectEverything()
+	
+	selection = cmds.ls(sl = True)
+	for each in selection:
+	    cmds.setKeyframe(each)
+	
+	cmds.select(clear = True)
+		    
+
+	#clean up
+	cmds.deleteUI(self.widgets["ghostPoseUI_window"])
+	cmds.delete(mesh)
+	cmds.delete(mesh + "_material")
+	
+	
+    # # # # # # # # # # # # # # # # # # # # # # # # # # # # # # # # # # # # # # # # # # # # # # # # # # # # # # # # # # # # # # # 
+    # # # # # # # # # # # # # # # # # # # # # # # # # # # # # # # # # # # # # # # # # # # # # # # # # # # # # # # # # # # # # # # 
+    # # # # # # # # # # # # # # # # # # # # # # # # # # # # # # # # # # # # # # # # # # # # # # # # # # # # # # # # # # # # # # # 
+    # # # # # # # # # # # # # # # # # # # # # # # # # # # # # # # # # # # # # # # # # # # # # # # # # # # # # # # # # # # # # # # 
+    # # # # # # # # # # # # # # # # # # # # # # # # # # # # # # # # # # # # # # # # # # # # # # # # # # # # # # # # # # # # # # # 
+    def getCategories(self, parent, *args):
+        
+
+        #clear all menu items from the categories option menu
+	if parent == self.widgets["loadPose_categories"]:
+	    children = cmds.optionMenu(self.widgets["loadPose_categories"], q = True, itemListLong = True)
+	    if children != None:
+		for child in children:
+		    cmds.deleteUI(child)
+	try:
+	    if parent == self.widgets["savePose_categories"]:
+		children = cmds.optionMenu(self.widgets["savePose_categories"], q = True, itemListLong = True)
+		if children != None:
+		    for child in children:
+			cmds.deleteUI(child)
+			
+		#add  Add back in
+		cmds.menuItem(label = "Add New Category", parent = self.widgets["savePose_categories"])
+			    
+	except:
+	    pass
+	
+	
+	if parent == self.widgets["loadPose_categories"]: 
+	    project = cmds.optionMenu(self.widgets["loadPose_projects"], q = True, v = True)
+	
+	else:
+	    if cmds.optionMenu(self.widgets["savePose_projects"], q = True, visible = True):
+		project = cmds.optionMenu(self.widgets["savePose_projects"], q = True, v = True)
+		
+	    else:
+		projects = os.listdir(self.mayaToolsDir + "/General/ART/Projects/")
+		project = projects[0]
+	
+	
+	
+        projectsPath = self.mayaToolsDir + "/General/ART/Projects/" + project + "/Poses/"
+	
+	if not os.path.exists(projectsPath):
+	    os.makedirs(projectsPath)
+	    
+        categories = os.listdir(projectsPath)
+        
+        
+        for category in categories:
+            cmds.menuItem(label = category, parent = parent)
+	    
+	#load poses
+	self.loadPoses()
+	    
+	    
+    # # # # # # # # # # # # # # # # # # # # # # # # # # # # # # # # # # # # # # # # # # # # # # # # # # # # # # # # # # # # # # # 
+    # # # # # # # # # # # # # # # # # # # # # # # # # # # # # # # # # # # # # # # # # # # # # # # # # # # # # # # # # # # # # # # 
+    # # # # # # # # # # # # # # # # # # # # # # # # # # # # # # # # # # # # # # # # # # # # # # # # # # # # # # # # # # # # # # # 
+    # # # # # # # # # # # # # # # # # # # # # # # # # # # # # # # # # # # # # # # # # # # # # # # # # # # # # # # # # # # # # # # 
+    # # # # # # # # # # # # # # # # # # # # # # # # # # # # # # # # # # # # # # # # # # # # # # # # # # # # # # # # # # # # # # # 
+    def getProjects(self, parent, *args):
+        
+        projectsPath = self.mayaToolsDir + "/General/ART/Projects/"
+        projects = os.listdir(projectsPath)
+        
+        for project in projects:
+            cmds.menuItem(label = project, parent = parent)
+	    
+	#set to favorite if it exists
+	settingsLocation = self.mayaToolsDir + "/General/Scripts/projectSettings.txt"
+	if os.path.exists(settingsLocation):
+	    f = open(settingsLocation, 'r')
+	    settings = cPickle.load(f)
+	    favoriteProject = settings.get("FavoriteProject")
+	    
+	    try:
+		cmds.optionMenu(self.widgets["loadPose_projects"], edit = True, v = favoriteProject)
+	    except:
+		pass
+	    
+    # # # # # # # # # # # # # # # # # # # # # # # # # # # # # # # # # # # # # # # # # # # # # # # # # # # # # # # # # # # # # # # 
+    # # # # # # # # # # # # # # # # # # # # # # # # # # # # # # # # # # # # # # # # # # # # # # # # # # # # # # # # # # # # # # # 
+    # # # # # # # # # # # # # # # # # # # # # # # # # # # # # # # # # # # # # # # # # # # # # # # # # # # # # # # # # # # # # # # 
+    # # # # # # # # # # # # # # # # # # # # # # # # # # # # # # # # # # # # # # # # # # # # # # # # # # # # # # # # # # # # # # # 
+    # # # # # # # # # # # # # # # # # # # # # # # # # # # # # # # # # # # # # # # # # # # # # # # # # # # # # # # # # # # # # # # 
+    def relaunch(self, *args):
+	
+	import ART_animationUI
+	reload(ART_animationUI)
+	self.animUIInst.poseEditor()
+        
+        
+    # # # # # # # # # # # # # # # # # # # # # # # # # # # # # # # # # # # # # # # # # # # # # # # # # # # # # # # # # # # # # # # 
+    # # # # # # # # # # # # # # # # # # # # # # # # # # # # # # # # # # # # # # # # # # # # # # # # # # # # # # # # # # # # # # # 
+    # # # # # # # # # # # # # # # # # # # # # # # # # # # # # # # # # # # # # # # # # # # # # # # # # # # # # # # # # # # # # # # 
+    # # # # # # # # # # # # # # # # # # # # # # # # # # # # # # # # # # # # # # # # # # # # # # # # # # # # # # # # # # # # # # # 
+    # # # # # # # # # # # # # # # # # # # # # # # # # # # # # # # # # # # # # # # # # # # # # # # # # # # # # # # # # # # # # # # 
+    def setControlSpace(self, spaceSwitchNode, attr, *args):
+	character = self.character
+
+	if attr == None:
+	    try:
+		#create temp locator to snap the space switch node to
+		control = character + ":" + spaceSwitchNode.rpartition("_space")[0]
+		currentTime = cmds.currentTime(q = True)
+    
+		    
+		cmds.setKeyframe(character + ":" + spaceSwitchNode, t = (currentTime - 1))
+		cmds.currentTime(currentTime - 1)
+		loc = cmds.spaceLocator()[0]
+		constraint = cmds.parentConstraint(control, loc)[0]
+		cmds.delete(constraint)
+		
+		constraint = cmds.parentConstraint(loc, control)[0]
+		cmds.setKeyframe(control, t = (cmds.currentTime(q = True)))
+		cmds.delete(constraint)
+		cmds.delete(loc)
+	    
+	    
+		#create our temp loc
+		tempLoc = cmds.spaceLocator()[0]
+		cmds.currentTime(currentTime)
+		    
+		#constrain temp loc
+		constraint = cmds.parentConstraint(control, tempLoc)[0]
+		cmds.delete(constraint)
+		    
+		#match and switch space
+		attrs = cmds.listAttr(character + ":" + spaceSwitchNode, string = "space_*")
+		for attr in attrs:
+		    cmds.setAttr(character + ":" + spaceSwitchNode+ "." + attr, 0)
+		    cmds.setKetfrane(character + ":" + spaceSwitchNode)
+		    
+    
+		    
+		constraint = cmds.parentConstraint(tempLoc, control)[0]
+		
+		
+		#zero out space node
+		for attr in [".tx", ".ty", ".tz", ".rx", ".ry", ".rz"]:
+		    cmds.setAttr(character + ":" + spaceSwitchNode + attr, 0)
+		
+		
+		cmds.setKeyframe(control, t = currentTime)
+		cmds.delete(constraint)
+		cmds.delete(tempLoc)
+		cmds.select(clear = True)
+		
+	    except:
+		pass
+	    
\ No newline at end of file
