Adding ART scripts to repository

diff --git a/Engine/Extras/Maya_AnimationRiggingTools/MayaTools/General/Scripts/ART_skeletonBuilder_UI.py b/Engine/Extras/Maya_AnimationRiggingTools/MayaTools/General/Scripts/ART_skeletonBuilder_UI.py
new file mode 100644
--- /dev/null
+++ b/Engine/Extras/Maya_AnimationRiggingTools/MayaTools/General/Scripts/ART_skeletonBuilder_UI.py
@@ -0,0 +1,12115 @@
+#GAME ENGINE AUTO RIG
+
+import maya.cmds as cmds
+import maya.mel as mel
+from functools import partial
+import os, cPickle, ast, time, subprocess
+
+
+#get existing instances
+try:
+    import gc
+    for obj in gc.get_objects():
+	if isinstance(obj, SkeletonBuilder_UI):
+	    #set the joint mover modes to be global move mode
+	    if cmds.objExists("Skeleton_Settings"):
+		if cmds.objExists("SceneLocked") == False:
+		    obj.jointMover_ChangeMoveMode("Geo", "off")
+		    obj.jointMover_ChangeMoveMode("Offset", "off")
+		    obj.jointMover_ChangeMoveMode("Global", "off")
+		    obj.jointMover_ChangeMoveMode("Global", "on")
+		    
+except:
+    pass
+	
+	
+class SkeletonBuilder_UI():
+    
+    def __init__(self):
+        
+        
+        #make sure world up is set to Z
+	up = cmds.upAxis(q = True, axis = True)
+	if up == "y":
+	    cmds.upAxis(axis = "z", rv = True)
+	    cmds.confirmDialog(icon = "information", title = "Epic Games", message = "World Up Axis changed to Z up.")
+	    
+	
+	
+        #setup a variable that will be used to know when the user has relaunched the UI, even though there is a joint mover in the scene
+        reinit = False
+        self.symmetryState = False
+	self.saveSelection = []
+	
+        #get access to our maya tools
+        toolsPath = cmds.internalVar(usd = True) + "mayaTools.txt"
+        if os.path.exists(toolsPath):
+            
+            f = open(toolsPath, 'r')
+            self.mayaToolsDir = f.readline()
+            f.close()
+            
+            
+        #create a dictionary for our UI widgets
+        self.widgets = {}
+	
+	#create list for paint weight settings
+	self.paintWeightSettings = []
+        
+        
+        #get the maya window height
+        gMainWindow = mel.eval("$temp1 = $gMainWindow;")
+        winHeight = cmds.window(gMainWindow, q = True, h = True)
+        
+	#If this is a relaunch of the UI, need to repopulate our lists for this instance
+	if cmds.objExists("Skeleton_Settings"):
+	    reinit = True
+	
+	    #recreate lists for global movers, offset movers, geo movers, and lras
+	    self.globalMovers = []
+	    self.offsetMovers = []
+	    self.geoMovers = []
+	    self.lraControls = []
+	    self.proxyGeo = []
+	    self.skelMeshRefGeo = []
+	    
+	    cmds.select("*_mover*")
+	    selection = cmds.ls(sl = True)
+	    
+	    for each in selection:
+		try:
+		    color = cmds.getAttr(each + ".overrideColor")
+		    if color == 17:
+			shape = cmds.listRelatives(each, children = True)[0]
+			self.globalMovers.append(each + "|" + shape)
+			
+		    if color == 18:
+			shape = cmds.listRelatives(each, children = True)[0]
+			self.offsetMovers.append(each + "|" + shape)
+			
+		    if color == 20:
+			baseName = each.partition("_geo_mover")[0]
+			suffix = each.partition("_geo_mover")[2]
+			mover = baseName + "_mover" + suffix
+			if cmds.objExists(mover):
+			    moverShape = cmds.listRelatives(mover, shapes = True)[0]
+			    vis = cmds.getAttr(mover + "|" + moverShape + ".v")
+			    if vis == True:
+				shape = cmds.listRelatives(each, shapes = True)[0]
+				self.geoMovers.append(each + "|" + shape)
+				
+		except:
+		    pass
+		    
+	    cmds.select("*_lra")
+	    selection = cmds.ls(sl = True)
+	    for each in selection:
+		self.lraControls.append(each)
+		
+	    cmds.select("proxy_geo_*")
+	    selection = cmds.ls(sl = True, exactType = "transform", shapes = False)
+	    for each in selection:
+		if each.find("grp") == -1:
+		    parent = cmds.listRelatives(each, parent = True)
+		    if parent != "Proxy_Geo_Skin_Grp":
+			self.proxyGeo.append(each)
+			
+
+	#see if the animation UI is up, and has a channel box built in
+	if cmds.dockControl("artAnimUIDock", exists = True):
+	    
+	    channelBox = cmds.formLayout("ART_cbFormLayout", q = True, childArray = True)
+	    if channelBox != None:
+		channelBox = channelBox[0]
+	    
+		#reparent the channelBox Layout back to maya's window
+		cmds.control(channelBox, e = True, p = "MainChannelsLayersLayout")
+		channelBoxLayout = mel.eval('$temp1=$gChannelsLayersForm')
+		channelBoxForm = mel.eval('$temp1 = $gChannelButtonForm')
+	    
+		#edit the channel box pane's attachment to the formLayout
+		cmds.formLayout(channelBoxLayout, edit = True, af = [(channelBox, "left", 0),(channelBox, "right", 0), (channelBox, "bottom", 0)], attachControl = (channelBox, "top", 0, channelBoxForm))
+	    
+	    cmds.deleteUI("artAnimUIDock")
+
+
+	#check to see if window exists. if so, delete
+	if cmds.dockControl("skeletonBuilder_dock", exists = True):
+	    
+	    channelBox = cmds.formLayout("SkelBuilder_channelBoxFormLayout", q = True, childArray = True)
+	    if channelBox != None:
+		channelBox = channelBox[0]
+	    
+		#reparent the channelBox Layout back to maya's window
+		cmds.control(channelBox, e = True, p = "MainChannelsLayersLayout")
+		channelBoxLayout = mel.eval('$temp1=$gChannelsLayersForm')
+		channelBoxForm = mel.eval('$temp1 = $gChannelButtonForm')
+	    
+		#edit the channel box pane's attachment to the formLayout
+		cmds.formLayout(channelBoxLayout, edit = True, af = [(channelBox, "left", 0),(channelBox, "right", 0), (channelBox, "bottom", 0)], attachControl = (channelBox, "top", 0, channelBoxForm))
+	    
+	    cmds.deleteUI("skeletonBuilder_dock")
+
+
+        #check to see if UI exists, if so, delete
+        if cmds.dockControl("skeletonBuilder_dock", exists = True):
+            cmds.deleteUI("skeletonBuilder_dock")
+	    
+        #create the window
+	if cmds.window("SkelBuilder_window", exists = True):
+	    cmds.deleteUI("SkelBuilder_window")
+	    
+        self.widgets["window"] = cmds.window("SkelBuilder_window", title = "Skeleton Builder", titleBarMenu = False, sizeable = False)
+        
+        #create layout
+        self.widgets["mainLayout"] = cmds.columnLayout(rs = 5)
+	
+	#create the menu bar layout
+	self.widgets["menuBarLayout"] = cmds.menuBarLayout(w = 400, parent = self.widgets["mainLayout"])
+	menu = cmds.menu( label= "Reference Pose Manager", parent = self.widgets["menuBarLayout"] )
+	cmds.menuItem( label= "Assume Model Pose", parent = menu, c = self.assumeModelPose, ann = "Puts the Character in the original pose")
+	cmds.menuItem( label= "Assume Rig Pose", parent = menu, c = self.assumeRigPose, ann = "Puts the character in the user-created rig pose")
+	cmds.menuItem( label= "Reset Model Pose", parent = menu, c = self.resetModelPose, ann = "Change the model pose")
+	cmds.menuItem( label= "Reset Rig Pose", parent = menu, c = partial(self.resetRigPose, False), ann = "Change the rig pose")
+	
+	self.widgets["jointMoverTools"] = cmds.menu( label= "Joint Mover Tools", parent = self.widgets["menuBarLayout"], enable = False )
+	cmds.menuItem( label= "Build/Rebuild Preview Skeleton", parent = self.widgets["jointMoverTools"], ann = "Creates a preview skeleton of what the final skeleton will look like given current joint mover transformations.", c = self.buildPreviewSkeleton)
+	cmds.menuItem( label= "Delete Preview Skeleton", parent = self.widgets["jointMoverTools"], ann = "Removes the preview skeleton, if there is one.", c = self.deletePreviewSkeleton)
+	cmds.menuItem( label= "Bake to Global", parent = self.widgets["jointMoverTools"], ann = "Transfer all values of the joint mover up to the global controls, zeroing out the offset controls in the process(Thus giving you clean offsets).", c = partial(self.transferToGlobal, False))
+	cmds.menuItem( label= "Auto World Space(Rig Pose)", parent = self.widgets["jointMoverTools"], ann = "Assists in getting arms, hands, and feet world axis aligned", c = self.rigPose_worldSpaceHelper)
+
+	
+	self.widgets["proxyMeshToolsMenu"] = cmds.menu( label= "Proxy Mesh Tools", parent = self.widgets["menuBarLayout"], enable = False )
+	cmds.menuItem( label= "Project Proxy Mesh Weights To Meshes", parent = self.widgets["proxyMeshToolsMenu"], c = self.projectWeights)
+	
+        
+        
+	#create area for prev/next step
+	self.widgets["stepNavLayout"] = cmds.rowColumnLayout(parent = self.widgets["mainLayout"], nc = 2, cs = [(1, 65),(2,5)])
+	self.widgets["previousStep"] = cmds.symbolButton(w = 170, h = 40, image = self.mayaToolsDir + "/General/Icons/ART/skelCreation.bmp", visible = False, parent = self.widgets["stepNavLayout"])
+	self.widgets["nextStep"] = cmds.symbolButton(w = 170, h = 40, image = self.mayaToolsDir + "/General/Icons/ART/skelPlacement.bmp", visible = True, parent = self.widgets["stepNavLayout"], c = self.build)
+	
+	
+	
+        #create the layout for our main toolbar
+	self.widgets["skelBuilderTopLevelRows"] = cmds.rowColumnLayout(nc = 3, cw = [(1, 425), (2, 25), (3, 220)], parent = self.widgets["mainLayout"])
+	
+	#create the layouts for the 3 rows
+	self.widgets["tabLayout"] = cmds.tabLayout(parent = self.widgets["skelBuilderTopLevelRows"], tv = False)
+	self.widgets["masterFormLayout"] = cmds.formLayout(w = 425, h = 600, parent = self.widgets["tabLayout"])
+	cmds.formLayout(w = 25, parent = self.widgets["skelBuilderTopLevelRows"])
+	self.widgets["channelBoxFormLayout"] = cmds.formLayout("SkelBuilder_channelBoxFormLayout",  w = 220, h = 600, parent = self.widgets["skelBuilderTopLevelRows"])
+	self.widgets["formLayout"] = cmds.formLayout(w = 425, h = 600, parent = self.widgets["masterFormLayout"])
+	self.widgets["twoColumnLayout"] = cmds.rowColumnLayout(nc = 2, cw = [(1, 60), (2, 365)], parent = self.widgets["formLayout"])
+	
+	self.widgets["toolbarFormLayout"] = cmds.formLayout(w = 60, h = 600, parent = self.widgets["twoColumnLayout"])
+	background = cmds.image(image = self.mayaToolsDir + "/General/Icons/ART/toolbar.jpg", parent = self.widgets["toolbarFormLayout"])
+	cmds.formLayout(self.widgets["toolbarFormLayout"], edit = True, af = [(background, "left", 0), (background, "top", 0)])
+	
+        self.widgets["scrollBarLayout"] = cmds.scrollLayout(w = 365, h = 600, hst = 0, parent = self.widgets["twoColumnLayout"])
+	self.widgets["skelBuilder_columnLayout"] = cmds.columnLayout(w = 340,  parent = self.widgets["scrollBarLayout"], rs = 10)
+
+	
+	#create listView tab
+	self.widgets["listViewTab"] = cmds.formLayout(w = 425, h = 600, parent = self.widgets["tabLayout"])
+
+	#label tabs
+	cmds.tabLayout(self.widgets["tabLayout"], edit = True, tabLabel = [(self.widgets["masterFormLayout"], "Picker"),  (self.widgets["listViewTab"], "Outliner")])
+
+
+        #place the toolbar mode buttons
+        self.createSkeletonToolbar("toolbarFormLayout")
+
+
+        #CREATE FIGURE MODE
+	self.createFigureModeUI()
+	
+        #CREATE OPTIONS MODE
+        self.createOptionsModeUI()
+        
+        #CREATE JOINT MOVER PAGES
+        self.createJointMover_BodyModeUI()
+        self.createJointMover_TorsoModeUI()
+        self.createJointMover_LeftHandModeUI()
+        self.createJointMover_RightHandModeUI()
+        self.createJointMover_LeftFootModeUI()
+        self.createJointMover_RightFootModeUI()
+        self.createJointMover_LockedUI()
+	self.createJointMover_PhysiqueUI()
+        
+        
+        #show the window
+        self.widgets["dock"] = cmds.dockControl("skeletonBuilder_dock", label = "(A)nimation and (R)igging (T)oolset", content = self.widgets["window"], area = "right", allowedArea = "right", visibleChangeCommand = self.interfaceScriptJob)
+        
+
+        #check to see if skeleton_settings exists
+        if cmds.objExists("Skeleton_Settings"):
+            
+            cmds.formLayout(self.widgets["formLayout"], edit = True, vis = False)
+            cmds.formLayout(self.widgets["optionsModeFormLayout"], edit = True, vis = True)
+            cmds.formLayout(self.widgets["jointMover_bodyFormLayout"], edit = True, vis = True)
+            cmds.formLayout(self.widgets["jointMover_torsoModeFormLayout"], edit = True, vis = False)
+            cmds.formLayout(self.widgets["jointMover_leftHandModeFormLayout"], edit = True, vis = False)
+            cmds.formLayout(self.widgets["jointMover_rightHandModeFormLayout"], edit = True, vis = False)
+            cmds.formLayout(self.widgets["jointMover_leftFootModeFormLayout"], edit = True, vis = False)
+            cmds.formLayout(self.widgets["jointMover_rightFootModeFormLayout"], edit = True, vis = False)
+	    cmds.formLayout(self.widgets["jointMover_physiqueFormLayout"], edit = True, vis = False)
+	    cmds.tabLayout(self.widgets["tabLayout"], edit = True, tv = True)
+	    self.createListView()
+	    self.listView_init()
+	    
+	    
+	#read the settings from the skeleton settings cache node to fill in our skeleton settings ui
+	if cmds.objExists("SkeletonSettings_Cache"):
+	    attrs = cmds.listAttr("SkeletonSettings_Cache")
+	    for attr in attrs:
+	    #set values
+		if attr.find("extraJoint") != 0:
+		    try:
+			if cmds.intField(self.widgets[attr], q = True,exists = True):
+			    data = cmds.getAttr("SkeletonSettings_Cache." + attr)
+			    data = ast.literal_eval(data)
+			    cmds.intField(self.widgets[attr], edit = True, value = data[0])
+			    cmds.intField(self.widgets[attr], edit = True, bgc = data[1])
+			    
+		    except KeyError:
+			pass
+	    
+	    for attr in attrs:
+		if attr.find("extraJoint") != 0:
+		    try:
+			if cmds.checkBox(self.widgets[attr], q = True,exists = True):
+			    data = cmds.getAttr("SkeletonSettings_Cache." + attr)
+			    data = ast.literal_eval(data)
+			    cmds.checkBox(self.widgets[attr], edit = True, value = data)
+			    
+		    except KeyError:
+			pass
+
+	    #custom rig modules
+	    extraJoints = []
+	    for attr in attrs:
+		if attr.find("extraJoint") == 0:
+		    extraJoints.append(attr)
+
+	    for joint in extraJoints:
+		info = cmds.getAttr("SkeletonSettings_Cache." + joint, asString = True)
+		type = info.partition("/")[2].partition("/")[0]
+		parent = info.partition("/")[0]
+		jointName = info.rpartition("/")[2]
+
+
+		if type == "leaf":
+		    name = jointName.partition(" (")[0]
+		    data = jointName.partition(" (")[2].partition(")")[0]
+		    
+		    translation = False
+		    rotation = False
+		    
+		    #get checkbox values
+		    if data.find("T") != -1:
+			translation = True
+		    if data.find("R") != -1:
+			rotation = True
+	    
+		    
+		    
+		    #add to the list of added rig modules
+		    cmds.textScrollList(self.widgets["addedRigModulesList"], edit = True, append = "[leaf] " + name, sc = self.rigMod_changeSettingsForms)
+		    #select the newly added item
+		    cmds.textScrollList(self.widgets["addedRigModulesList"], edit = True, si = "[leaf] " + name)
+		    
+		    #create the entry info and parent to the self.widgets["rigModuleSettingsFrame"]
+		    #for leaf joints, we need name, parent, translate/rotate
+		    
+		    #hide existing formLayouts
+		    self.rigMod_hideSettingsForms()
+		    
+		    self.widgets["rigMod_" + name + "_settingsForm"] = cmds.formLayout(w = 150, h = 220, parent = self.widgets["rigModuleSettingsFrame"])
+		    label_name = cmds.text(label = "Name: ")
+		    self.widgets["rigMod_" + name + "_settings_name"] = cmds.textField(w = 100, parent = self.widgets["rigMod_" + name + "_settingsForm"], text = name, cc = partial(self.updateRigModuleList, name, "leaf"))
+		    label_parent = cmds.text(label = "Parent: ")
+		    self.widgets["rigMod_" + name + "_settings_parent"] = cmds.textField(w = 80, parent = self.widgets["rigMod_" + name + "_settingsForm"], text = parent, editable = False)
+		    self.widgets["rigMod_" + name + "_settings_changeParent"] = cmds.button(label = "...", w = 20, parent = self.widgets["rigMod_" + name + "_settingsForm"], ann = "Change Parent", c = partial(self.extraJointInfo_UI_ParentList, name))
+		    self.widgets["rigMod_" + name + "_settings_translate"] = cmds.checkBox(label = "Translate", v = translation, ann = "Joint can Translate")
+		    self.widgets["rigMod_" + name + "_settings_rotate"] = cmds.checkBox(label = "Rotate", v = rotation, ann = "Joint can Rotate")
+		    
+		    cmds.formLayout(self.widgets["rigMod_" + name + "_settingsForm"], edit = True, af = [(label_name, "top", 5),(label_name, "left", 5)])
+		    cmds.formLayout(self.widgets["rigMod_" + name + "_settingsForm"], edit = True, af = [(self.widgets["rigMod_" + name + "_settings_name"], "top", 5),(self.widgets["rigMod_" + name + "_settings_name"], "right", 5)])
+		    cmds.formLayout(self.widgets["rigMod_" + name + "_settingsForm"], edit = True, af = [(label_parent, "top", 35),(label_parent, "left", 5)])
+		    cmds.formLayout(self.widgets["rigMod_" + name + "_settingsForm"], edit = True, af = [(self.widgets["rigMod_" + name + "_settings_parent"], "top", 35),(self.widgets["rigMod_" + name + "_settings_parent"], "right", 25)])	    
+		    cmds.formLayout(self.widgets["rigMod_" + name + "_settingsForm"], edit = True, af = [(self.widgets["rigMod_" + name + "_settings_changeParent"], "top", 35),(self.widgets["rigMod_" + name + "_settings_changeParent"], "right", 5)])	    
+		    cmds.formLayout(self.widgets["rigMod_" + name + "_settingsForm"], edit = True, af = [(self.widgets["rigMod_" + name + "_settings_translate"], "top", 65),(self.widgets["rigMod_" + name + "_settings_translate"], "left", 5)])	    
+		    cmds.formLayout(self.widgets["rigMod_" + name + "_settingsForm"], edit = True, af = [(self.widgets["rigMod_" + name + "_settings_rotate"], "top", 85),(self.widgets["rigMod_" + name + "_settings_rotate"], "left", 5)])	    
+		    
+	    
+	    
+		if type == "jiggle":
+		    name = jointName.partition(" (")[0]
+		    
+		    #add to the list of added rig modules
+		    cmds.textScrollList(self.widgets["addedRigModulesList"], edit = True, append = "[jiggle] " + name, sc = self.rigMod_changeSettingsForms)
+		    
+		    #select the newly added item
+		    cmds.textScrollList(self.widgets["addedRigModulesList"], edit = True, si = "[jiggle] " + name)
+		    
+		    #create the entry info and parent to the self.widgets["rigModuleSettingsFrame"]
+		    
+		    #hide existing formLayouts
+		    self.rigMod_hideSettingsForms()
+		    
+		    self.widgets["rigMod_" + name + "_settingsForm"] = cmds.formLayout(w = 150, h = 220, parent = self.widgets["rigModuleSettingsFrame"])
+		    label_name = cmds.text(label = "Name: ")
+		    self.widgets["rigMod_" + name + "_settings_name"] = cmds.textField(w = 100, parent = self.widgets["rigMod_" + name + "_settingsForm"], text = name)
+		    label_parent = cmds.text(label = "Parent: ")
+		    self.widgets["rigMod_" + name + "_settings_parent"] = cmds.textField(w = 80, parent = self.widgets["rigMod_" + name + "_settingsForm"], text = parent, editable = False)
+		    self.widgets["rigMod_" + name + "_settings_changeParent"] = cmds.button(label = "...", w = 20, parent = self.widgets["rigMod_" + name + "_settingsForm"], ann = "Change Parent", c = partial(self.extraJointInfo_UI_ParentList, name))
+		    
+		    cmds.formLayout(self.widgets["rigMod_" + name + "_settingsForm"], edit = True, af = [(label_name, "top", 5),(label_name, "left", 5)])
+		    cmds.formLayout(self.widgets["rigMod_" + name + "_settingsForm"], edit = True, af = [(self.widgets["rigMod_" + name + "_settings_name"], "top", 5),(self.widgets["rigMod_" + name + "_settings_name"], "right", 5)])
+		    cmds.formLayout(self.widgets["rigMod_" + name + "_settingsForm"], edit = True, af = [(label_parent, "top", 35),(label_parent, "left", 5)])
+		    cmds.formLayout(self.widgets["rigMod_" + name + "_settingsForm"], edit = True, af = [(self.widgets["rigMod_" + name + "_settings_parent"], "top", 35),(self.widgets["rigMod_" + name + "_settings_parent"], "right", 25)])	    
+		    cmds.formLayout(self.widgets["rigMod_" + name + "_settingsForm"], edit = True, af = [(self.widgets["rigMod_" + name + "_settings_changeParent"], "top", 35),(self.widgets["rigMod_" + name + "_settings_changeParent"], "right", 5)])	    
+		    
+		    
+		if type == "chain":
+		    name = jointName.partition(" (")[0]
+		    num = jointName.partition(" (")[2].partition(")")[0]
+		    
+		    
+		    #add to the list of added rig modules
+		    cmds.textScrollList(self.widgets["addedRigModulesList"], edit = True, append = "[chain] " + name, sc = self.rigMod_changeSettingsForms)
+		    #select the newly added item
+		    cmds.textScrollList(self.widgets["addedRigModulesList"], edit = True, si = "[chain] " + name)
+		    
+		    #create the entry info and parent to the self.widgets["rigModuleSettingsFrame"]
+		    #for leaf joints, we need name, parent, translate/rotate
+		    
+		    #hide existing formLayouts
+		    self.rigMod_hideSettingsForms()
+		    
+		    self.widgets["rigMod_" + name + "_settingsForm"] = cmds.formLayout(w = 150, h = 220, parent = self.widgets["rigModuleSettingsFrame"])
+		    label_name = cmds.text(label = "Name: ")
+		    self.widgets["rigMod_" + name + "_settings_name"] = cmds.textField(w = 100, parent = self.widgets["rigMod_" + name + "_settingsForm"], text = name)
+		    label_parent = cmds.text(label = "Parent: ")
+		    self.widgets["rigMod_" + name + "_settings_parent"] = cmds.textField(w = 80, parent = self.widgets["rigMod_" + name + "_settingsForm"], text = parent, editable = False)
+		    self.widgets["rigMod_" + name + "_settings_changeParent"] = cmds.button(label = "...", w = 20, parent = self.widgets["rigMod_" + name + "_settingsForm"], ann = "Change Parent", c = partial(self.extraJointInfo_UI_ParentList, name))
+		    label_number = cmds.text(label = "Count: ")
+		    self.widgets["rigMod_" + name + "_settings_numberInChain"] = cmds.textField(w = 100, parent = self.widgets["rigMod_" + name + "_settingsForm"], text = str(num), ann = "Number of joints in the chain.")
+		    
+		    cmds.formLayout(self.widgets["rigMod_" + name + "_settingsForm"], edit = True, af = [(label_name, "top", 5),(label_name, "left", 5)])
+		    cmds.formLayout(self.widgets["rigMod_" + name + "_settingsForm"], edit = True, af = [(self.widgets["rigMod_" + name + "_settings_name"], "top", 5),(self.widgets["rigMod_" + name + "_settings_name"], "right", 5)])
+		    cmds.formLayout(self.widgets["rigMod_" + name + "_settingsForm"], edit = True, af = [(label_parent, "top", 35),(label_parent, "left", 5)])
+		    cmds.formLayout(self.widgets["rigMod_" + name + "_settingsForm"], edit = True, af = [(self.widgets["rigMod_" + name + "_settings_parent"], "top", 35),(self.widgets["rigMod_" + name + "_settings_parent"], "right", 25)])	    
+		    cmds.formLayout(self.widgets["rigMod_" + name + "_settingsForm"], edit = True, af = [(self.widgets["rigMod_" + name + "_settings_changeParent"], "top", 35),(self.widgets["rigMod_" + name + "_settings_changeParent"], "right", 5)])	    
+		    cmds.formLayout(self.widgets["rigMod_" + name + "_settingsForm"], edit = True, af = [(label_number, "top", 65),(label_number, "left", 5)])
+		    cmds.formLayout(self.widgets["rigMod_" + name + "_settingsForm"], edit = True, af = [(self.widgets["rigMod_" + name + "_settings_numberInChain"], "top", 65),(self.widgets["rigMod_" + name + "_settings_numberInChain"], "right", 5)])
+
+	    
+		    
+	    #set mirror settings
+	    try:
+		mirrorSettings = cmds.getAttr("SkeletonSettings_Cache.mirrorSetting")
+		mirrorLeft = cmds.getAttr("SkeletonSettings_Cache.mirrorLeft")
+		mirrorRight = cmds.getAttr("SkeletonSettings_Cache.mirrorRight")
+		
+		cmds.optionMenu(self.widgets["rigModuleMirrorMode"], edit = True, v = mirrorSettings)
+		cmds.textField(self.widgets["rigModuleMirrorTextL"], edit = True, text = mirrorLeft)
+		cmds.textField(self.widgets["rigModuleMirrorTextR"], edit = True, text = mirrorRight)
+		
+	    except:
+		pass
+	    
+	    
+		
+        #if this is a relaunch of the window, even though the joint mover exists, repopulate the spine and neck bone buttons
+        if reinit == True:
+            #repopulate the spine and neck bone buttons
+            self.jointMoverUI_spineBones()
+            self.jointMoverUI_neckBones()
+            self.jointMoverUI_updateFingerControlVis("numLeftFingers", "left")
+            self.jointMoverUI_updateFingerControlVis("numRightFingers", "right")
+            self.jointMoverUI_updateToeControlVis("numLeftToes", "left")
+            self.jointMoverUI_updateToeControlVis("numRightToes", "right")
+	    
+
+            #change flow chart buttons:
+	    cmds.symbolButton(self.widgets["previousStep"], edit = True, visible = True, image = self.mayaToolsDir + "/General/Icons/ART/skelCreation.bmp", c = self.edit)
+	    cmds.symbolButton(self.widgets["nextStep"], edit = True, visible = True, image = self.mayaToolsDir + "/General/Icons/ART/defSetup.bmp", c = self.lock_phase1)
+
+	    
+	    if cmds.objExists("AimModeNodes"):
+		cmds.iconTextCheckBox(self.widgets["aimMode"], edit = True, v = True)
+		cmds.symbolButton(self.widgets["reset"], edit = True, enable = False)
+		cmds.symbolButton(self.widgets["templatesButton"], edit = True, enable = False)
+		cmds.iconTextCheckBox(self.widgets["symmetry"], edit = True, enable = False)		    
+	    
+	
+
+
+
+	#if this is a locked scene, draw the UI correctly
+	if cmds.objExists("SceneLocked"):
+	    self.switchMode("jointMover_lockedFormLayout", None)
+	    
+
+	    #disable the other toolbar buttons
+	    cmds.iconTextCheckBox(self.widgets["aimMode"], edit = True, enable = False, vis = False)
+	    cmds.iconTextCheckBox(self.widgets["symmetry"], edit = True, enable = False, vis = False)
+	    cmds.symbolButton(self.widgets["reset"], edit = True, enable = False, vis = False)
+	    cmds.symbolButton(self.widgets["templatesButton"], edit = True, enable = False, vis = False)
+	    cmds.symbolButton(self.widgets["romButton"], edit = True, visible = True, vis = True)
+	    cmds.symbolButton(self.widgets["paintWeightsButton"], edit = True, visible = True, vis = True)
+	    cmds.iconTextCheckBox(self.widgets["tabletPressure"], edit = True, visible = True, vis = True)
+	    cmds.symbolButton(self.widgets["mirrorWeightsButton"], edit = True, visible = True, vis = True)
+	    cmds.symbolButton(self.widgets["pruneWeightsButton"], edit = True, visible = True, vis = True)
+	    cmds.symbolButton(self.widgets["addInfluenceButton"], edit = True, visible = True, vis = True)
+	    cmds.symbolButton(self.widgets["removeInfluenceButton"], edit = True, visible = True, vis = True)
+	    cmds.symbolButton(self.widgets["exportSkinWeightsButton"], edit = True, visible = True, vis = True)
+	    cmds.symbolButton(self.widgets["importSkinWeightsButton"], edit = True, visible = True, vis = True)
+	    
+	    cmds.iconTextRadioButton(self.widgets["moverMode"], edit = True, visible = True, vis = False)
+	    cmds.iconTextRadioButton(self.widgets["offsetMode"], edit = True, visible = True, vis = False)
+	    cmds.iconTextRadioButton(self.widgets["geoMoverMode"], edit = True, visible = True, vis = False)
+	    cmds.symbolButton(self.widgets["physiqueMode"], edit = True, visible = True, vis = False)
+	    cmds.symbolButton(self.widgets["meshMode"], edit = True, visible = True, vis = False)
+	    
+	    cmds.tabLayout(self.widgets["tabLayout"], edit = True, tv = False)
+	    
+	    #change out help button command
+	    cmds.symbolButton(self.widgets["helpButton"], edit = True, c = self.skeletonBuilderHelp)
+	
+
+	    #change flow chart buttons:
+	    cmds.symbolButton(self.widgets["previousStep"], edit = True, visible = True, image = self.mayaToolsDir + "/General/Icons/ART/skelPlacementBack.bmp", c = self.unlock)
+	    cmds.symbolButton(self.widgets["nextStep"], edit = True, visible = True, image = self.mayaToolsDir + "/General/Icons/ART/buildRig.bmp", c = self.publish_UI)	    
+
+	    
+
+		    
+		    
+	else:
+	    if reinit:
+		self.jointMover_ChangeMoveMode("Geo", "off")
+		self.jointMover_ChangeMoveMode("Offset", "off")
+		self.jointMover_ChangeMoveMode("Global", "off")
+		self.jointMover_ChangeMoveMode("Global", "on")
+		
+		cmds.menu(self.widgets["jointMoverTools"], edit = True, enable = True)
+		cmds.tabLayout(self.widgets["tabLayout"], edit = True, tv = True)
+		
+		
+		    
+        #clear the selection
+        cmds.select(clear = True)
+	
+	#setup the scriptJob for closing the UI on new scene opened
+	scriptJobNum = cmds.scriptJob(event = ["readingFile", self.killUiScriptJob], runOnce = True)
+	self.listView_ScriptJob()
+
+	#channel box
+	self.showChannelBox()
+        
+
+	#make sure symmetry mode is not active
+	if reinit:
+	    try:
+		self.jointMover_ChangeMoveMode("Geo", "on")
+		self.jointMover_ChangeMoveMode("Offset", "on")
+		self.jointMover_ChangeMoveMode("Global", "on")
+    
+		self.jointMover_ChangeMoveMode("Geo", "off")
+		self.jointMover_ChangeMoveMode("Offset", "off")
+		
+		self.jointMover_ChangeMoveMode("Global", "off")
+		self.jointMover_ChangeMoveMode("Global", "on")
+		
+	    except:
+		pass
+	    
+
+
+
+    def null(self, *args):
+	pass
+    
+    
+    # # # # # # # # # # # # # # # # # # # # # # # # # # # # # # # # # # # # # # # # # # # # # # # # # # # # # # # # # # # # # # # 
+    # # # # # # # # # # # # # # # # # # # # # # # # # # # # # # # # # # # # # # # # # # # # # # # # # # # # # # # # # # # # # # # 
+    # # # # # # # # # # # # # # # # # # # # # # # # # # # # # # # # # # # # # # # # # # # # # # # # # # # # # # # # # # # # # # # 
+    # # # # # # # # # # # # # # # # # # # # # # # # # # # # # # # # # # # # # # # # # # # # # # # # # # # # # # # # # # # # # # # 
+    # # # # # # # # # # # # # # # # # # # # # # # # # # # # # # # # # # # # # # # # # # # # # # # # # # # # # # # # # # # # # # # 
+    def interfaceScriptJob(self, *args):
+    
+	if cmds.dockControl(self.widgets["dock"], q = True, visible = True) == False:
+	    #re-sort out the channel box
+	    channelBox = cmds.formLayout("SkelBuilder_channelBoxFormLayout", q = True, childArray = True)
+	    if channelBox != None:
+		channelBox = channelBox[0]
+	    
+		#reparent the channelBox Layout back to maya's window
+		cmds.control(channelBox, e = True, p = "MainChannelsLayersLayout")
+		channelBoxLayout = mel.eval('$temp1=$gChannelsLayersForm')
+		channelBoxForm = mel.eval('$temp1 = $gChannelButtonForm')
+	    
+		#edit the channel box pane's attachment to the formLayout
+		cmds.formLayout(channelBoxLayout, edit = True, af = [(channelBox, "left", 0),(channelBox, "right", 0), (channelBox, "bottom", 0)], attachControl = (channelBox, "top", 0, channelBoxForm))
+	
+	    
+    # # # # # # # # # # # # # # # # # # # # # # # # # # # # # # # # # # # # # # # # # # # # # # # # # # # # # # # # # # # # # # # 
+    # # # # # # # # # # # # # # # # # # # # # # # # # # # # # # # # # # # # # # # # # # # # # # # # # # # # # # # # # # # # # # # 
+    # # # # # # # # # # # # # # # # # # # # # # # # # # # # # # # # # # # # # # # # # # # # # # # # # # # # # # # # # # # # # # # 
+    # # # # # # # # # # # # # # # # # # # # # # # # # # # # # # # # # # # # # # # # # # # # # # # # # # # # # # # # # # # # # # # 
+    # # # # # # # # # # # # # # # # # # # # # # # # # # # # # # # # # # # # # # # # # # # # # # # # # # # # # # # # # # # # # # # 
+    def killUiScriptJob(self):
+	    
+	if cmds.dockControl("skeletonBuilder_dock", exists = True):
+	    
+	    channelBox = cmds.formLayout("SkelBuilder_channelBoxFormLayout", q = True, childArray = True)
+	    if channelBox != None:
+		channelBox = channelBox[0]
+	    
+		#reparent the channelBox Layout back to maya's window
+		cmds.control(channelBox, e = True, p = "MainChannelsLayersLayout")
+		channelBoxLayout = mel.eval('$temp1=$gChannelsLayersForm')
+		channelBoxForm = mel.eval('$temp1 = $gChannelButtonForm')
+	    
+		#edit the channel box pane's attachment to the formLayout
+		cmds.formLayout(channelBoxLayout, edit = True, af = [(channelBox, "left", 0),(channelBox, "right", 0), (channelBox, "bottom", 0)], attachControl = (channelBox, "top", 0, channelBoxForm))
+		
+		
+
+	    cmds.deleteUI("skeletonBuilder_dock")
+	    if cmds.window("SkelBuilder_window", exists = True):
+		cmds.deleteUI("SkelBuilder_window")
+		
+
+
+    # # # # # # # # # # # # # # # # # # # # # # # # # # # # # # # # # # # # # # # # # # # # # # # # # # # # # # # # # # # # # # # 
+    # # # # # # # # # # # # # # # # # # # # # # # # # # # # # # # # # # # # # # # # # # # # # # # # # # # # # # # # # # # # # # # 
+    # # # # # # # # # # # # # # # # # # # # # # # # # # # # # # # # # # # # # # # # # # # # # # # # # # # # # # # # # # # # # # # 
+    # # # # # # # # # # # # # # # # # # # # # # # # # # # # # # # # # # # # # # # # # # # # # # # # # # # # # # # # # # # # # # # 
+    # # # # # # # # # # # # # # # # # # # # # # # # # # # # # # # # # # # # # # # # # # # # # # # # # # # # # # # # # # # # # # # 
+    def showChannelBox(self, *args):
+
+	self.channelBoxLayout = mel.eval('$temp1=$gChannelsLayersForm')
+	self.channelBoxForm = mel.eval('$temp1 = $gChannelButtonForm')
+	self.channelBox = mel.eval('$temp1=$gChannelsLayersPane;')
+
+
+
+	#parent the channel box to our anim UI
+	cmds.control(self.channelBox, e = True, p = self.widgets["channelBoxFormLayout"])
+	cmds.formLayout(self.widgets["channelBoxFormLayout"], edit = True, af = [(self.channelBox, "left", 0),(self.channelBox, "right", 0), (self.channelBox, "bottom", 0), (self.channelBox, "top", 0)])
+	channelBox = cmds.formLayout(self.widgets["channelBoxFormLayout"], q = True, childArray = True)[0]
+	self.channelBox = channelBox
+
+	
+    # # # # # # # # # # # # # # # # # # # # # # # # # # # # # # # # # # # # # # # # # # # # # # # # # # # # # # # # # # # # # # # 
+    # # # # # # # # # # # # # # # # # # # # # # # # # # # # # # # # # # # # # # # # # # # # # # # # # # # # # # # # # # # # # # # 
+    # # # # # # # # # # # # # # # # # # # # # # # # # # # # # # # # # # # # # # # # # # # # # # # # # # # # # # # # # # # # # # # 
+    # # # # # # # # # # # # # # # # # # # # # # # # # # # # # # # # # # # # # # # # # # # # # # # # # # # # # # # # # # # # # # # 
+    # # # # # # # # # # # # # # # # # # # # # # # # # # # # # # # # # # # # # # # # # # # # # # # # # # # # # # # # # # # # # # # 
+    def createFigureModeUI(self):
+	
+	#create frame layouts for the torso, arms, and leg options
+	#Title
+	self.widgets["figureMode_TitleFrame"] = cmds.frameLayout( cll = False, label='Skeleton Creation Settings', li = 84, borderStyle='in', parent = self.widgets["skelBuilder_columnLayout"], w = 340, bgc = [.3, .3, .3])
+	
+	
+	#TORSO
+	self.widgets["figureMode_TorsoOptions"] = cmds.frameLayout( cll = True, label='Body Options', borderStyle='in', parent = self.widgets["skelBuilder_columnLayout"], w = 340, bgc = [.27, .5, .65])
+	self.widgets["torsoOptions_layout"] = cmds.rowColumnLayout(nc = 2, cw = [(1, 150), (2, 185)],  parent = self.widgets["figureMode_TorsoOptions"])
+	
+	
+	cmds.text(label = "")
+	cmds.text(label = "")
+	
+	cmds.text(label = "Number of Spine Bones:", align = "left")
+	self.widgets["torsoOptions_numSpineBones"] = cmds.intField(w = 185, minValue = 3, maxValue = 5, value = 3, step = 1)
+	cmds.intField(self.widgets["torsoOptions_numSpineBones"], edit = True, changeCommand = partial(self.intFieldChanged, 3, self.widgets["torsoOptions_numSpineBones"]))
+	
+	cmds.text(label = "Number of Neck Bones:", align = "left")
+	self.widgets["torsoOptions_numNeckBones"] = cmds.intField(w = 185, minValue = 1, maxValue = 3, value = 2, step = 1)
+	cmds.intField(self.widgets["torsoOptions_numNeckBones"], edit = True, changeCommand = partial(self.intFieldChanged, 2, self.widgets["torsoOptions_numNeckBones"]))
+	
+	cmds.text(label = "")
+	cmds.text(label = "")
+	
+	
+	#ARMS
+	self.widgets["figureMode_ArmOptions"] = cmds.frameLayout( cll = True, label='Arm Options', borderStyle='in', parent = self.widgets["skelBuilder_columnLayout"], w = 340, marginWidth = 20, bgc = [.27, .5, .65])
+	
+	
+	#LEFT ARM
+	self.widgets["figureMode_LeftArmOptions"] = cmds.frameLayout( cll = True, label='Left Arm', borderStyle='in', parent = self.widgets["figureMode_ArmOptions"], w = 345, marginWidth = 20, cl = True)
+	self.widgets["leftArmOptions_layout"] = cmds.rowColumnLayout(nc = 2, cw = [(1, 150), (2, 185)],  parent = self.widgets["figureMode_LeftArmOptions"])
+	
+	#upper arm twists
+	cmds.text(label = "")
+	cmds.text(label = "")
+	
+	cmds.text(label = "Upper Arm Twists:", align = "left")
+	self.widgets["leftArmOptions_numUpArmTwistBones"] = cmds.intField(w = 185, minValue = 0, maxValue = 3, value = 1, step = 1)
+	cmds.intField(self.widgets["leftArmOptions_numUpArmTwistBones"], edit = True, changeCommand = partial(self.intFieldChanged, 1, self.widgets["leftArmOptions_numUpArmTwistBones"]))
+	
+	#lower arm twists
+	cmds.text(label = "Lower Arm Twists:", align = "left")
+	self.widgets["leftArmOptions_numLowArmTwistBones"] = cmds.intField(w = 185, minValue = 0, maxValue = 3, value = 1, step = 1)
+	cmds.intField(self.widgets["leftArmOptions_numLowArmTwistBones"], edit = True, changeCommand = partial(self.intFieldChanged, 1, self.widgets["leftArmOptions_numLowArmTwistBones"]))
+	
+	cmds.text(label = "")
+	cmds.text(label = "")
+	
+	
+	#Left Hand Options
+	self.widgets["figureMode_LeftHandOptions"] = cmds.frameLayout( cll = True, label='Left Hand', borderStyle='in', parent = self.widgets["figureMode_LeftArmOptions"], w = 345, cl = True, bgc = [.4, .4, .4])
+	self.widgets["leftHandOptions_layout"] = cmds.rowColumnLayout(nc = 2, cw = [(1, 150), (2, 185)],  parent = self.widgets["figureMode_LeftHandOptions"])
+
+	#thumb joints
+	cmds.text(label = "")
+	cmds.text(label = "")
+	
+	cmds.text(label = "Number of Thumb Joints:", align = "left")
+	self.widgets["leftHandOptions_numThumbBones"] = cmds.intField(w = 185, minValue = 0, maxValue = 3, value = 3, step = 1)
+	cmds.intField(self.widgets["leftHandOptions_numThumbBones"], edit = True, changeCommand = partial(self.intFieldChanged, 3, self.widgets["leftHandOptions_numThumbBones"]))
+	
+	#index finger joints
+	cmds.text(label = "Number of Index Joints:", align = "left")
+	self.widgets["leftHandOptions_numIndexBones"] = cmds.intField(w = 185, minValue = 0, maxValue = 4, value = 3, step = 1)
+	cmds.intField(self.widgets["leftHandOptions_numIndexBones"], edit = True, changeCommand = partial(self.intFieldChanged, 3, self.widgets["leftHandOptions_numIndexBones"]))
+	
+	#middle finger joints
+	cmds.text(label = "Number of Middle Joints:", align = "left")
+	self.widgets["leftHandOptions_numMiddleBones"] = cmds.intField(w = 185, minValue = 0, maxValue = 4, value = 3, step = 1)
+	cmds.intField(self.widgets["leftHandOptions_numMiddleBones"], edit = True, changeCommand = partial(self.intFieldChanged, 3, self.widgets["leftHandOptions_numMiddleBones"]))
+	
+	#ring finger joints
+	cmds.text(label = "Number of Ring Joints:", align = "left")
+	self.widgets["leftHandOptions_numRingBones"] = cmds.intField(w = 185, minValue = 0, maxValue = 4, value = 3, step = 1)
+	cmds.intField(self.widgets["leftHandOptions_numRingBones"], edit = True, changeCommand = partial(self.intFieldChanged, 3, self.widgets["leftHandOptions_numRingBones"]))
+	
+	#pinky finger joints
+	cmds.text(label = "Number of Pinky Joints:", align = "left")
+	self.widgets["leftHandOptions_numPinkyBones"] = cmds.intField(w = 185, minValue = 0, maxValue = 4, value = 3, step = 1)
+	cmds.intField(self.widgets["leftHandOptions_numPinkyBones"], edit = True, changeCommand = partial(self.intFieldChanged, 3, self.widgets["leftHandOptions_numPinkyBones"]))
+	
+	cmds.text(label = "")
+	cmds.text(label = "")
+	
+	
+	
+	
+	
+	#RIGHT ARM
+	self.widgets["figureMode_RightArmOptions"] = cmds.frameLayout( cll = True, label='Right Arm', borderStyle='in', parent = self.widgets["figureMode_ArmOptions"], w = 345, marginWidth = 20, cl = True)
+	self.widgets["rightArmOptions_layout"] = cmds.rowColumnLayout(nc = 2, cw = [(1, 150), (2, 185)],  parent = self.widgets["figureMode_RightArmOptions"])
+	
+	#upper arm twists
+	cmds.text(label = "")
+	cmds.text(label = "")
+	
+	self.widgets["rightArmMirrorOptions_CB"] = cmds.checkBox(v = 0, label = "Same as Left Arm", onc = self.copyLeftArmSettings, ofc = partial(self.resetLimbOptions, "arm", "right"))
+	cmds.text(label = "")
+	
+	cmds.text(label = "")
+	cmds.text(label = "")
+	
+	cmds.text(label = "Upper Arm Twists:", align = "left")
+	self.widgets["rightArmOptions_numUpArmTwistBones"] = cmds.intField(w = 185, minValue = 0, maxValue = 3, value = 1, step = 1)
+	cmds.intField(self.widgets["rightArmOptions_numUpArmTwistBones"], edit = True, changeCommand = partial(self.intFieldChanged, 1, self.widgets["rightArmOptions_numUpArmTwistBones"]))
+	
+	#lower arm twists
+	cmds.text(label = "Lower Arm Twists:", align = "left")
+	self.widgets["rightArmOptions_numLowArmTwistBones"] = cmds.intField(w = 185, minValue = 0, maxValue = 3, value = 1, step = 1)
+	cmds.intField(self.widgets["rightArmOptions_numLowArmTwistBones"], edit = True, changeCommand = partial(self.intFieldChanged, 1, self.widgets["rightArmOptions_numLowArmTwistBones"]))
+	
+	cmds.text(label = "")
+	cmds.text(label = "")
+	
+	
+	#Right Hand Options
+	self.widgets["figureMode_RightHandOptions"] = cmds.frameLayout( cll = True, label='Right Hand', borderStyle='in', parent = self.widgets["figureMode_RightArmOptions"], w = 345, cl = True, bgc = [.4, .4, .4])
+	self.widgets["rightHandOptions_layout"] = cmds.rowColumnLayout(nc = 2, cw = [(1, 150), (2, 185)],  parent = self.widgets["figureMode_RightHandOptions"])
+
+	#thumb joints
+	cmds.text(label = "")
+	cmds.text(label = "")
+	
+	cmds.text(label = "Number of Thumb Joints:", align = "left")
+	self.widgets["rightHandOptions_numThumbBones"] = cmds.intField(w = 185, minValue = 0, maxValue = 3, value = 3, step = 1)
+	cmds.intField(self.widgets["rightHandOptions_numThumbBones"], edit = True, changeCommand = partial(self.intFieldChanged, 3, self.widgets["rightHandOptions_numThumbBones"]))
+	
+	#index finger joints
+	cmds.text(label = "Number of Index Joints:", align = "left")
+	self.widgets["rightHandOptions_numIndexBones"] = cmds.intField(w = 185, minValue = 0, maxValue = 4, value = 3, step = 1)
+	cmds.intField(self.widgets["rightHandOptions_numIndexBones"], edit = True, changeCommand = partial(self.intFieldChanged, 3, self.widgets["rightHandOptions_numIndexBones"]))
+	
+	#middle finger joints
+	cmds.text(label = "Number of Middle Joints:", align = "left")
+	self.widgets["rightHandOptions_numMiddleBones"] = cmds.intField(w = 185, minValue = 0, maxValue = 4, value = 3, step = 1)
+	cmds.intField(self.widgets["rightHandOptions_numMiddleBones"], edit = True, changeCommand = partial(self.intFieldChanged, 3, self.widgets["rightHandOptions_numMiddleBones"]))
+	
+	#ring finger joints
+	cmds.text(label = "Number of Ring Joints:", align = "left")
+	self.widgets["rightHandOptions_numRingBones"] = cmds.intField(w = 185, minValue = 0, maxValue = 4, value = 3, step = 1)
+	cmds.intField(self.widgets["rightHandOptions_numRingBones"], edit = True, changeCommand = partial(self.intFieldChanged, 3, self.widgets["rightHandOptions_numRingBones"]))
+	
+	#pinky finger joints
+	cmds.text(label = "Number of Pinky Joints:", align = "left")
+	self.widgets["rightHandOptions_numPinkyBones"] = cmds.intField(w = 185, minValue = 0, maxValue = 4, value = 3, step = 1)
+	cmds.intField(self.widgets["rightHandOptions_numPinkyBones"], edit = True, changeCommand = partial(self.intFieldChanged, 3, self.widgets["rightHandOptions_numPinkyBones"]))
+	
+	cmds.text(label = "")
+	cmds.text(label = "")
+	
+	
+	
+	#LEGS
+	self.widgets["figureMode_LegOptions"] = cmds.frameLayout( cll = True, label='Leg Options', borderStyle='in', parent = self.widgets["skelBuilder_columnLayout"], w = 340, marginWidth = 20, bgc = [.27, .5, .65])
+	
+	#leg style settings
+	self.widgets["figureMode_legStyle"] = cmds.columnLayout(w = 340, parent = self.widgets["figureMode_LegOptions"])
+	self.widgets["legStyleOptionMenu"] = cmds.optionMenu(w = 315, label = "Leg Style: ", parent = self.widgets["figureMode_legStyle"], enable = False)
+	cmds.menuItem(label = "Standard Biped", parent = self.widgets["legStyleOptionMenu"])
+	cmds.menuItem(label = "Hind Leg", parent = self.widgets["legStyleOptionMenu"])
+	cmds.menuItem(label = "Foreleg", parent = self.widgets["legStyleOptionMenu"])
+	
+	#LEFT LEG
+	self.widgets["figureMode_LeftLegOptions"] = cmds.frameLayout( cll = True, label='Left Leg', borderStyle='in', parent = self.widgets["figureMode_LegOptions"], w = 345, marginWidth = 20, cl = True)
+	self.widgets["leftLegOptions_layout"] = cmds.rowColumnLayout(nc = 2, cw = [(1, 150), (2, 185)],  parent = self.widgets["figureMode_LeftLegOptions"])
+	
+	#upper leg twists
+	cmds.text(label = "")
+	cmds.text(label = "")
+	
+	cmds.text(label = "Upper Thigh Twists:", align = "left")
+	self.widgets["leftLegOptions_numThighTwistBones"] = cmds.intField(w = 185, minValue = 0, maxValue = 3, value = 1, step = 1)
+	cmds.intField(self.widgets["leftLegOptions_numThighTwistBones"], edit = True, changeCommand = partial(self.intFieldChanged, 1, self.widgets["leftLegOptions_numThighTwistBones"]))
+	
+	#lower leg twists
+	cmds.text(label = "Upper Calf Twists:", align = "left")
+	self.widgets["leftLegOptions_numCalfTwistBones"] = cmds.intField(w = 185, minValue = 0, maxValue = 3, value = 1, step = 1)
+	cmds.intField(self.widgets["leftLegOptions_numCalfTwistBones"], edit = True, changeCommand = partial(self.intFieldChanged, 1, self.widgets["leftLegOptions_numCalfTwistBones"]))
+	
+	#include ball joint
+	cmds.text(label = "")
+	cmds.text(label = "")
+	self.widgets["leftLegIncludeBallJoint_CB"] = cmds.checkBox(v = 0, label = "Include Ball Joint")
+	cmds.text(label = "")
+	
+	cmds.text(label = "")
+	cmds.text(label = "")
+	
+	#Left Foot Options
+	self.widgets["figureMode_LeftFootOptions"] = cmds.frameLayout( cll = True, label='Left Foot', borderStyle='in', parent = self.widgets["figureMode_LeftLegOptions"], w = 345, cl = True, bgc = [.4, .4, .4])
+	self.widgets["leftFootOptions_layout"] = cmds.rowColumnLayout(nc = 2, cw = [(1, 150), (2, 185)],  parent = self.widgets["figureMode_LeftFootOptions"])
+
+	#thumb joints
+	cmds.text(label = "")
+	cmds.text(label = "")
+	
+	cmds.text(label = "Number of Big Toe Joints:", align = "left")
+	self.widgets["leftFootOptions_numThumbBones"] = cmds.intField(w = 185, minValue = 0, maxValue = 3, value = 0, step = 1)
+	cmds.intField(self.widgets["leftFootOptions_numThumbBones"], edit = True, changeCommand = partial(self.intFieldChanged, 0, self.widgets["leftFootOptions_numThumbBones"]))
+	
+	#index finger joints
+	cmds.text(label = "Number of Index Toe Joints:", align = "left")
+	self.widgets["leftFootOptions_numIndexBones"] = cmds.intField(w = 185, minValue = 0, maxValue = 4, value = 0, step = 1)
+	cmds.intField(self.widgets["leftFootOptions_numIndexBones"], edit = True, changeCommand = partial(self.intFieldChanged, 0, self.widgets["leftFootOptions_numIndexBones"]))
+	
+	#middle finger joints
+	cmds.text(label = "Number of Middle Toe Joints:", align = "left")
+	self.widgets["leftFootOptions_numMiddleBones"] = cmds.intField(w = 185, minValue = 0, maxValue = 4, value = 0, step = 1)
+	cmds.intField(self.widgets["leftFootOptions_numMiddleBones"], edit = True, changeCommand = partial(self.intFieldChanged, 0, self.widgets["leftFootOptions_numMiddleBones"]))
+	
+	#ring finger joints
+	cmds.text(label = "Number of Ring Toe Joints:", align = "left")
+	self.widgets["leftFootOptions_numRingBones"] = cmds.intField(w = 185, minValue = 0, maxValue = 4, value = 0, step = 1)
+	cmds.intField(self.widgets["leftFootOptions_numRingBones"], edit = True, changeCommand = partial(self.intFieldChanged, 0, self.widgets["leftFootOptions_numRingBones"]))
+	
+	#pinky finger joints
+	cmds.text(label = "Number of Pinky Toe Joints:", align = "left")
+	self.widgets["leftFootOptions_numPinkyBones"] = cmds.intField(w = 185, minValue = 0, maxValue = 4, value = 0, step = 1)
+	cmds.intField(self.widgets["leftFootOptions_numPinkyBones"], edit = True, changeCommand = partial(self.intFieldChanged, 0, self.widgets["leftFootOptions_numPinkyBones"]))
+	
+	cmds.text(label = "")
+	cmds.text(label = "")
+	
+	
+	
+	
+	#RIGHT LEG
+	self.widgets["figureMode_RightLegOptions"] = cmds.frameLayout( cll = True, label='Right Leg', borderStyle='in', parent = self.widgets["figureMode_LegOptions"], w = 345, marginWidth = 20, cl = True)
+	self.widgets["rightLegOptions_layout"] = cmds.rowColumnLayout(nc = 2, cw = [(1, 150), (2, 185)],  parent = self.widgets["figureMode_RightLegOptions"])
+	
+	#Mirror left leg settings
+	cmds.text(label = "")
+	cmds.text(label = "")
+	
+	self.widgets["rightLegMirrorOptions_CB"] = cmds.checkBox(v = 0, label = "Same as Left Leg", onc = self.copyLeftLegSettings, ofc = partial(self.resetLimbOptions, "leg", "right"))
+	cmds.text(label = "")
+	
+	cmds.text(label = "")
+	cmds.text(label = "")
+	
+	
+	#upper leg twists
+	cmds.text(label = "")
+	cmds.text(label = "")
+	
+	cmds.text(label = "Upper Thigh Twists:", align = "left")
+	self.widgets["rightLegOptions_numThighTwistBones"] = cmds.intField(w = 185, minValue = 0, maxValue = 3, value = 1, step = 1)
+	cmds.intField(self.widgets["rightLegOptions_numThighTwistBones"], edit = True, changeCommand = partial(self.intFieldChanged, 1, self.widgets["rightLegOptions_numThighTwistBones"]))
+	
+	#lower leg twists
+	cmds.text(label = "Upper Calf Twists:", align = "left")
+	self.widgets["rightLegOptions_numCalfTwistBones"] = cmds.intField(w = 185, minValue = 0, maxValue = 3, value = 1, step = 1)
+	cmds.intField(self.widgets["rightLegOptions_numCalfTwistBones"], edit = True, changeCommand = partial(self.intFieldChanged, 1, self.widgets["rightLegOptions_numCalfTwistBones"]))
+	
+	#include ball joint
+	cmds.text(label = "")
+	cmds.text(label = "")
+	self.widgets["rightLegIncludeBallJoint_CB"] = cmds.checkBox(v = 0, label = "Include Ball Joint")
+	cmds.text(label = "")
+	
+	cmds.text(label = "")
+	cmds.text(label = "")
+	
+	#Right Foot Options
+	self.widgets["figureMode_RightFootOptions"] = cmds.frameLayout( cll = True, label='Right Foot', borderStyle='in', parent = self.widgets["figureMode_RightLegOptions"], w = 345, cl = True, bgc = [.4, .4, .4])
+	self.widgets["rightFootOptions_layout"] = cmds.rowColumnLayout(nc = 2, cw = [(1, 150), (2, 185)],  parent = self.widgets["figureMode_RightFootOptions"])
+
+	#thumb joints
+	cmds.text(label = "")
+	cmds.text(label = "")
+	
+	cmds.text(label = "Number of Big Toe Joints:", align = "left")
+	self.widgets["rightFootOptions_numThumbBones"] = cmds.intField(w = 185, minValue = 0, maxValue = 3, value = 0, step = 1)
+	cmds.intField(self.widgets["rightFootOptions_numThumbBones"], edit = True, changeCommand = partial(self.intFieldChanged, 0, self.widgets["rightFootOptions_numThumbBones"]))
+	
+	#index finger joints
+	cmds.text(label = "Number of Index Toe Joints:", align = "left")
+	self.widgets["rightFootOptions_numIndexBones"] = cmds.intField(w = 185, minValue = 0, maxValue = 4, value = 0, step = 1)
+	cmds.intField(self.widgets["rightFootOptions_numIndexBones"], edit = True, changeCommand = partial(self.intFieldChanged, 0, self.widgets["rightFootOptions_numIndexBones"]))
+	
+	#middle finger joints
+	cmds.text(label = "Number of Middle Toe Joints:", align = "left")
+	self.widgets["rightFootOptions_numMiddleBones"] = cmds.intField(w = 185, minValue = 0, maxValue = 4, value = 0, step = 1)
+	cmds.intField(self.widgets["rightFootOptions_numMiddleBones"], edit = True, changeCommand = partial(self.intFieldChanged, 0, self.widgets["rightFootOptions_numMiddleBones"]))
+	
+	#ring finger joints
+	cmds.text(label = "Number of Ring Toe Joints:", align = "left")
+	self.widgets["rightFootOptions_numRingBones"] = cmds.intField(w = 185, minValue = 0, maxValue = 4, value = 0, step = 1)
+	cmds.intField(self.widgets["rightFootOptions_numRingBones"], edit = True, changeCommand = partial(self.intFieldChanged, 0, self.widgets["rightFootOptions_numRingBones"]))
+	
+	#pinky finger joints
+	cmds.text(label = "Number of Pinky Toe Joints:", align = "left")
+	self.widgets["rightFootOptions_numPinkyBones"] = cmds.intField(w = 185, minValue = 0, maxValue = 4, value = 0, step = 1)
+	cmds.intField(self.widgets["rightFootOptions_numPinkyBones"], edit = True, changeCommand = partial(self.intFieldChanged, 0, self.widgets["rightFootOptions_numPinkyBones"]))
+	
+	cmds.text(label = "")
+	cmds.text(label = "")
+	
+	
+	
+	#CREATE CONTEXT MENUS FOR EACH FRAME LAYOUT
+	menu = cmds.popupMenu(b = 3, parent = self.widgets["torsoOptions_layout"])
+	cmds.menuItem(parent = menu, label = "Reset to Defaults", c = partial(self.resetLimbOptions, "torso", None))
+	
+	menu = cmds.popupMenu(b = 3, parent = self.widgets["leftArmOptions_layout"])
+	cmds.menuItem(parent = menu, label = "Reset to Defaults", c = partial(self.resetLimbOptions, "arm", "left"))
+	
+	menu = cmds.popupMenu(b = 3, parent = self.widgets["rightArmOptions_layout"])
+	cmds.menuItem(parent = menu, label = "Reset to Defaults", c = partial(self.resetLimbOptions, "arm", "right"))
+	
+	menu = cmds.popupMenu(b = 3, parent = self.widgets["leftLegOptions_layout"])
+	cmds.menuItem(parent = menu, label = "Reset to Defaults", c = partial(self.resetLimbOptions, "leg", "left"))
+	
+	menu = cmds.popupMenu(b = 3, parent = self.widgets["rightLegOptions_layout"])
+	cmds.menuItem(parent = menu, label = "Reset to Defaults", c = partial(self.resetLimbOptions, "leg", "right"))
+	
+    # # # # # # # # # # # # # # # # # # # # # # # # # # # # # # # # # # # # # # # # # # # # # # # # # # # # # # # # # # # # # # # 
+    # # # # # # # # # # # # # # # # # # # # # # # # # # # # # # # # # # # # # # # # # # # # # # # # # # # # # # # # # # # # # # # 
+    # # # # # # # # # # # # # # # # # # # # # # # # # # # # # # # # # # # # # # # # # # # # # # # # # # # # # # # # # # # # # # # 
+    # # # # # # # # # # # # # # # # # # # # # # # # # # # # # # # # # # # # # # # # # # # # # # # # # # # # # # # # # # # # # # # 
+    # # # # # # # # # # # # # # # # # # # # # # # # # # # # # # # # # # # # # # # # # # # # # # # # # # # # # # # # # # # # # # # 
+    def createOptionsModeUI(self):
+        #create a new formLayout, parent it to the scrollLayout, and make it invisible
+	self.widgets["optionsMode_FrameLayout"] = cmds.frameLayout(parent = self.widgets["skelBuilder_columnLayout"], cll = True, cl = True, label='Add Rig Modules', borderStyle='in',  w = 345, marginWidth = 10,  bgc = [.27, .5, .65])
+        self.widgets["optionsModeFormLayout"] = cmds.formLayout(w = 345, h = 2000, vis = True, parent = self.widgets["optionsMode_FrameLayout"])
+        
+        
+        #create buttons for types of rig modules
+	
+	#first create the layout for these modules
+	self.widgets["rigModulesFrame"] = cmds.frameLayout(parent = self.widgets["optionsModeFormLayout"], cll = False, bs = "etchedIn", lv = False, w = 320, h = 90, mh = 5)
+	cmds.formLayout(self.widgets["optionsModeFormLayout"], edit = True, af = [(self.widgets["rigModulesFrame"], 'top', 0), (self.widgets["rigModulesFrame"], 'left', 0)])
+	
+	self.widgets["rigModuleScroll"] = cmds.scrollLayout(parent = self.widgets["rigModulesFrame"], w = 320, h = 80, hst = 5, vst = -1)
+	self.widgets["rigModulesLayout"] = cmds.rowColumnLayout(h = 55, w = 400, nc = 6, cs = [(1, 5), (2, 5), (3, 5), (4, 5), (5, 5), (6, 5)], parent = self.widgets["rigModuleScroll"], bgc = [.1, .1, .1])
+	
+	
+	#rig module buttons:
+        self.widgets["leafJointButton"] = cmds.symbolButton(w = 50, h = 50,  image = self.mayaToolsDir + "/General/Icons/ART/skeletonbuilder_leaf.bmp", c = partial(self.extraJointInfo_UI, "leaf"), parent = self.widgets["rigModulesLayout"], ann = "Leaf Joints are useful for creating IK bones,\nattachment bones, etc.")
+        self.widgets["jiggleJointButton"] = cmds.symbolButton(w = 50, h = 50,  image = self.mayaToolsDir + "/General/Icons/ART/skeletonbuilder_jiggle.bmp", c = partial(self.extraJointInfo_UI, "jiggle"), parent = self.widgets["rigModulesLayout"], ann = "Jiggle Joints are great for simulating fat,\nor loose clothing items, etc.")  
+        self.widgets["dynamicJointButton"] = cmds.symbolButton(w = 50, h = 50,  image = self.mayaToolsDir + "/General/Icons/ART/skeletonbuilder_dynamic.bmp", c = self.dynamicJointInfo_UI, parent = self.widgets["rigModulesLayout"], ann = "Dynamic Chain allows you to create a joint chain\nwith the number of joints specified, that, when\nrigged,will be fully dynamic. This is really useful\nfor cloth, hair, actual chains, etc.")
+
+        
+        #global mirror settings
+	label1 = cmds.text(label = "Rig Module Mirror Prefix/Suffix", parent = self.widgets["optionsModeFormLayout"])
+	label2 = cmds.text(label = "Left:", parent = self.widgets["optionsModeFormLayout"])
+	label3 = cmds.text(label = "Right:", parent = self.widgets["optionsModeFormLayout"])
+	
+	cmds.formLayout(self.widgets["optionsModeFormLayout"], edit = True, af = [(label1, 'top', 100), (label1, 'left', 0)])
+	cmds.formLayout(self.widgets["optionsModeFormLayout"], edit = True, af = [(label2, 'top', 100), (label2, 'left', 205)])
+	cmds.formLayout(self.widgets["optionsModeFormLayout"], edit = True, af = [(label3, 'top', 100), (label3, 'left', 260)])
+	
+	
+	self.widgets["rigModuleMirrorMode"] = cmds.optionMenu(w = 150, parent = self.widgets["optionsModeFormLayout"])
+	cmds.menuItem(label = "Mirror Suffix:", parent = self.widgets["rigModuleMirrorMode"])
+        cmds.menuItem(label = "Mirror Prefix:", parent = self.widgets["rigModuleMirrorMode"])
+	
+	self.widgets["rigModuleMirrorTextL"] = cmds.textField(w = 50, text = "_l", parent = self.widgets["optionsModeFormLayout"] )
+        self.widgets["rigModuleMirrorTextR"] = cmds.textField(w = 50, text = "_r", parent = self.widgets["optionsModeFormLayout"] )
+	
+	cmds.formLayout(self.widgets["optionsModeFormLayout"], edit = True, af = [(self.widgets["rigModuleMirrorMode"], 'top', 120), (self.widgets["rigModuleMirrorMode"], 'left', 0)])
+	cmds.formLayout(self.widgets["optionsModeFormLayout"], edit = True, af = [(self.widgets["rigModuleMirrorTextL"], 'top', 120), (self.widgets["rigModuleMirrorTextL"], 'left', 205)])
+	cmds.formLayout(self.widgets["optionsModeFormLayout"], edit = True, af = [(self.widgets["rigModuleMirrorTextR"], 'top', 120), (self.widgets["rigModuleMirrorTextR"], 'left', 260)])
+	
+	
+	
+	#textScrollList of entries
+	sep1 = cmds.separator(w = 400,  h = 10, style = "in", parent = self.widgets["optionsModeFormLayout"])
+	self.widgets["addedRigModulesList"] = cmds.textScrollList(w = 150, h = 200, parent = self.widgets["optionsModeFormLayout"], ams = False, deleteKeyCommand = self.removeSelectedModule, ann = "Press [Delete] to remove an item from the list")
+	label4 = cmds.text(label = "Added Modules:", parent = self.widgets["optionsModeFormLayout"])
+	
+	cmds.formLayout(self.widgets["optionsModeFormLayout"], edit = True, af = [(sep1, 'top', 150), (sep1, 'left', 0)])
+	cmds.formLayout(self.widgets["optionsModeFormLayout"], edit = True, af = [(label4, 'top', 170), (label4, 'left', 0)])
+	cmds.formLayout(self.widgets["optionsModeFormLayout"], edit = True, af = [(self.widgets["addedRigModulesList"], 'top', 190), (self.widgets["addedRigModulesList"], 'left', 0)])
+	
+	
+	#rig module settings
+	self.widgets["rigModuleSettingsFrame"] = cmds.frameLayout(parent = self.widgets["optionsModeFormLayout"], cll = False, bs = "etchedIn", label = "Module Settings", w = 150, h = 220)
+	cmds.formLayout(self.widgets["optionsModeFormLayout"], edit = True, af = [(self.widgets["rigModuleSettingsFrame"], 'top', 170), (self.widgets["rigModuleSettingsFrame"], 'left', 165)])
+	
+	#create the default "screen" to display in the settings
+	self.widgets["rigModuleDefaultSettingsLayout"] = cmds.formLayout(parent = self.widgets["rigModuleSettingsFrame"], w = 150, h = 220)
+	defaultLabel = cmds.text(label = "Please select a module\nfrom the list to view\nor edit settings.", parent = self.widgets["rigModuleDefaultSettingsLayout"])
+	cmds.formLayout(self.widgets["rigModuleDefaultSettingsLayout"], edit = True, af = [(defaultLabel, 'top', 70), (defaultLabel, 'left', 20)])
+	
+	#sort list menu
+	self.widgets["rigModuleSortListMenu"] = cmds.optionMenu(h = 30, w = 100, parent = self.widgets["optionsModeFormLayout"])
+	cmds.menuItem(label = "Alphabetically", parent = self.widgets["rigModuleSortListMenu"])
+	cmds.menuItem(label = "By Type", parent = self.widgets["rigModuleSortListMenu"])
+	self.widgets["rigModuleSortBtn"] = cmds.button(h = 30, w = 40, label = "Sort", parent = self.widgets["optionsModeFormLayout"], c = self.extraJointInfo_UI_Sort)
+	
+	cmds.formLayout(self.widgets["optionsModeFormLayout"], edit = True, af = [(self.widgets["rigModuleSortListMenu"], 'top', 400), (self.widgets["rigModuleSortListMenu"], 'left', 0)])
+	cmds.formLayout(self.widgets["optionsModeFormLayout"], edit = True, af = [(self.widgets["rigModuleSortBtn"], 'top', 400), (self.widgets["rigModuleSortBtn"], 'left', 110)])	
+
+
+
+    # # # # # # # # # # # # # # # # # # # # # # # # # # # # # # # # # # # # # # # # # # # # # # # # # # # # # # # # # # # # # # # 
+    # # # # # # # # # # # # # # # # # # # # # # # # # # # # # # # # # # # # # # # # # # # # # # # # # # # # # # # # # # # # # # # 
+    # # # # # # # # # # # # # # # # # # # # # # # # # # # # # # # # # # # # # # # # # # # # # # # # # # # # # # # # # # # # # # # 
+    # # # # # # # # # # # # # # # # # # # # # # # # # # # # # # # # # # # # # # # # # # # # # # # # # # # # # # # # # # # # # # # 
+    # # # # # # # # # # # # # # # # # # # # # # # # # # # # # # # # # # # # # # # # # # # # # # # # # # # # # # # # # # # # # # # 
+    def copyLeftArmSettings(self, *args):
+	
+	for field in ["ArmOptions_numUpArmTwistBones", "ArmOptions_numLowArmTwistBones", "HandOptions_numThumbBones", "HandOptions_numIndexBones", "HandOptions_numMiddleBones", "HandOptions_numRingBones", "HandOptions_numPinkyBones"]:
+	    
+	    value = cmds.intField(self.widgets["left" + field], q = True, value = True)
+	    bgc = cmds.intField(self.widgets["left" + field], q = True, bgc = True)
+	    
+	    cmds.intField(self.widgets["right" + field], edit = True, value = value, bgc = bgc)
+	    
+	
+    # # # # # # # # # # # # # # # # # # # # # # # # # # # # # # # # # # # # # # # # # # # # # # # # # # # # # # # # # # # # # # # 
+    # # # # # # # # # # # # # # # # # # # # # # # # # # # # # # # # # # # # # # # # # # # # # # # # # # # # # # # # # # # # # # # 
+    # # # # # # # # # # # # # # # # # # # # # # # # # # # # # # # # # # # # # # # # # # # # # # # # # # # # # # # # # # # # # # # 
+    # # # # # # # # # # # # # # # # # # # # # # # # # # # # # # # # # # # # # # # # # # # # # # # # # # # # # # # # # # # # # # # 
+    # # # # # # # # # # # # # # # # # # # # # # # # # # # # # # # # # # # # # # # # # # # # # # # # # # # # # # # # # # # # # # # 
+    def copyLeftLegSettings(self, *args):
+	
+	for field in ["LegOptions_numThighTwistBones", "LegOptions_numCalfTwistBones", "FootOptions_numThumbBones", "FootOptions_numIndexBones", "FootOptions_numMiddleBones", "FootOptions_numRingBones", "FootOptions_numPinkyBones"]:
+	    
+	    value = cmds.intField(self.widgets["left" + field], q = True, value = True)
+	    bgc = cmds.intField(self.widgets["left" + field], q = True, bgc = True)
+	    
+	    cmds.intField(self.widgets["right" + field], edit = True, value = value, bgc = bgc)
+	    
+	value = cmds.checkBox(self.widgets["leftLegIncludeBallJoint_CB"], q = True, value = True)
+	cmds.checkBox(self.widgets["rightLegIncludeBallJoint_CB"], edit = True, value = value)
+	    
+    # # # # # # # # # # # # # # # # # # # # # # # # # # # # # # # # # # # # # # # # # # # # # # # # # # # # # # # # # # # # # # # 
+    # # # # # # # # # # # # # # # # # # # # # # # # # # # # # # # # # # # # # # # # # # # # # # # # # # # # # # # # # # # # # # # 
+    # # # # # # # # # # # # # # # # # # # # # # # # # # # # # # # # # # # # # # # # # # # # # # # # # # # # # # # # # # # # # # # 
+    # # # # # # # # # # # # # # # # # # # # # # # # # # # # # # # # # # # # # # # # # # # # # # # # # # # # # # # # # # # # # # # 
+    # # # # # # # # # # # # # # # # # # # # # # # # # # # # # # # # # # # # # # # # # # # # # # # # # # # # # # # # # # # # # # # 
+    def resetLimbOptions(self, limb, side, *args):
+	
+	if limb == "leg":
+	    for field in ["LegOptions_numThighTwistBones", "LegOptions_numCalfTwistBones"]:
+		cmds.intField(self.widgets[side + field], edit = True, value = 1, bgc = [.15, .15, .15])
+		
+	    for field in ["FootOptions_numThumbBones", "FootOptions_numIndexBones", "FootOptions_numMiddleBones", "FootOptions_numRingBones", "FootOptions_numPinkyBones"]:
+		cmds.intField(self.widgets[side + field], edit = True, value = 3, bgc = [.15, .15, .15])
+		
+	    cmds.checkBox(self.widgets[side + "LegIncludeBallJoint_CB"], edit = True, value = False)
+		
+		
+		
+	if limb == "arm":
+	    for field in ["ArmOptions_numUpArmTwistBones", "ArmOptions_numLowArmTwistBones"]:
+		cmds.intField(self.widgets[side + field], edit = True, value = 1, bgc = [.15, .15, .15])
+		
+	    for field in ["HandOptions_numThumbBones", "HandOptions_numIndexBones", "HandOptions_numMiddleBones", "HandOptions_numRingBones", "HandOptions_numPinkyBones"]:
+		cmds.intField(self.widgets[side + field], edit = True, value = 3, bgc = [.15, .15, .15])
+		
+		
+	if limb == "torso":
+	    cmds.intField(self.widgets["torsoOptions_numSpineBones"], edit = True, value = 3, bgc = [.15, .15, .15])
+	    cmds.intField(self.widgets["torsoOptions_numNeckBones"], edit = True, value = 2, bgc = [.15, .15, .15])
+
+	
+    # # # # # # # # # # # # # # # # # # # # # # # # # # # # # # # # # # # # # # # # # # # # # # # # # # # # # # # # # # # # # # # 
+    # # # # # # # # # # # # # # # # # # # # # # # # # # # # # # # # # # # # # # # # # # # # # # # # # # # # # # # # # # # # # # # 
+    # # # # # # # # # # # # # # # # # # # # # # # # # # # # # # # # # # # # # # # # # # # # # # # # # # # # # # # # # # # # # # # 
+    # # # # # # # # # # # # # # # # # # # # # # # # # # # # # # # # # # # # # # # # # # # # # # # # # # # # # # # # # # # # # # # 
+    # # # # # # # # # # # # # # # # # # # # # # # # # # # # # # # # # # # # # # # # # # # # # # # # # # # # # # # # # # # # # # # 
+    def intFieldChanged(self, defaultValue, intField, *args):
+	
+	#get the current value
+	value = cmds.intField(intField, q = True, value = True)
+	if value != defaultValue:
+	    cmds.intField(intField, e = True, bgc = [1, .68, 0])
+	    
+	if value == defaultValue:
+	    cmds.intField(intField, e = True, bgc = [.15, .15, .15])
+	    
+	#Toe check
+	ltoeIntFields = [self.widgets["leftFootOptions_numThumbBones"], self.widgets["leftFootOptions_numIndexBones"], self.widgets["leftFootOptions_numMiddleBones"], self.widgets["leftFootOptions_numRingBones"], self.widgets["leftFootOptions_numPinkyBones"]]
+	rtoeIntFields = [self.widgets["rightFootOptions_numThumbBones"], self.widgets["rightFootOptions_numIndexBones"], self.widgets["rightFootOptions_numMiddleBones"], self.widgets["rightFootOptions_numRingBones"], self.widgets["rightFootOptions_numPinkyBones"]]
+
+	if intField in ltoeIntFields:
+	    if value > 0:
+		
+		#make sure that the ball joint checkbox is checked
+		cmds.checkBox(self.widgets["leftLegIncludeBallJoint_CB"], edit = True, value = True)
+		
+	if intField in rtoeIntFields:
+	    if value > 0:
+		
+		#make sure that the ball joint checkbox is checked
+		cmds.checkBox(self.widgets["rightLegIncludeBallJoint_CB"], edit = True, value = True)
+	    
+	
+    # # # # # # # # # # # # # # # # # # # # # # # # # # # # # # # # # # # # # # # # # # # # # # # # # # # # # # # # # # # # # # # 
+    # # # # # # # # # # # # # # # # # # # # # # # # # # # # # # # # # # # # # # # # # # # # # # # # # # # # # # # # # # # # # # # 
+    # # # # # # # # # # # # # # # # # # # # # # # # # # # # # # # # # # # # # # # # # # # # # # # # # # # # # # # # # # # # # # # 
+    # # # # # # # # # # # # # # # # # # # # # # # # # # # # # # # # # # # # # # # # # # # # # # # # # # # # # # # # # # # # # # # 
+    # # # # # # # # # # # # # # # # # # # # # # # # # # # # # # # # # # # # # # # # # # # # # # # # # # # # # # # # # # # # # # # 
+    def createJointMover_LockedUI(self):
+        
+        self.widgets["jointMover_lockedFormLayout"] = cmds.formLayout(w = 400, h = 600, vis = False, parent = self.widgets["masterFormLayout"])
+        
+        #get our image
+        backgroundImage = self.mayaToolsDir + "/General/Icons/ART/jointMoverLock.jpg"
+        self.widgets["jointMover_bodyBackgroundImage"] = cmds.image(image = backgroundImage, w = 400, h = 600)
+        
+	
+	#TESTING!
+	self.widgets["paint_toolSettings_frame"] = cmds.frameLayout( cll = True, label='Tool Settings', borderStyle='in', parent = self.widgets["jointMover_lockedFormLayout"])
+	cmds.formLayout(self.widgets["jointMover_lockedFormLayout"], edit = True, af = [(self.widgets["paint_toolSettings_frame"], "top", 5), (self.widgets["paint_toolSettings_frame"], "left", 65)])
+	
+	
+	#create TOOL SETTINGS layout and widgets
+	self.widgets["paint_toolSettings_form"] = cmds.formLayout(w = 350, h = 70, vis = True, parent = self.widgets["paint_toolSettings_frame"])
+	self.widgets["paint_profileText_tsf"] = cmds.text(label = "Profile:", align = "right")
+	self.widgets["paint_spGaussianChBx"] = cmds.symbolCheckBox(w = 35, h = 36, image = "circleGaus.png", v = 1, onc = partial(self.changeBrushStamp, "gaussian"))
+	self.widgets["paint_spPolyBrushChBx"] = cmds.symbolCheckBox(w = 35, h = 36, image = "circlePoly.png", v = 0, onc = partial(self.changeBrushStamp, "poly"))
+	self.widgets["paint_spSolidChBx"] = cmds.symbolCheckBox(w = 35, h = 36, image = "circleSolid.png", v = 0, onc = partial(self.changeBrushStamp, "solid"))
+	self.widgets["paint_spRectBrushChBx"] = cmds.symbolCheckBox(w = 35, h = 36, image = "rect.png", v = 0, onc = partial(self.changeBrushStamp, "square"))
+	self.widgets["paint_skinTools_textLabel"] = cmds.text(label = "Skin Tools:")
+	self.widgets["paint_skinTools_skinWeightCopy"] = cmds.symbolButton(image = "skinWeightCopy.png", w = 20, h = 20, c = self.paintWeightsSkinWeightCopy, ann = "Copy skin weights from selected vertex")
+	self.widgets["paint_skinTools_skinWeightPaste"] = cmds.symbolButton(image = "skinWeightPaste.png", w = 20, h = 20, c = self.paintWeightsSkinWeightPaste, ann = "Paste skin weight values")
+	self.widgets["paint_skinTools_weightHammer"] = cmds.symbolButton(image = "weightHammer.png", w = 20, h = 20, c = self.paintWeightsHammerVerts, ann = "Hammer weights tool")
+	self.widgets["paint_skinTools_moveVertexWeights"] = cmds.symbolButton(image = "moveVertexWeights.png", w = 20, h = 20, c = self.paintWeightsMoveInfs_UI, ann = "Move influences tool")
+	self.widgets["paint_skinTools_showInfVerts"] = cmds.symbolButton(image = self.mayaToolsDir + "/General/Icons/ART/showInfVerts.bmp", w = 25, h = 25, c = self.showInfluencedVerts, ann = "Show influenced verts")
+	self.widgets["paint_skinTools_paintToolMode"] = cmds.optionMenu(label = "Mode", w = 100, cc = self.paintWeightsWhichMode, ann = "Choose from paint or select modes")
+	cmds.menuItem(label = "Paint", parent = self.widgets["paint_skinTools_paintToolMode"])
+	cmds.menuItem(label = "Select", parent = self.widgets["paint_skinTools_paintToolMode"])
+
+
+
+
+	#place the widgets for TOOL SETTINGS
+	cmds.formLayout(self.widgets["paint_toolSettings_form"], edit = True, af = [(self.widgets["paint_profileText_tsf"], "top", 15), (self.widgets["paint_profileText_tsf"], "left", 5)])
+	cmds.formLayout(self.widgets["paint_toolSettings_form"], edit = True, af = [(self.widgets["paint_spGaussianChBx"], "top", 5), (self.widgets["paint_spGaussianChBx"], "left", 70)])
+	cmds.formLayout(self.widgets["paint_toolSettings_form"], edit = True, af = [(self.widgets["paint_spPolyBrushChBx"], "top", 5), (self.widgets["paint_spPolyBrushChBx"], "left", 105)])
+	cmds.formLayout(self.widgets["paint_toolSettings_form"], edit = True, af = [(self.widgets["paint_spSolidChBx"], "top", 5), (self.widgets["paint_spSolidChBx"], "left", 140)])
+	cmds.formLayout(self.widgets["paint_toolSettings_form"], edit = True, af = [(self.widgets["paint_spRectBrushChBx"], "top", 5), (self.widgets["paint_spRectBrushChBx"], "left", 175)])
+	cmds.formLayout(self.widgets["paint_toolSettings_form"], edit = True, af = [(self.widgets["paint_skinTools_textLabel"], "top", 45), (self.widgets["paint_skinTools_textLabel"], "left", 5)])
+	cmds.formLayout(self.widgets["paint_toolSettings_form"], edit = True, af = [(self.widgets["paint_skinTools_skinWeightCopy"], "top", 45), (self.widgets["paint_skinTools_skinWeightCopy"], "left", 70)])
+	cmds.formLayout(self.widgets["paint_toolSettings_form"], edit = True, af = [(self.widgets["paint_skinTools_skinWeightPaste"], "top", 45), (self.widgets["paint_skinTools_skinWeightPaste"], "left", 105)])
+	cmds.formLayout(self.widgets["paint_toolSettings_form"], edit = True, af = [(self.widgets["paint_skinTools_weightHammer"], "top", 45), (self.widgets["paint_skinTools_weightHammer"], "left", 140)])
+	cmds.formLayout(self.widgets["paint_toolSettings_form"], edit = True, af = [(self.widgets["paint_skinTools_moveVertexWeights"], "top", 45), (self.widgets["paint_skinTools_moveVertexWeights"], "left", 175)])
+	cmds.formLayout(self.widgets["paint_toolSettings_form"], edit = True, af = [(self.widgets["paint_skinTools_showInfVerts"], "top", 42), (self.widgets["paint_skinTools_showInfVerts"], "left", 203)])
+	cmds.formLayout(self.widgets["paint_toolSettings_form"], edit = True, af = [(self.widgets["paint_skinTools_paintToolMode"], "top", 45), (self.widgets["paint_skinTools_paintToolMode"], "left", 230)])
+
+    
+	#Influence Section
+	
+	self.widgets["paint_influences_frame"] = cmds.frameLayout( cll = True, label='Influences', borderStyle='in', parent = self.widgets["jointMover_lockedFormLayout"])
+	cmds.formLayout(self.widgets["jointMover_lockedFormLayout"], edit = True, af = [(self.widgets["paint_influences_frame"], "top", 95), (self.widgets["paint_influences_frame"], "left", 65)])
+	self.widgets["paint_influences_form"] = cmds.formLayout(w = 350, h = 400, vis = True, parent = self.widgets["paint_influences_frame"])
+	self.widgets["paint_lSortLayout_form"] = cmds.formLayout(w = 350, h = 400, vis = True, parent = self.widgets["paint_influences_form"])
+	self.widgets["paint_lSortLayout_label"] = cmds.text(label = "Sort:", parent = self.widgets["paint_lSortLayout_form"])
+	self.widgets["paint_lSortLayout_radioCollection"] = cmds.radioCollection(parent = self.widgets["paint_lSortLayout_form"])
+	self.widgets["paint_lSortLayout_sortMethod"] = cmds.optionMenu(w = 200, parent = self.widgets["paint_lSortLayout_form"], cc = self.paintWeightsSortList)
+	for item in ["By Hierarchy", "Alphabetically", "Flat", "Unlocked Only"]:
+	    cmds.menuItem(label = item, parent = self.widgets["paint_lSortLayout_sortMethod"])
+
+
+	#attach items to lSortLayout form
+	cmds.formLayout(self.widgets["paint_lSortLayout_form"], edit = True, af = [(self.widgets["paint_lSortLayout_label"], "top", 7), (self.widgets["paint_lSortLayout_label"], "left", 5)])
+	cmds.formLayout(self.widgets["paint_lSortLayout_form"], edit = True, af = [(self.widgets["paint_lSortLayout_sortMethod"], "top", 5), (self.widgets["paint_lSortLayout_sortMethod"], "left", 40)])
+
+
+	#create the joints influence list view
+	self.widgets["paint_influences_lListFrame"] = cmds.frameLayout( cll = False, borderVisible = True, label='Influence List', labelVisible = True, borderStyle='in', parent = self.widgets["paint_lSortLayout_form"])
+	self.widgets["paint_filterAndListLayout_form"] = cmds.formLayout("paint_filterAndListLayout_form", w = 300, h = 200, parent = self.widgets["paint_influences_lListFrame"])
+	self.widgets["paint_influencesList_list"] = cmds.treeView("paint_influencesList_list", w = 300, h = 170, enableKeys = True, numberOfButtons = 1, ahp = True, selectionChangedCommand = self.updatedPaintInfluenceSelection)
+	string = "filterUICreateField(\"paint_influencesList_list\",\"paint_filterAndListLayout_form\");"
+	self.widgets["paint_influenceList_filterField"] = mel.eval(string)
+
+
+	#create the lock/unlock buttons and sort buttons
+	self.widgets["paint_artSkinLockInfButton"] = cmds.symbolButton(w = 20, h = 20, image = "Lock_ON.png", parent = self.widgets["paint_filterAndListLayout_form"], c = self.paintWeightsLockWeights)
+	self.widgets["paint_artSkinUnlockInfButton"] = cmds.symbolButton(w = 20, h = 20, image = "Lock_OFF.png", parent = self.widgets["paint_filterAndListLayout_form"], c = self.paintWeightsUnlockWeights)
+	
+	
+	#PAINT OPERATION
+	self.widgets["paint_operation_label"]  = cmds.text(label = "Paint Operation:", parent = self.widgets["paint_lSortLayout_form"],font = "boldLabelFont")
+	self.widgets["paint_artAttrOper_radioButton1"] = cmds.radioButton(label = "Replace",parent = self.widgets["paint_lSortLayout_form"], cc = self.changePaintOperation)
+	self.widgets["paint_artAttrOper_radioButton2"] = cmds.radioButton(label = "Add",parent = self.widgets["paint_lSortLayout_form"], cc = self.changePaintOperation, sl = True)
+	self.widgets["paint_artAttrOper_radioButton3"] = cmds.radioButton(label = "Scale",parent = self.widgets["paint_lSortLayout_form"], cc = self.changePaintOperation)
+	self.widgets["paint_artAttrOper_radioButton4"] = cmds.radioButton(label = "Smooth",parent = self.widgets["paint_lSortLayout_form"], cc = self.changePaintOperation)
+    
+
+	#VALUE SLIDERS AND LABELS
+	self.widgets["paint_paint_attrValueSlider_label"]  = cmds.text(label = "Opacity:", parent = self.widgets["paint_lSortLayout_form"],font = "boldLabelFont")
+	self.widgets["paint_attrValueSlider_label"]  = cmds.text(label = "Value:", parent = self.widgets["paint_lSortLayout_form"],font = "boldLabelFont")
+	self.widgets["paint_opacitySlider"] = cmds.floatSliderGrp("paint_opacitySlider", field = True, precision = 4, cw = [(1, 50), (2, 110)], min = 0.0, max = 1.0, step = 0.5, parent = self.widgets["paint_lSortLayout_form"], v = 1.0, cc = self.changeBrushOpac)
+	self.widgets["paint_attrValueSlider"] = cmds.floatSliderGrp("paint_attrValueSlider", field = True, precision = 4, cw = [(1, 50), (2, 110)], min = 0.0, max = 1.0, step = 0.5, parent = self.widgets["paint_lSortLayout_form"], v = 1.0, cc = self.changeBrushValue)
+	
+	#Value Bookmark Buttons
+	self.widgets["value25Percent"] = cmds.button(w = 20, h = 15, label = ".25", parent = self.widgets["paint_lSortLayout_form"], c = partial(self.paintWeightsBrushPresets, "value", .25))
+	self.widgets["value50Percent"] = cmds.button(w = 20, h = 15, label = ".50", parent = self.widgets["paint_lSortLayout_form"], c = partial(self.paintWeightsBrushPresets, "value", .50))
+	self.widgets["value75Percent"] = cmds.button(w = 20, h = 15, label = ".75", parent = self.widgets["paint_lSortLayout_form"], c = partial(self.paintWeightsBrushPresets, "value", .75))
+	self.widgets["value100Percent"] = cmds.button(w = 20, h = 15, label = "1.0", parent = self.widgets["paint_lSortLayout_form"], c = partial(self.paintWeightsBrushPresets, "value", 1.0))
+	
+	#Opacity Bookmark Buttons
+	self.widgets["opac25Percent"] = cmds.button(w = 20, h = 15, label = ".25", parent = self.widgets["paint_lSortLayout_form"], c = partial(self.paintWeightsBrushPresets, "opacity", .25))
+	self.widgets["opac50Percent"] = cmds.button(w = 20, h = 15, label = ".50", parent = self.widgets["paint_lSortLayout_form"], c = partial(self.paintWeightsBrushPresets, "opacity", .50))
+	self.widgets["opac75Percent"] = cmds.button(w = 20, h = 15, label = ".75", parent = self.widgets["paint_lSortLayout_form"], c = partial(self.paintWeightsBrushPresets, "opacity", .75))
+	self.widgets["opac100Percent"] = cmds.button(w = 20, h = 15, label = "1.0", parent = self.widgets["paint_lSortLayout_form"], c = partial(self.paintWeightsBrushPresets, "opacity", 1.0))
+	
+	#Flood Button
+	self.widgets["paint_floodBtn"] = cmds.button(label = "Flood", w = 32, h = 32, parent = self.widgets["paint_lSortLayout_form"], c = self.paintFloodWeights, ann = "Flood weights")
+
+
+	#attach items to lSortLayout form
+	cmds.formLayout(self.widgets["paint_lSortLayout_form"], edit = True, af = [(self.widgets["paint_influences_lListFrame"], "top", 35), (self.widgets["paint_influences_lListFrame"], "left", 25)])
+	cmds.formLayout(self.widgets["paint_filterAndListLayout_form"], edit = True, af = [(self.widgets["paint_influencesList_list"], "top", 25), (self.widgets["paint_influencesList_list"], "left", 0)])
+	cmds.formLayout(self.widgets["paint_filterAndListLayout_form"], edit = True, af = [(self.widgets["paint_influenceList_filterField"], "top", 0), (self.widgets["paint_influenceList_filterField"], "left", 0)])
+	cmds.formLayout(self.widgets["paint_filterAndListLayout_form"], edit = True, af = [(self.widgets["paint_artSkinLockInfButton"], "top", 0), (self.widgets["paint_artSkinLockInfButton"], "right", 25)])
+	cmds.formLayout(self.widgets["paint_filterAndListLayout_form"], edit = True, af = [(self.widgets["paint_artSkinUnlockInfButton"], "top", 0), (self.widgets["paint_artSkinUnlockInfButton"], "right", 45)])
+	cmds.formLayout(self.widgets["paint_lSortLayout_form"], edit = True, af = [(self.widgets["paint_operation_label"], "top", 277), (self.widgets["paint_operation_label"], "left", 25)])
+	cmds.formLayout(self.widgets["paint_lSortLayout_form"], edit = True, af = [(self.widgets["paint_artAttrOper_radioButton1"], "top", 265), (self.widgets["paint_artAttrOper_radioButton1"], "left", 130)])
+	cmds.formLayout(self.widgets["paint_lSortLayout_form"], edit = True, af = [(self.widgets["paint_artAttrOper_radioButton2"], "top", 265), (self.widgets["paint_artAttrOper_radioButton2"], "left", 230)])
+	cmds.formLayout(self.widgets["paint_lSortLayout_form"], edit = True, af = [(self.widgets["paint_artAttrOper_radioButton3"], "top", 285), (self.widgets["paint_artAttrOper_radioButton3"], "left", 130)])
+	cmds.formLayout(self.widgets["paint_lSortLayout_form"], edit = True, af = [(self.widgets["paint_artAttrOper_radioButton4"], "top", 285), (self.widgets["paint_artAttrOper_radioButton4"], "left", 230)])
+	cmds.formLayout(self.widgets["paint_lSortLayout_form"], edit = True, af = [(self.widgets["paint_opacitySlider"], "top", 320), (self.widgets["paint_opacitySlider"], "left", 130)])
+	cmds.formLayout(self.widgets["paint_lSortLayout_form"], edit = True, af = [(self.widgets["paint_attrValueSlider"], "top", 340), (self.widgets["paint_attrValueSlider"], "left", 130)])
+	cmds.formLayout(self.widgets["paint_lSortLayout_form"], edit = True, af = [(self.widgets["paint_floodBtn"], "top", 325), (self.widgets["paint_floodBtn"], "right", 18)])
+	cmds.formLayout(self.widgets["paint_lSortLayout_form"], edit = True, af = [(self.widgets["paint_paint_attrValueSlider_label"], "top", 322), (self.widgets["paint_paint_attrValueSlider_label"], "left", 25)])
+	cmds.formLayout(self.widgets["paint_lSortLayout_form"], edit = True, af = [(self.widgets["paint_attrValueSlider_label"], "top", 342), (self.widgets["paint_attrValueSlider_label"], "left", 25)])
+	cmds.formLayout(self.widgets["paint_lSortLayout_form"], edit = True, af = [(self.widgets["value25Percent"], "top", 355), (self.widgets["value25Percent"], "left", 185)])
+	cmds.formLayout(self.widgets["paint_lSortLayout_form"], edit = True, af = [(self.widgets["value50Percent"], "top", 355), (self.widgets["value50Percent"], "left", 210)])
+	cmds.formLayout(self.widgets["paint_lSortLayout_form"], edit = True, af = [(self.widgets["value75Percent"], "top", 355), (self.widgets["value75Percent"], "left", 235)])
+	cmds.formLayout(self.widgets["paint_lSortLayout_form"], edit = True, af = [(self.widgets["value100Percent"], "top", 355), (self.widgets["value100Percent"], "left", 260)])
+	cmds.formLayout(self.widgets["paint_lSortLayout_form"], edit = True, af = [(self.widgets["opac25Percent"], "top", 310), (self.widgets["opac25Percent"], "left", 185)])
+	cmds.formLayout(self.widgets["paint_lSortLayout_form"], edit = True, af = [(self.widgets["opac50Percent"], "top", 310), (self.widgets["opac50Percent"], "left", 210)])
+	cmds.formLayout(self.widgets["paint_lSortLayout_form"], edit = True, af = [(self.widgets["opac75Percent"], "top", 310), (self.widgets["opac75Percent"], "left", 235)])
+	cmds.formLayout(self.widgets["paint_lSortLayout_form"], edit = True, af = [(self.widgets["opac100Percent"], "top", 310), (self.widgets["opac100Percent"], "left", 260)])
+	
+	
+	
+	
+	
+	#Selection Section
+	
+	self.widgets["paint_selectionTools_frame"] = cmds.frameLayout( collapse = True, vis = False, cll = True, label='Selection Tools', borderStyle='in', parent = self.widgets["jointMover_lockedFormLayout"])
+	cmds.formLayout(self.widgets["jointMover_lockedFormLayout"], edit = True, af = [(self.widgets["paint_selectionTools_frame"], "top", 125), (self.widgets["paint_selectionTools_frame"], "left", 65)])
+	self.widgets["paint_SelectionTools_form"] = cmds.formLayout(w = 350, h = 400, vis = True, parent = self.widgets["paint_selectionTools_frame"])
+	
+	#create buttons [element, isolate, convert to verts, isolate element from selection]
+	self.widgets["paint_SelectionTools_isoElemFromSel"] = cmds.button(label = "Isolate Element from Selection", w = 300, h = 40, c = self.isolateElementFromSelection, ann = "Select a vertex, and this will grow the selection to grab the entire element, then isolate that selection.")
+	self.widgets["paint_SelectionTools_selElem"] = cmds.button(label = "Select Element", w = 300, h = 40, c = self.selectElement, ann = "Select a vertex and this will grow the selection to grab the entire element.")
+	self.widgets["paint_SelectionTools_isolate"] = cmds.button(label = "Isolate Selection", w = 300, h = 40, c = self.isolateSelection, ann = "Isolate selection.")
+	self.widgets["paint_SelectionTools_convertToVert"] = cmds.button(label = "Convert Selection to Verts", w = 300, h = 40, c = self.convertSelToVerts, ann = "Convert the selection to vertices.")
+	
+	#attach buttons to form
+	cmds.formLayout(self.widgets["paint_SelectionTools_form"], edit = True, af = [(self.widgets["paint_SelectionTools_isoElemFromSel"], "top", 5), (self.widgets["paint_SelectionTools_isoElemFromSel"], "left", 20)])
+	cmds.formLayout(self.widgets["paint_SelectionTools_form"], edit = True, af = [(self.widgets["paint_SelectionTools_selElem"], "top", 50), (self.widgets["paint_SelectionTools_selElem"], "left", 20)])
+	cmds.formLayout(self.widgets["paint_SelectionTools_form"], edit = True, af = [(self.widgets["paint_SelectionTools_isolate"], "top", 95), (self.widgets["paint_SelectionTools_isolate"], "left", 20)])
+	cmds.formLayout(self.widgets["paint_SelectionTools_form"], edit = True, af = [(self.widgets["paint_SelectionTools_convertToVert"], "top", 140), (self.widgets["paint_SelectionTools_convertToVert"], "left", 20)])
+
+
+	    
+	    
+	
+	#DISPLAY
+	self.widgets["paint_display_frame"] = cmds.frameLayout( cll = True, label='Display Settings', borderStyle='in', parent = self.widgets["jointMover_lockedFormLayout"])
+	cmds.formLayout(self.widgets["jointMover_lockedFormLayout"], edit = True, af = [(self.widgets["paint_display_frame"], "top", 495), (self.widgets["paint_display_frame"], "left", 65)])
+	
+	
+	#create TOOL SETTINGS layout and widgets
+	self.widgets["paint_displaySettings_form"] = cmds.rowColumnLayout(nc = 3, w = 350, cat = [(1, "both", 10), (2, "both", 10), (3, "both", 10)], vis = True, parent = self.widgets["paint_display_frame"])
+	self.widgets["paint_displaySettings_drawBrushCB"] = cmds.checkBox(label = "Draw Brush", parent = self.widgets["paint_displaySettings_form"], v = 1, cc = self.paintWeightsDrawBrushToggle)
+	self.widgets["paint_displaySettings_wireframeCB"] = cmds.checkBox(label = "Show Wireframe", parent = self.widgets["paint_displaySettings_form"], v = 1, cc = self.paintWeightsDrawWireToggle)
+	self.widgets["paint_displaySettings_colorCB"] = cmds.checkBox(label = "Color Feedback", parent = self.widgets["paint_displaySettings_form"], v = 1, cc = self.paintWeightsDrawColorToggle)
+
+
+	
+	
+
+
+    # # # # # # # # # # # # # # # # # # # # # # # # # # # # # # # # # # # # # # # # # # # # # # # # # # # # # # # # # # # # # # # 
+    # # # # # # # # # # # # # # # # # # # # # # # # # # # # # # # # # # # # # # # # # # # # # # # # # # # # # # # # # # # # # # # 
+    # # # # # # # # # # # # # # # # # # # # # # # # # # # # # # # # # # # # # # # # # # # # # # # # # # # # # # # # # # # # # # # 
+    # # # # # # # # # # # # # # # # # # # # # # # # # # # # # # # # # # # # # # # # # # # # # # # # # # # # # # # # # # # # # # # 
+    # # # # # # # # # # # # # # # # # # # # # # # # # # # # # # # # # # # # # # # # # # # # # # # # # # # # # # # # # # # # # # # 
+    def createJointMover_PhysiqueUI(self):
+	
+	
+	self.widgets["jointMover_physiqueFormLayout"] = cmds.formLayout(w = 410, h = 1200, vis = False, parent = self.widgets["masterFormLayout"])
+	
+        #get our image
+        backgroundImage = self.mayaToolsDir + "/General/Icons/ART/jointMoverLock.jpg"
+        self.widgets["jointMover_physiqueBackgroundImage"] = cmds.image(image = backgroundImage, w = 410, h = 600)
+	
+	#create a toolbar at the top of the UI
+	self.widgets["physiqueToolbarLayout"] = cmds.rowColumnLayout(h = 1200, nc = 5,  cat = [(1, "both", 2),(2, "both", 2),(3, "both", 2),(4, "both", 2),(5, "both", 2)], parent = self.widgets["jointMover_physiqueFormLayout"])
+	cmds.formLayout(self.widgets["jointMover_physiqueFormLayout"], edit = True, af = [(self.widgets["physiqueToolbarLayout"], "top", 0), (self.widgets["physiqueToolbarLayout"], "left", 65)])
+	
+	cmds.button(label = "Mirror L->R", w = 62, h = 45, c = partial(self.physique_Mirror, "left"))
+	cmds.button(label = "Mirror R->L", w = 62, h = 45, c = partial(self.physique_Mirror, "right"))
+	cmds.button(label = "Reset All", w = 62, h = 45, c = self.resetPhysiqueSliders_All)
+	cmds.button(label = "Save", w = 62, h = 45, c = self.physiqueSave)
+	cmds.button(label = "Load", w = 62, h = 45, c = self.physiqueLoad)
+	
+	
+	#create a row columnLayout with 3 columns
+	self.widgets["physique_scrollLayout"] = cmds.scrollLayout(h = 550, w = 350, hst = 0, parent = self.widgets["jointMover_physiqueFormLayout"])
+	self.widgets["physiqueRowColLayout"] = cmds.rowColumnLayout(nc = 3, cat = [(1, "both", 3), (2, "both", 3), (3, "both", 3)], parent = self.widgets["physique_scrollLayout"])
+	cmds.formLayout(self.widgets["jointMover_physiqueFormLayout"], edit = True, af = [(self.widgets["physique_scrollLayout"], "top", 50), (self.widgets["physique_scrollLayout"], "left", 65)])
+	
+	#create 3 columnLayouts, each in one of the columns of the row column Layout. This is just so we don't have to juggle what goes where.
+	self.widgets["phys_leftCol"] = cmds.columnLayout(parent = self.widgets["physiqueRowColLayout"], rs = 10)
+	self.widgets["phys_midCol"] = cmds.columnLayout(parent = self.widgets["physiqueRowColLayout"], rs = 10)
+	self.widgets["phys_rightCol"] = cmds.columnLayout(parent = self.widgets["physiqueRowColLayout"], rs = 10)
+	
+	# # # # # # # # # # # # # # # # # # # # # # # # # # # # # # # # # # # # # # # # # # # # # # # # # # # # #
+	#left column
+	armFrame = cmds.frameLayout(w = 100, cll = True, collapse = False, label = "Right Arm", borderStyle = "in", parent = self.widgets["phys_leftCol"], bgc = [.27, .5, .65])
+	legFrame = cmds.frameLayout(w = 100, cll = True, collapse = True, label = "Right Leg", borderStyle = "in", parent = self.widgets["phys_leftCol"], bgc = [.27, .5, .65])
+	
+	
+	#upper arm
+	cmds.text(label = "Up Arm Top Taper", parent = armFrame)
+	self.widgets["upperTopTaperR"] = cmds.floatSlider( min = -2, max = 2, value = 0, step = .01, parent = armFrame )
+	
+	cmds.text(label = "Up Arm Bot Taper", parent = armFrame)
+	self.widgets["upperBotTaperR"] = cmds.floatSlider( min = -2, max = 2, value = 0, step = .01, parent = armFrame )
+	
+	cmds.text(label = "Up Arm Mid Taper", parent = armFrame)
+	self.widgets["upperMidTaperR"] = cmds.floatSlider( min = -2, max = 2, value = 0, step = .01, parent = armFrame )
+	
+	cmds.text(label = "Up Arm Mid Expand", parent = armFrame)
+	self.widgets["upperMidExpandR"] = cmds.floatSlider( min = -2, max = 2, value = 0, step = .01, parent = armFrame )
+	
+	cmds.text(label = "Shoulder Size", parent = armFrame)
+	self.widgets["shoulderSizeR"] = cmds.floatSlider( min = -2, max = 2, value = 0, step = .01, parent = armFrame )
+	
+	cmds.text(label = "Elbow Size", parent = armFrame)
+	self.widgets["elbowSizeR"] = cmds.floatSlider( min = -2, max = 2, value = 0, step = .01, parent = armFrame )
+	
+	cmds.text(label = "Lo Arm Top Taper", parent = armFrame)
+	self.widgets["lowerTopTaperR"] = cmds.floatSlider( min = -2, max = 2, value = 0, step = .01, parent = armFrame )
+	
+	cmds.text(label = "Lo Arm Bot Taper", parent = armFrame)
+	self.widgets["lowerBotTaperR"] = cmds.floatSlider( min = -2, max = 2, value = 0, step = .01, parent = armFrame )
+	
+	cmds.text(label = "Lo Arm Mid Taper", parent = armFrame)
+	self.widgets["lowerMidTaperR"] = cmds.floatSlider( min = -2, max = 2, value = 0, step = .01, parent = armFrame )
+	
+	cmds.text(label = "Lo Arm Mid Expand", parent = armFrame)
+	self.widgets["lowerMidExpandR"] = cmds.floatSlider( min = -2, max = 2, value = 0, step = .01, parent = armFrame )
+	
+	
+	#leg
+	cmds.text(label = "Thigh Top Taper", parent = legFrame)
+	self.widgets["upperThighTopTaperR"] = cmds.floatSlider( min = -2, max = 2, value = 0, step = .01, parent = legFrame )
+	
+	cmds.text(label = "Thigh Bot Taper", parent = legFrame)
+	self.widgets["upperThighBotTaperR"] = cmds.floatSlider( min = -2, max = 2, value = 0, step = .01, parent = legFrame )
+	
+	cmds.text(label = "Thigh Inner", parent = legFrame)
+	self.widgets["upperThighInnerR"] = cmds.floatSlider( min = -2, max = 2, value = 0, step = .01, parent = legFrame )
+	
+	cmds.text(label = "Thigh Outer", parent = legFrame)
+	self.widgets["upperThighOuterR"] = cmds.floatSlider( min = -2, max = 2, value = 0, step = .01, parent = legFrame )
+	
+		
+	cmds.text(label = "Calf Top Taper", parent = legFrame)
+	self.widgets["calfTopTaperR"] = cmds.floatSlider( min = -2, max = 2, value = 0, step = .01, parent = legFrame )
+	
+	cmds.text(label = "Calf Bot Taper", parent = legFrame)
+	self.widgets["calfBotTaperR"] = cmds.floatSlider( min = -2, max = 2, value = 0, step = .01, parent = legFrame )
+	
+	cmds.text(label = "Calf Mid Taper", parent = legFrame)
+	self.widgets["calfMidTaperR"] = cmds.floatSlider( min = -2, max = 2, value = 0, step = .01, parent = legFrame )
+	
+	cmds.text(label = "Calf Mid Expand", parent = legFrame)
+	self.widgets["calfMidExpandR"] = cmds.floatSlider( min = -2, max = 2, value = 0, step = .01, parent = legFrame )
+	
+	
+	#RMB menu creation
+	menu = cmds.popupMenu(b = 3, parent = armFrame)
+	cmds.menuItem(parent = menu, label = "Reset", c = partial(self.resetPhysiqueSliders, ["upperTopTaperR","upperBotTaperR","upperMidTaperR","upperMidExpandR","lowerTopTaperR","lowerBotTaperR","lowerMidTaperR","lowerMidExpandR", "shoulderSizeR", "elbowSizeR" ], ["r_upperarm_shapes", "r_lowerarm_shapes", "r_shoulder_shapes", "r_elbow_shapes"] ))
+	
+	menu = cmds.popupMenu(b = 3, parent = legFrame)
+	cmds.menuItem(parent = menu, label = "Reset", c = partial(self.resetPhysiqueSliders, ["upperThighTopTaperR","upperThighBotTaperR","upperThighInnerR","upperThighOuterR","calfTopTaperR","calfBotTaperR","calfMidTaperR","calfMidExpandR"], ["thigh_r_shapes", "calf_r_shapes"]))
+	
+	#Connect Controls
+	cmds.connectControl(self.widgets["upperTopTaperR"], "r_upperarm_shapes.r_upperarm_top_taper")
+	cmds.connectControl(self.widgets["upperBotTaperR"], "r_upperarm_shapes.r_upperarm_bot_taper")
+	cmds.connectControl(self.widgets["upperMidTaperR"], "r_upperarm_shapes.r_upperarm_mid_taper")
+	cmds.connectControl(self.widgets["upperMidExpandR"], "r_upperarm_shapes.r_upperarm_mid_expand")
+	
+	cmds.connectControl(self.widgets["lowerTopTaperR"], "r_lowerarm_shapes.r_lowerarm_top_taper")
+	cmds.connectControl(self.widgets["lowerBotTaperR"], "r_lowerarm_shapes.r_lowerarm_bot_taper")
+	cmds.connectControl(self.widgets["lowerMidTaperR"], "r_lowerarm_shapes.r_lowerarm_mid_taper")
+	cmds.connectControl(self.widgets["lowerMidExpandR"], "r_lowerarm_shapes.r_lowerarm_mid_expand")
+	
+	cmds.connectControl(self.widgets["upperThighTopTaperR"], "thigh_r_shapes.r_thigh_top_taper")
+	cmds.connectControl(self.widgets["upperThighBotTaperR"], "thigh_r_shapes.r_thigh_bot_taper")
+	cmds.connectControl(self.widgets["upperThighInnerR"], "thigh_r_shapes.r_thigh_inner_bulge")
+	cmds.connectControl(self.widgets["upperThighOuterR"], "thigh_r_shapes.r_thigh_outer_bulge")
+	
+	cmds.connectControl(self.widgets["calfTopTaperR"], "calf_r_shapes.r_calf_top_taper")
+	cmds.connectControl(self.widgets["calfBotTaperR"], "calf_r_shapes.r_calf_bot_taper")
+	cmds.connectControl(self.widgets["calfMidTaperR"], "calf_r_shapes.r_calf_mid_taper")
+	cmds.connectControl(self.widgets["calfMidExpandR"], "calf_r_shapes.r_calf_mid_expand")
+	
+	cmds.connectControl(self.widgets["shoulderSizeR"], "r_shoulder_shapes.r_shoulder_grow")
+	cmds.connectControl(self.widgets["elbowSizeR"], "r_elbow_shapes.r_elbow_grow")
+	
+	
+	
+	
+	# # # # # # # # # # # # # # # # # # # # # # # # # # # # # # # # # # # # # # # # # # # # # # # # # # # # #
+	#mid column
+	headFrame = cmds.frameLayout(w = 100, cll = True, collapse = False, label = "Head", borderStyle = "in", parent = self.widgets["phys_midCol"],bgc = [.27, .5, .65])
+	neckFrame = cmds.frameLayout(w = 100, cll = True, collapse = False, label = "Neck", borderStyle = "in", parent = self.widgets["phys_midCol"],bgc = [.27, .5, .65])
+	bodyFrmae = cmds.frameLayout(w = 100, cll = True, collapse = True, label = "Body", borderStyle = "in", parent = self.widgets["phys_midCol"],bgc = [.27, .5, .65])
+	hipFrmae = cmds.frameLayout(w = 100, cll = True, collapse = True, label = "Pelvis", borderStyle = "in", parent = self.widgets["phys_midCol"],bgc = [.27, .5, .65])
+	
+	
+	#head
+	cmds.text(label = "Square Shape", parent = headFrame)
+	self.widgets["head_squareShape"] = cmds.floatSlider( min = -2, max = 2, value = 0, step = .01, parent = headFrame )
+	
+	cmds.text(label = "Top Taper", parent = headFrame)
+	self.widgets["head_topTaper"] = cmds.floatSlider( min = -2, max = 2, value = 0, step = .01, parent = headFrame )
+	
+	cmds.text(label = "Bottom Taper", parent = headFrame)
+	self.widgets["head_botTaper"] = cmds.floatSlider( min = -2, max = 2, value = 0, step = .01, parent = headFrame )
+	
+	
+	#neck
+	cmds.text(label = "Neck 1 Bot Taper", parent = neckFrame)
+	self.widgets["neck1_botTaper"] = cmds.floatSlider( min = -2, max = 2, value = 0, step = .01, parent = neckFrame )
+	
+	cmds.text(label = "Neck 1 Top Taper", parent = neckFrame)
+	self.widgets["neck1_topTaper"] = cmds.floatSlider( min = -2, max = 2, value = 0, step = .01, parent = neckFrame )
+	
+	cmds.text(label = "Neck 2 Bot Taper", parent = neckFrame)
+	self.widgets["neck2_botTaper"] = cmds.floatSlider( min = -2, max = 2, value = 0, step = .01, parent = neckFrame )
+	
+	cmds.text(label = "Neck 2 Top Taper", parent = neckFrame)
+	self.widgets["neck2_topTaper"] = cmds.floatSlider( min = -2, max = 2, value = 0, step = .01, parent = neckFrame )
+	
+	cmds.text(label = "Neck 3 Bot Taper", parent = neckFrame)
+	self.widgets["neck3_botTaper"] = cmds.floatSlider( min = -2, max = 2, value = 0, step = .01, parent = neckFrame )
+	
+	cmds.text(label = "Neck 3 Top Taper", parent = neckFrame)
+	self.widgets["neck3_topTaper"] = cmds.floatSlider( min = -2, max = 2, value = 0, step = .01, parent = neckFrame )
+	
+	
+	#body
+	cmds.text(label = "Weight", parent = bodyFrmae)
+	self.widgets["body_weight"] = cmds.floatSlider( min = 0, max = 2, value = 0, step = .01, parent = bodyFrmae )
+	
+	cmds.text(label = "Bust Size", parent = bodyFrmae)
+	self.widgets["body_bust"] = cmds.floatSlider( min = 0, max = 2, value = 0, step = .01, parent = bodyFrmae )
+	
+	bodyFrameAdvanced = cmds.frameLayout(w = 90, collapse = True, cll = True, label = "Advanced", borderStyle = "in", parent = bodyFrmae)
+	
+	cmds.text(label = "Seg 1 Bot Taper", parent = bodyFrameAdvanced)
+	self.widgets["spine1_botTaper"] = cmds.floatSlider( min = -2, max = 2, value = 0, step = .01, parent = bodyFrameAdvanced )
+	
+	cmds.text(label = "Seg 1 Top Taper", parent = bodyFrameAdvanced)
+	self.widgets["spine1_topTaper"] = cmds.floatSlider( min = -2, max = 2, value = 0, step = .01, parent = bodyFrameAdvanced )
+	
+	cmds.text(label = "Seg 2 Bot Taper", parent = bodyFrameAdvanced)
+	self.widgets["spine2_botTaper"] = cmds.floatSlider( min = -2, max = 2, value = 0, step = .01, parent = bodyFrameAdvanced )
+	
+	cmds.text(label = "Seg 2 Top Taper", parent = bodyFrameAdvanced)
+	self.widgets["spine2_topTaper"] = cmds.floatSlider( min = -2, max = 2, value = 0, step = .01, parent = bodyFrameAdvanced )
+	
+	cmds.text(label = "Seg 3 Bot Taper", parent = bodyFrameAdvanced)
+	self.widgets["spine3_botTaper"] = cmds.floatSlider( min = -2, max = 2, value = 0, step = .01, parent = bodyFrameAdvanced )
+	
+	cmds.text(label = "Seg 3 Top Taper", parent = bodyFrameAdvanced)
+	self.widgets["spine3_topTaper"] = cmds.floatSlider( min = -2, max = 2, value = 0, step = .01, parent = bodyFrameAdvanced )
+	
+	cmds.text(label = "Seg 4 Bot Taper", parent = bodyFrameAdvanced)
+	self.widgets["spine4_botTaper"] = cmds.floatSlider( min = -2, max = 2, value = 0, step = .01, parent = bodyFrameAdvanced )
+	
+	cmds.text(label = "Seg 4 Top Taper", parent = bodyFrameAdvanced)
+	self.widgets["spine4_topTaper"] = cmds.floatSlider( min = -2, max = 2, value = 0, step = .01, parent = bodyFrameAdvanced )
+	
+	cmds.text(label = "Seg 5 Bot Taper", parent = bodyFrameAdvanced)
+	self.widgets["spine5_botTaper"] = cmds.floatSlider( min = -2, max = 2, value = 0, step = .01, parent = bodyFrameAdvanced )
+	
+	cmds.text(label = "Seg 5 Top Taper", parent = bodyFrameAdvanced)
+	self.widgets["spine5_topTaper"] = cmds.floatSlider( min = -2, max = 2, value = 0, step = .01, parent = bodyFrameAdvanced )
+	
+	
+	
+	#pelvis
+	cmds.text(label = "Top Taper", parent = hipFrmae)
+	self.widgets["hip_topTaper"] = cmds.floatSlider( min = -2, max = 2, value = 0, step = .01, parent = hipFrmae )
+	
+	cmds.text(label = "Bottom Taper", parent = hipFrmae)
+	self.widgets["hip_bottomTaper"] = cmds.floatSlider( min = -2, max = 2, value = 0, step = .01, parent = hipFrmae )
+	
+	cmds.text(label = "Backside", parent = hipFrmae)
+	self.widgets["hip_backside"] = cmds.floatSlider( min = -2, max = 2, value = 0, step = .01, parent = hipFrmae )
+	
+	cmds.text(label = "Curvature", parent = hipFrmae)
+	self.widgets["hip_curvature"] = cmds.floatSlider( min = -2, max = 2, value = 0, step = .01, parent = hipFrmae )
+	
+
+	
+	#RMB menu creation
+	menu = cmds.popupMenu(b = 3, parent = headFrame)
+	cmds.menuItem(parent = menu, label = "Reset", c = partial(self.resetPhysiqueSliders, ["head_squareShape","head_topTaper","head_botTaper"], ["head_shapes"] ))
+	
+	menu = cmds.popupMenu(b = 3, parent = neckFrame)
+	cmds.menuItem(parent = menu, label = "Reset", c = partial(self.resetPhysiqueSliders, ["neck1_botTaper","neck1_topTaper","neck2_botTaper","neck2_topTaper","neck3_botTaper","neck3_topTaper"], ["neck1_shapes", "neck2_shapes", "neck3_shapes"]))
+
+	menu = cmds.popupMenu(b = 3, parent = bodyFrmae)
+	cmds.menuItem(parent = menu, label = "Reset", c = partial(self.resetPhysiqueSliders, ["body_weight","body_bust","spine1_botTaper","spine1_topTaper","spine2_botTaper","spine2_topTaper","spine3_botTaper","spine3_topTaper", "spine4_botTaper","spine4_topTaper", "spine5_botTaper", "spine5_topTaper"], ["spine1_shapes", "spine2_shapes", "spine3_shapes", "spine4_shapes", "spine5_shapes"]))
+	
+	menu = cmds.popupMenu(b = 3, parent = hipFrmae)
+	cmds.menuItem(parent = menu, label = "Reset", c = partial(self.resetPhysiqueSliders, ["hip_topTaper","hip_bottomTaper","hip_backside","hip_curvature"], ["pelvis_shapes"]))
+	
+	
+	
+	#Connect Controls
+	cmds.connectControl(self.widgets["head_squareShape"], "head_shapes.head_square")
+	cmds.connectControl(self.widgets["head_topTaper"], "head_shapes.head_drop")
+	cmds.connectControl(self.widgets["head_botTaper"], "head_shapes.head_heart")
+	
+	cmds.connectControl(self.widgets["neck1_botTaper"], "neck1_shapes.neck_1_bottom_taper")
+	cmds.connectControl(self.widgets["neck1_topTaper"], "neck1_shapes.neck_1_top_taper")
+	cmds.connectControl(self.widgets["neck2_botTaper"], "neck2_shapes.neck_2_bottom_taper")
+	cmds.connectControl(self.widgets["neck2_topTaper"], "neck2_shapes.neck_2_top_taper")
+	cmds.connectControl(self.widgets["neck3_botTaper"], "neck3_shapes.neck_3_bottom_taper")
+	cmds.connectControl(self.widgets["neck3_topTaper"], "neck3_shapes.neck_3_top_taper")
+	
+	cmds.connectControl(self.widgets["body_weight"], "spine1_shapes.spine_1_fat")
+	cmds.connectControl(self.widgets["body_bust"], "spine1_shapes.spine_1_bust")
+	
+	try:
+	    cmds.connectAttr("spine1_shapes.spine_1_bust", "spine2_shapes.spine_2_bust")
+	    cmds.connectAttr("spine1_shapes.spine_1_bust", "spine3_shapes.spine_3_bust")
+	    cmds.connectAttr("spine1_shapes.spine_1_bust", "spine4_shapes.spine_4_bust")
+	    cmds.connectAttr("spine1_shapes.spine_1_bust", "spine5_shapes.spine_5_bust")
+	    
+	    cmds.connectAttr("spine1_shapes.spine_1_fat", "spine2_shapes.spine_2_fat")
+	    cmds.connectAttr("spine1_shapes.spine_1_fat", "spine3_shapes.spine_3_fat")
+	    cmds.connectAttr("spine1_shapes.spine_1_fat", "spine4_shapes.spine_4_fat")
+	    cmds.connectAttr("spine1_shapes.spine_1_fat", "spine5_shapes.spine_5_fat")
+	    
+	except:
+	    pass
+	
+	cmds.connectControl(self.widgets["spine1_botTaper"], "spine1_shapes.spine_1_bot_taper")
+	cmds.connectControl(self.widgets["spine1_topTaper"], "spine1_shapes.spine_1_top_taper")
+	cmds.connectControl(self.widgets["spine2_botTaper"], "spine2_shapes.spine_2_bot_taper")
+	cmds.connectControl(self.widgets["spine2_topTaper"], "spine2_shapes.spine_2_top_taper")
+	cmds.connectControl(self.widgets["spine3_botTaper"], "spine3_shapes.spine_3_bot_taper")
+	cmds.connectControl(self.widgets["spine3_topTaper"], "spine3_shapes.spine_3_top_taper")
+	cmds.connectControl(self.widgets["spine4_botTaper"], "spine4_shapes.spine_4_bot_taper")
+	cmds.connectControl(self.widgets["spine4_topTaper"], "spine4_shapes.spine_4_top_taper")
+	cmds.connectControl(self.widgets["spine5_botTaper"], "spine5_shapes.spine_5_bot_taper")
+	cmds.connectControl(self.widgets["spine5_topTaper"], "spine5_shapes.spine_5_top_taper")
+	
+	cmds.connectControl(self.widgets["hip_topTaper"], "pelvis_shapes.pelvis_top_taper")
+	cmds.connectControl(self.widgets["hip_bottomTaper"], "pelvis_shapes.pelvis_bottom_taper")
+	cmds.connectControl(self.widgets["hip_backside"], "pelvis_shapes.pelvis_butt")
+	cmds.connectControl(self.widgets["hip_curvature"], "pelvis_shapes.pelvis_hips")
+
+
+	
+	
+	
+	
+	# # # # # # # # # # # # # # # # # # # # # # # # # # # # # # # # # # # # # # # # # # # # # # # # # # # # #
+	#right column
+	lArmFrame = cmds.frameLayout(w = 100, cll = True, collapse = False, label = "Left Arm", borderStyle = "in", parent = self.widgets["phys_rightCol"],bgc = [.27, .5, .65])
+	lLegFrame = cmds.frameLayout(w = 100, cll = True, collapse = True, label = "Left Leg", borderStyle = "in", parent = self.widgets["phys_rightCol"],bgc = [.27, .5, .65])
+	
+	#upper arm
+	cmds.text(label = "Up Arm Top Taper", parent = lArmFrame)
+	self.widgets["upperTopTaperL"] = cmds.floatSlider( min = -2, max = 2, value = 0, step = .01, parent = lArmFrame )
+	
+	cmds.text(label = "Up Arm Bot Taper", parent = lArmFrame)
+	self.widgets["upperBotTaperL"] = cmds.floatSlider( min = -2, max = 2, value = 0, step = .01, parent = lArmFrame )
+	
+	cmds.text(label = "Up Arm Mid Taper", parent = lArmFrame)
+	self.widgets["upperMidTaperL"] = cmds.floatSlider( min = -2, max = 2, value = 0, step = .01, parent = lArmFrame )
+	
+	cmds.text(label = "Up Arm Mid Expand", parent = lArmFrame)
+	self.widgets["upperMidExpandL"] = cmds.floatSlider( min = -2, max = 2, value = 0, step = .01, parent = lArmFrame )
+
+	cmds.text(label = "Shoulder Size", parent = lArmFrame)
+	self.widgets["shoulderSizeL"] = cmds.floatSlider( min = -2, max = 2, value = 0, step = .01, parent = lArmFrame )
+	
+	cmds.text(label = "Elbow Size", parent = lArmFrame)
+	self.widgets["elbowSizeL"] = cmds.floatSlider( min = -2, max = 2, value = 0, step = .01, parent = lArmFrame )
+	
+	cmds.text(label = "Lo Arm Top Taper", parent = lArmFrame)
+	self.widgets["lowerTopTaperL"] = cmds.floatSlider( min = -2, max = 2, value = 0, step = .01, parent = lArmFrame )
+	
+	cmds.text(label = "Lo Arm Bot Taper", parent = lArmFrame)
+	self.widgets["lowerBotTaperL"] = cmds.floatSlider( min = -2, max = 2, value = 0, step = .01, parent = lArmFrame )
+	
+	cmds.text(label = "Lo Arm Mid Taper", parent = lArmFrame)
+	self.widgets["lowerMidTaperL"] = cmds.floatSlider( min = -2, max = 2, value = 0, step = .01, parent = lArmFrame )
+	
+	cmds.text(label = "Lo Arm Mid Expand", parent = lArmFrame)
+	self.widgets["lowerMidExpandL"] = cmds.floatSlider( min = -2, max = 2, value = 0, step = .01, parent = lArmFrame )
+	
+	
+	
+	#leg
+	cmds.text(label = "Thigh Top Taper", parent = lLegFrame)
+	self.widgets["upperThighTopTaperL"] = cmds.floatSlider( min = -2, max = 2, value = 0, step = .01, parent = lLegFrame )
+	
+	cmds.text(label = "Thigh Bot Taper", parent = lLegFrame)
+	self.widgets["upperThighBotTaperL"] = cmds.floatSlider( min = -2, max = 2, value = 0, step = .01, parent = lLegFrame )
+	
+	cmds.text(label = "Thigh Inner", parent = lLegFrame)
+	self.widgets["upperThighInnerL"] = cmds.floatSlider( min = -2, max = 2, value = 0, step = .01, parent = lLegFrame )
+	
+	cmds.text(label = "Thigh Outer", parent = lLegFrame)
+	self.widgets["upperThighOuterL"] = cmds.floatSlider( min = -2, max = 2, value = 0, step = .01, parent = lLegFrame )
+	
+		
+	cmds.text(label = "Calf Top Taper", parent = lLegFrame)
+	self.widgets["calfTopTaperL"] = cmds.floatSlider( min = -2, max = 2, value = 0, step = .01, parent = lLegFrame )
+	
+	cmds.text(label = "Calf Bot Taper", parent = lLegFrame)
+	self.widgets["calfBotTaperL"] = cmds.floatSlider( min = -2, max = 2, value = 0, step = .01, parent = lLegFrame )
+	
+	cmds.text(label = "Calf Mid Taper", parent = lLegFrame)
+	self.widgets["calfMidTaperL"] = cmds.floatSlider( min = -2, max = 2, value = 0, step = .01, parent = lLegFrame )
+	
+	cmds.text(label = "Calf Mid Expand", parent = lLegFrame)
+	self.widgets["calfMidExpandL"] = cmds.floatSlider( min = -2, max = 2, value = 0, step = .01, parent = lLegFrame )
+	
+	
+	#RMB menu creation
+	menu = cmds.popupMenu(b = 3, parent = lArmFrame)
+	cmds.menuItem(parent = menu, label = "Reset", c = partial(self.resetPhysiqueSliders, ["upperTopTaperL","upperBotTaperL","upperMidTaperL","upperMidExpandL","lowerTopTaperL","lowerBotTaperL","lowerMidTaperL","lowerMidExpandL", "shoulderSizeL", "elbowSizeL"], ["l_upperarm_shapes", "l_lowerarm_shapes", "l_elbow_shapes", "l_shoulder_shapes" ] ))
+	
+	menu = cmds.popupMenu(b = 3, parent = lLegFrame)
+	cmds.menuItem(parent = menu, label = "Reset", c = partial(self.resetPhysiqueSliders, ["upperThighTopTaperL","upperThighBotTaperL","upperThighInnerL","upperThighOuterL","calfTopTaperL","calfBotTaperL","calfMidTaperL","calfMidExpandL"], ["thigh_l_shapes", "calf_l_shapes"]))
+	
+	#Connect Controls
+	cmds.connectControl(self.widgets["upperTopTaperL"], "l_upperarm_shapes.l_upperarm_top_taper")
+	cmds.connectControl(self.widgets["upperBotTaperL"], "l_upperarm_shapes.l_upperarm_bot_taper")
+	cmds.connectControl(self.widgets["upperMidTaperL"], "l_upperarm_shapes.l_upperarm_mid_taper")
+	cmds.connectControl(self.widgets["upperMidExpandL"], "l_upperarm_shapes.l_upperarm_mid_expand")
+	
+	cmds.connectControl(self.widgets["lowerTopTaperL"], "l_lowerarm_shapes.l_lowerarm_top_taper")
+	cmds.connectControl(self.widgets["lowerBotTaperL"], "l_lowerarm_shapes.l_lowerarm_bot_taper")
+	cmds.connectControl(self.widgets["lowerMidTaperL"], "l_lowerarm_shapes.l_lowerarm_mid_taper")
+	cmds.connectControl(self.widgets["lowerMidExpandL"], "l_lowerarm_shapes.l_lowerarm_mid_expand")
+	
+	cmds.connectControl(self.widgets["upperThighTopTaperL"], "thigh_l_shapes.l_thigh_top_taper")
+	cmds.connectControl(self.widgets["upperThighBotTaperL"], "thigh_l_shapes.l_thigh_bot_taper")
+	cmds.connectControl(self.widgets["upperThighInnerL"], "thigh_l_shapes.l_thigh_inner_bulge")
+	cmds.connectControl(self.widgets["upperThighOuterL"], "thigh_l_shapes.l_thigh_outer_bulge")
+	
+	cmds.connectControl(self.widgets["calfTopTaperL"], "calf_l_shapes.l_calf_top_taper")
+	cmds.connectControl(self.widgets["calfBotTaperL"], "calf_l_shapes.l_calf_bot_taper")
+	cmds.connectControl(self.widgets["calfMidTaperL"], "calf_l_shapes.l_calf_mid_taper")
+	cmds.connectControl(self.widgets["calfMidExpandL"], "calf_l_shapes.l_calf_mid_expand")
+	
+	cmds.connectControl(self.widgets["shoulderSizeL"], "l_shoulder_shapes.l_shoulder_grow")
+	cmds.connectControl(self.widgets["elbowSizeL"], "l_elbow_shapes.l_elbow_grow")
+	                 
+
+    # # # # # # # # # # # # # # # # # # # # # # # # # # # # # # # # # # # # # # # # # # # # # # # # # # # # # # # # # # # # # # # 
+    # # # # # # # # # # # # # # # # # # # # # # # # # # # # # # # # # # # # # # # # # # # # # # # # # # # # # # # # # # # # # # # 
+    # # # # # # # # # # # # # # # # # # # # # # # # # # # # # # # # # # # # # # # # # # # # # # # # # # # # # # # # # # # # # # # 
+    # # # # # # # # # # # # # # # # # # # # # # # # # # # # # # # # # # # # # # # # # # # # # # # # # # # # # # # # # # # # # # # 
+    # # # # # # # # # # # # # # # # # # # # # # # # # # # # # # # # # # # # # # # # # # # # # # # # # # # # # # # # # # # # # # # 
+    def resetPhysiqueSliders(self, sliders, blendshapes, *args):
+	
+	for slider in sliders:
+	    cmds.floatSlider(self.widgets[slider], edit = True, value = 0)
+	    
+	for shape in blendshapes:
+	    targets = cmds.listAttr(shape + ".w", m = True)
+	    for target in targets:
+		try:
+		    cmds.setAttr(shape + "." + target, 0)
+		except:
+		    pass
+
+    # # # # # # # # # # # # # # # # # # # # # # # # # # # # # # # # # # # # # # # # # # # # # # # # # # # # # # # # # # # # # # # 
+    # # # # # # # # # # # # # # # # # # # # # # # # # # # # # # # # # # # # # # # # # # # # # # # # # # # # # # # # # # # # # # # 
+    # # # # # # # # # # # # # # # # # # # # # # # # # # # # # # # # # # # # # # # # # # # # # # # # # # # # # # # # # # # # # # # 
+    # # # # # # # # # # # # # # # # # # # # # # # # # # # # # # # # # # # # # # # # # # # # # # # # # # # # # # # # # # # # # # # 
+    # # # # # # # # # # # # # # # # # # # # # # # # # # # # # # # # # # # # # # # # # # # # # # # # # # # # # # # # # # # # # # # 
+    def resetPhysiqueSliders_All(self, *args):
+	
+	sliders = ["upperTopTaperR","upperBotTaperR","upperMidTaperR","upperMidExpandR","lowerTopTaperR","lowerBotTaperR","lowerMidTaperR","lowerMidExpandR", "upperThighTopTaperR","upperThighBotTaperR","upperThighInnerR","upperThighOuterR",
+	           "calfTopTaperR","calfBotTaperR","calfMidTaperR","calfMidExpandR", "head_squareShape","head_topTaper","head_botTaper", "neck1_botTaper","neck1_topTaper","neck2_botTaper","neck2_topTaper","neck3_botTaper","neck3_topTaper",
+	           "body_weight","body_bust","spine1_botTaper","spine1_topTaper","spine2_botTaper","spine2_topTaper","spine3_botTaper","spine3_topTaper", "spine4_botTaper","spine4_topTaper", "spine5_botTaper", "spine5_topTaper",
+	           "hip_topTaper","hip_bottomTaper","hip_backside","hip_curvature", "upperTopTaperL","upperBotTaperL","upperMidTaperL","upperMidExpandL","lowerTopTaperL","lowerBotTaperL","lowerMidTaperL","lowerMidExpandL",
+	           "upperThighTopTaperL","upperThighBotTaperL","upperThighInnerL","upperThighOuterL","calfTopTaperL","calfBotTaperL","calfMidTaperL","calfMidExpandL", "shoulderSizeL", "shoulderSizeR", "elbowSizeL", "elbowSizeR"]
+	
+	blendshapes = ["r_upperarm_shapes", "r_lowerarm_shapes", "l_upperarm_shapes", "l_lowerarm_shapes", "thigh_r_shapes", "calf_r_shapes", "thigh_l_shapes", "calf_l_shapes", "head_shapes",
+	               "neck1_shapes", "neck2_shapes", "neck3_shapes", "spine1_shapes", "spine2_shapes", "spine3_shapes", "spine4_shapes", "spine5_shapes", "pelvis_shapes", "l_elbow_shapes", "r_elbow_shapes", "l_shoulder_shapes", "r_shoulder_shapes"]
+	
+	for slider in sliders:
+	    cmds.floatSlider(self.widgets[slider], edit = True, value = 0)
+	    
+	for shape in blendshapes:
+	    targets = cmds.listAttr(shape + ".w", m = True)
+	    for target in targets:
+		try:
+		    cmds.setAttr(shape + "." + target, 0)
+		except:
+		    pass
+	
+	
+    # # # # # # # # # # # # # # # # # # # # # # # # # # # # # # # # # # # # # # # # # # # # # # # # # # # # # # # # # # # # # # # 
+    # # # # # # # # # # # # # # # # # # # # # # # # # # # # # # # # # # # # # # # # # # # # # # # # # # # # # # # # # # # # # # # 
+    # # # # # # # # # # # # # # # # # # # # # # # # # # # # # # # # # # # # # # # # # # # # # # # # # # # # # # # # # # # # # # # 
+    # # # # # # # # # # # # # # # # # # # # # # # # # # # # # # # # # # # # # # # # # # # # # # # # # # # # # # # # # # # # # # # 
+    # # # # # # # # # # # # # # # # # # # # # # # # # # # # # # # # # # # # # # # # # # # # # # # # # # # # # # # # # # # # # # # 
+    def physique_Mirror(self, fromSide, *args):
+	
+	
+	rightSideSliders = ["upperTopTaperR","upperBotTaperR","upperMidTaperR","upperMidExpandR","lowerTopTaperR","lowerBotTaperR","lowerMidTaperR","lowerMidExpandR", "upperThighTopTaperR","upperThighBotTaperR","upperThighInnerR","upperThighOuterR",
+	           "calfTopTaperR","calfBotTaperR","calfMidTaperR","calfMidExpandR", "shoulderSizeR", "elbowSizeR"]
+	
+	rightSideShapes = ["r_upperarm_shapes", "r_lowerarm_shapes", "thigh_r_shapes", "calf_r_shapes", "r_elbow_shapes", "r_shoulder_shapes"]
+	
+	leftSideSliders = ["upperTopTaperL","upperBotTaperL","upperMidTaperL","upperMidExpandL","lowerTopTaperL","lowerBotTaperL","lowerMidTaperL","lowerMidExpandL",
+	           "upperThighTopTaperL","upperThighBotTaperL","upperThighInnerL","upperThighOuterL","calfTopTaperL","calfBotTaperL","calfMidTaperL","calfMidExpandL", "shoulderSizeL", "elbowSizeL"]
+	
+	leftSideShapes = ["l_upperarm_shapes", "l_lowerarm_shapes", "thigh_l_shapes", "calf_l_shapes", "l_elbow_shapes", "l_shoulder_shapes"]
+	
+	if fromSide == "right":
+	    for i in range(len(rightSideSliders)):
+		val = cmds.floatSlider(self.widgets[rightSideSliders[i]], q = True, v = True)
+		cmds.floatSlider(self.widgets[leftSideSliders[i]], edit = True, v = val)
+	    
+	    for i in range(len(rightSideShapes)):
+		targets = cmds.listAttr(rightSideShapes[i] + ".w", m = True)
+		for target in targets:
+		    if target.find("r_") == 0:
+			newTarget = "l_" + target.partition("r_")[2]
+			    
+		    if target.find("_r_") != -1:
+			newTarget = target.partition("_r_")[0] + "_l_" + target.partition("_r_")[2]
+			
+		    val = cmds.getAttr(rightSideShapes[i] + "." + target)
+		    cmds.setAttr(leftSideShapes[i] + "." + newTarget, val)
+			
+		
+	if fromSide == "left":
+	    for i in range(len(leftSideSliders)):
+		val = cmds.floatSlider(self.widgets[leftSideSliders[i]], q = True, v = True)
+		cmds.floatSlider(self.widgets[rightSideSliders[i]], edit = True, v = val)
+		
+	    for i in range(len(leftSideShapes)):
+		targets = cmds.listAttr(leftSideShapes[i] + ".w", m = True)
+		for target in targets:
+		    if target.find("l_") == 0:
+			newTarget = "r_" + target.partition("l_")[2]
+			    
+		    if target.find("_l_") != -1:
+			newTarget = target.partition("_l_")[0] + "_r_" + target.partition("_l_")[2]
+			
+		    val = cmds.getAttr(leftSideShapes[i] + "." + target)
+		    cmds.setAttr(rightSideShapes[i] + "." + newTarget, val)
+
+	
+    # # # # # # # # # # # # # # # # # # # # # # # # # # # # # # # # # # # # # # # # # # # # # # # # # # # # # # # # # # # # # # # 
+    # # # # # # # # # # # # # # # # # # # # # # # # # # # # # # # # # # # # # # # # # # # # # # # # # # # # # # # # # # # # # # # 
+    # # # # # # # # # # # # # # # # # # # # # # # # # # # # # # # # # # # # # # # # # # # # # # # # # # # # # # # # # # # # # # # 
+    # # # # # # # # # # # # # # # # # # # # # # # # # # # # # # # # # # # # # # # # # # # # # # # # # # # # # # # # # # # # # # # 
+    # # # # # # # # # # # # # # # # # # # # # # # # # # # # # # # # # # # # # # # # # # # # # # # # # # # # # # # # # # # # # # # 
+    def physiqueSave(self, *args):
+	
+	cmds.promptDialog(title = "Physique", message = "Template Name: ")
+	name = cmds.promptDialog(q = True, text = True)
+	
+	blendshapes = ["r_upperarm_shapes", "r_lowerarm_shapes", "l_upperarm_shapes", "l_lowerarm_shapes", "thigh_r_shapes", "calf_r_shapes", "thigh_l_shapes", "calf_l_shapes", "head_shapes",
+	               "neck1_shapes", "neck2_shapes", "neck3_shapes", "spine1_shapes", "spine2_shapes", "spine3_shapes", "spine4_shapes", "spine5_shapes", "pelvis_shapes", "l_elbow_shapes", "r_elbow_shapes", "l_shoulder_shapes", "r_shoulder_shapes"]
+	
+	writePath = self.mayaToolsDir + "/General/ART/Physiques"
+	if not os.path.exists(writePath):
+	    os.makedirs(writePath)
+	    
+	f = open(writePath + "/" + name + ".txt", 'w')
+	data = []
+	
+	for shape in blendshapes:
+	    targets = cmds.listAttr(shape + ".w", m = True)
+	    for target in targets:
+		val = cmds.getAttr(shape + "." + target)
+		data.append([target, val])
+		
+	cPickle.dump(data, f)
+	f.close()
+	
+		
+    # # # # # # # # # # # # # # # # # # # # # # # # # # # # # # # # # # # # # # # # # # # # # # # # # # # # # # # # # # # # # # # 
+    # # # # # # # # # # # # # # # # # # # # # # # # # # # # # # # # # # # # # # # # # # # # # # # # # # # # # # # # # # # # # # # 
+    # # # # # # # # # # # # # # # # # # # # # # # # # # # # # # # # # # # # # # # # # # # # # # # # # # # # # # # # # # # # # # # 
+    # # # # # # # # # # # # # # # # # # # # # # # # # # # # # # # # # # # # # # # # # # # # # # # # # # # # # # # # # # # # # # # 
+    # # # # # # # # # # # # # # # # # # # # # # # # # # # # # # # # # # # # # # # # # # # # # # # # # # # # # # # # # # # # # # # 
+    def physiqueLoad(self, *args):
+	
+	if not os.path.exists(self.mayaToolsDir + "/General/ART/Physiques"):
+	    os.makedirs(self.mayaToolsDir + "/General/ART/Physiques")
+	    
+	    
+	fileName = cmds.fileDialog2(startingDirectory = self.mayaToolsDir + "/General/ART/Physiques/", fm = 1)[0]
+	
+	
+	blendshapes = ["r_upperarm_shapes", "r_lowerarm_shapes", "l_upperarm_shapes", "l_lowerarm_shapes", "thigh_r_shapes", "calf_r_shapes", "thigh_l_shapes", "calf_l_shapes", "head_shapes",
+	               "neck1_shapes", "neck2_shapes", "neck3_shapes", "spine1_shapes", "spine2_shapes", "spine3_shapes", "spine4_shapes", "spine5_shapes", "pelvis_shapes", "l_elbow_shapes", "r_elbow_shapes", "l_shoulder_shapes", "r_shoulder_shapes"]
+	
+	filePath = fileName
+	if os.path.exists:
+	    f = open(filePath, 'r')
+	    data = cPickle.load(f)
+	    
+	    for d in data:
+		origTarget = d[0]
+		value = d[1]
+		
+		for shape in blendshapes:
+		    targets = cmds.listAttr(shape + ".w", m = True)
+		    for target in targets:
+			if target == origTarget:
+			    try:
+				cmds.setAttr(shape + "." + target, value)
+			    except:
+				pass
+
+	    f.close()
+	
+    # # # # # # # # # # # # # # # # # # # # # # # # # # # # # # # # # # # # # # # # # # # # # # # # # # # # # # # # # # # # # # # 
+    # # # # # # # # # # # # # # # # # # # # # # # # # # # # # # # # # # # # # # # # # # # # # # # # # # # # # # # # # # # # # # # 
+    # # # # # # # # # # # # # # # # # # # # # # # # # # # # # # # # # # # # # # # # # # # # # # # # # # # # # # # # # # # # # # # 
+    # # # # # # # # # # # # # # # # # # # # # # # # # # # # # # # # # # # # # # # # # # # # # # # # # # # # # # # # # # # # # # # 
+    # # # # # # # # # # # # # # # # # # # # # # # # # # # # # # # # # # # # # # # # # # # # # # # # # # # # # # # # # # # # # # # 
+    def createJointMover_BodyModeUI(self):
+        
+        self.widgets["jointMover_bodyFormLayout"] = cmds.formLayout(w = 400, h = 600, vis = False, parent = self.widgets["masterFormLayout"])
+        
+        #get our image
+        backgroundImage = self.mayaToolsDir + "/General/Icons/ART/jointMover.jpg"
+        self.widgets["jointMover_bodyBackgroundImage"] = cmds.image(image = backgroundImage, w = 400, h = 600)
+        
+        self.createJointMoverToolbar("jointMover_bodyFormLayout", True, False, False)
+        
+        #create the widgets
+        self.widgets["jointMoverUI_clavicle_l"] = cmds.button(w = 15, h = 15, label = "", bgc = [0, .30, 1], c = partial(self.jointMoverSelect, "clavicle", "l", "global"))
+        self.widgets["jointMoverUI_shoulder_l"] = cmds.button(w = 15, h = 15, label = "", bgc = [0, .30, 1], c = partial(self.jointMoverSelect, "upperarm", "l", "global"))
+        self.widgets["jointMoverUI_elbow_l"] = cmds.button(w = 15, h = 15, label = "", bgc = [0, .30, 1], c = partial(self.jointMoverSelect, "lowerarm", "l", "global"))
+        self.widgets["jointMoverUI_wrist_l"] = cmds.button(w = 15, h = 15, label = "", bgc = [0, .30, 1], c = partial(self.jointMoverSelect, "hand", "l", "global"))
+        
+        
+        self.widgets["jointMoverUI_bodyPageThumb01_l"] = cmds.button(w = 10, h = 10, label = "", bgc = [0, .30, 1], c = partial(self.jointMoverSelect, "thumb_01", "l", "global"), vis = False)
+        self.widgets["jointMoverUI_bodyPageThumb02_l"] = cmds.button(w = 10, h = 10, label = "", bgc = [0, .30, 1], c = partial(self.jointMoverSelect, "thumb_02", "l", "global"), vis = False)
+        self.widgets["jointMoverUI_bodyPageThumb03_l"] = cmds.button(w = 10, h = 10, label = "", bgc = [0, .30, 1], c = partial(self.jointMoverSelect, "thumb_03", "l", "global"), vis = False)
+
+
+        self.widgets["jointMoverUI_clavicle_r"] = cmds.button(w = 15, h = 15, label = "", bgc = [0, .60, .1], c = partial(self.jointMoverSelect, "clavicle", "r", "global"))
+        self.widgets["jointMoverUI_shoulder_r"] = cmds.button(w = 15, h = 15, label = "", bgc = [0, .60, .1], c = partial(self.jointMoverSelect, "upperarm", "r", "global"))
+        self.widgets["jointMoverUI_elbow_r"] = cmds.button(w = 15, h = 15, label = "", bgc = [0, .60, .1], c = partial(self.jointMoverSelect, "lowerarm", "r", "global"))
+        self.widgets["jointMoverUI_wrist_r"] = cmds.button(w = 15, h = 15, label = "", bgc = [0, .60, .1], c = partial(self.jointMoverSelect, "hand", "r", "global"))
+        
+        
+        self.widgets["jointMoverUI_bodyPageThumb01_r"] = cmds.button(w = 10, h = 10, label = "", bgc = [0, .60, .1], c = partial(self.jointMoverSelect, "thumb_01", "r", "global"), vis = False)
+        self.widgets["jointMoverUI_bodyPageThumb02_r"] = cmds.button(w = 10, h = 10, label = "", bgc = [0, .60, .1], c = partial(self.jointMoverSelect, "thumb_02", "r", "global"), vis = False)
+        self.widgets["jointMoverUI_bodyPageThumb03_r"] = cmds.button(w = 10, h = 10, label = "", bgc = [0, .60, .1], c = partial(self.jointMoverSelect, "thumb_03", "r", "global"), vis = False)
+
+
+        self.widgets["jointMoverUI_thigh_l"] = cmds.button(w = 15, h = 15, label = "", bgc = [0, .30, 1], c = partial(self.jointMoverSelect, "thigh", "l", "global"))
+        self.widgets["jointMoverUI_calf_l"] = cmds.button(w = 15, h = 15, label = "", bgc = [0, .30, 1], c = partial(self.jointMoverSelect, "calf", "l", "global"))
+        self.widgets["jointMoverUI_foot_l"] = cmds.button(w = 15, h = 15, label = "", bgc = [0, .30, 1], c = partial(self.jointMoverSelect, "foot", "l", "global"))
+        
+        self.widgets["jointMoverUI_thigh_r"] = cmds.button(w = 15, h = 15, label = "", bgc = [0, .60, .1], c = partial(self.jointMoverSelect, "thigh", "r", "global"))
+        self.widgets["jointMoverUI_calf_r"] = cmds.button(w = 15, h = 15, label = "", bgc = [0, .60, .1], c = partial(self.jointMoverSelect, "calf", "r", "global"))
+        self.widgets["jointMoverUI_foot_r"] = cmds.button(w = 15, h = 15, label = "", bgc = [0, .60, .1], c = partial(self.jointMoverSelect, "foot", "r", "global"))
+        
+        self.widgets["jointMoverUI_torsoPage"] = cmds.button(w = 50, h = 80, label = "", bgc = [.92, .93, .14],c = partial(self.switchMode, "jointMover_torsoModeFormLayout", None))
+        self.widgets["jointMoverUI_headPage"] = cmds.button(w = 30, h = 30, label = "", bgc = [.92, .93, .14],c = partial(self.switchMode, "jointMover_torsoModeFormLayout", None))
+        self.widgets["jointMoverUI_leftHandPage"] = cmds.button(w = 15, h = 15, label = "", bgc = [.92, .93, .14],c = partial(self.switchMode, "jointMover_leftHandModeFormLayout", None))
+        self.widgets["jointMoverUI_rightHandPage"] = cmds.button(w = 15, h = 15, label = "", bgc = [.92, .93, .14],c = partial(self.switchMode, "jointMover_rightHandModeFormLayout", None))
+        self.widgets["jointMoverUI_leftFootPage"] = cmds.button(w = 15, h = 15, label = "", bgc = [.92, .93, .14],c = partial(self.switchMode, "jointMover_leftFootModeFormLayout", None))
+        self.widgets["jointMoverUI_rightFootPage"] = cmds.button(w = 15, h = 15, label = "", bgc = [.92, .93, .14], c = partial(self.switchMode, "jointMover_rightFootModeFormLayout", None))
+        
+        self.widgets["jointMoverUI_root"] = cmds.button(w = 100, h = 15, label = "", bgc = [.96, .96, .65], c = partial(self.jointMoverSelect, "root", None, "global"))
+
+        
+        #place the widgets
+        cmds.formLayout(self.widgets["jointMover_bodyFormLayout"], edit = True, af = [(self.widgets["jointMoverUI_clavicle_l"], 'top', 100), (self.widgets["jointMoverUI_clavicle_l"], 'left', 230)])
+        cmds.formLayout(self.widgets["jointMover_bodyFormLayout"], edit = True, af = [(self.widgets["jointMoverUI_shoulder_l"], 'top', 115), (self.widgets["jointMoverUI_shoulder_l"], 'left', 260)])
+        cmds.formLayout(self.widgets["jointMover_bodyFormLayout"], edit = True, af = [(self.widgets["jointMoverUI_elbow_l"], 'top', 165), (self.widgets["jointMoverUI_elbow_l"], 'left', 292)])
+        cmds.formLayout(self.widgets["jointMover_bodyFormLayout"], edit = True, af = [(self.widgets["jointMoverUI_wrist_l"], 'top', 210), (self.widgets["jointMoverUI_wrist_l"], 'left', 318)])
+        
+        cmds.formLayout(self.widgets["jointMover_bodyFormLayout"], edit = True, af = [(self.widgets["jointMoverUI_bodyPageThumb01_l"], 'top', 227), (self.widgets["jointMoverUI_bodyPageThumb01_l"], 'left', 315)])
+        cmds.formLayout(self.widgets["jointMover_bodyFormLayout"], edit = True, af = [(self.widgets["jointMoverUI_bodyPageThumb02_l"], 'top', 241), (self.widgets["jointMoverUI_bodyPageThumb02_l"], 'left', 310)])
+        cmds.formLayout(self.widgets["jointMover_bodyFormLayout"], edit = True, af = [(self.widgets["jointMoverUI_bodyPageThumb03_l"], 'top', 252), (self.widgets["jointMoverUI_bodyPageThumb03_l"], 'left', 305)])
+
+        cmds.formLayout(self.widgets["jointMover_bodyFormLayout"], edit = True, af = [(self.widgets["jointMoverUI_clavicle_r"], 'top', 100), (self.widgets["jointMoverUI_clavicle_r"], 'left', 212)])
+        cmds.formLayout(self.widgets["jointMover_bodyFormLayout"], edit = True, af = [(self.widgets["jointMoverUI_shoulder_r"], 'top', 115), (self.widgets["jointMoverUI_shoulder_r"], 'left', 180)])
+        cmds.formLayout(self.widgets["jointMover_bodyFormLayout"], edit = True, af = [(self.widgets["jointMoverUI_elbow_r"], 'top', 165), (self.widgets["jointMoverUI_elbow_r"], 'left', 150)])
+        cmds.formLayout(self.widgets["jointMover_bodyFormLayout"], edit = True, af = [(self.widgets["jointMoverUI_wrist_r"], 'top', 210), (self.widgets["jointMoverUI_wrist_r"], 'left', 124)])
+        
+        cmds.formLayout(self.widgets["jointMover_bodyFormLayout"], edit = True, af = [(self.widgets["jointMoverUI_bodyPageThumb01_r"], 'top', 227), (self.widgets["jointMoverUI_bodyPageThumb01_r"], 'left', 137)])
+        cmds.formLayout(self.widgets["jointMover_bodyFormLayout"], edit = True, af = [(self.widgets["jointMoverUI_bodyPageThumb02_r"], 'top', 241), (self.widgets["jointMoverUI_bodyPageThumb02_r"], 'left', 140)])
+        cmds.formLayout(self.widgets["jointMover_bodyFormLayout"], edit = True, af = [(self.widgets["jointMoverUI_bodyPageThumb03_r"], 'top', 252), (self.widgets["jointMoverUI_bodyPageThumb03_r"], 'left', 143)])
+        
+        cmds.formLayout(self.widgets["jointMover_bodyFormLayout"], edit = True, af = [(self.widgets["jointMoverUI_thigh_l"], 'top', 238), (self.widgets["jointMoverUI_thigh_l"], 'left', 248)])
+        cmds.formLayout(self.widgets["jointMover_bodyFormLayout"], edit = True, af = [(self.widgets["jointMoverUI_calf_l"], 'top', 320), (self.widgets["jointMoverUI_calf_l"], 'left', 248)])
+        cmds.formLayout(self.widgets["jointMover_bodyFormLayout"], edit = True, af = [(self.widgets["jointMoverUI_foot_l"], 'top', 420), (self.widgets["jointMoverUI_foot_l"], 'left', 248)])
+        
+        cmds.formLayout(self.widgets["jointMover_bodyFormLayout"], edit = True, af = [(self.widgets["jointMoverUI_thigh_r"], 'top', 238), (self.widgets["jointMoverUI_thigh_r"], 'left', 194)])
+        cmds.formLayout(self.widgets["jointMover_bodyFormLayout"], edit = True, af = [(self.widgets["jointMoverUI_calf_r"], 'top', 320), (self.widgets["jointMoverUI_calf_r"], 'left', 194)])
+        cmds.formLayout(self.widgets["jointMover_bodyFormLayout"], edit = True, af = [(self.widgets["jointMoverUI_foot_r"], 'top', 420), (self.widgets["jointMoverUI_foot_r"], 'left', 194)])
+        
+        cmds.formLayout(self.widgets["jointMover_bodyFormLayout"], edit = True, af = [(self.widgets["jointMoverUI_torsoPage"], 'top', 135), (self.widgets["jointMoverUI_torsoPage"], 'left', 204)])
+        cmds.formLayout(self.widgets["jointMover_bodyFormLayout"], edit = True, af = [(self.widgets["jointMoverUI_headPage"], 'top', 35), (self.widgets["jointMoverUI_headPage"], 'left', 214)])
+        cmds.formLayout(self.widgets["jointMover_bodyFormLayout"], edit = True, af = [(self.widgets["jointMoverUI_leftHandPage"], 'top', 240), (self.widgets["jointMoverUI_leftHandPage"], 'left', 340)])
+        cmds.formLayout(self.widgets["jointMover_bodyFormLayout"], edit = True, af = [(self.widgets["jointMoverUI_rightHandPage"], 'top', 240), (self.widgets["jointMoverUI_rightHandPage"], 'left', 106)])
+        cmds.formLayout(self.widgets["jointMover_bodyFormLayout"], edit = True, af = [(self.widgets["jointMoverUI_leftFootPage"], 'top', 447), (self.widgets["jointMoverUI_leftFootPage"], 'left', 250)])
+        cmds.formLayout(self.widgets["jointMover_bodyFormLayout"], edit = True, af = [(self.widgets["jointMoverUI_rightFootPage"], 'top', 447), (self.widgets["jointMoverUI_rightFootPage"], 'left', 192)])
+
+        cmds.formLayout(self.widgets["jointMover_bodyFormLayout"], edit = True, af = [(self.widgets["jointMoverUI_root"], 'top', 467), (self.widgets["jointMoverUI_root"], 'left', 180)])
+	
+	
+
+    # # # # # # # # # # # # # # # # # # # # # # # # # # # # # # # # # # # # # # # # # # # # # # # # # # # # # # # # # # # # # # # 
+    # # # # # # # # # # # # # # # # # # # # # # # # # # # # # # # # # # # # # # # # # # # # # # # # # # # # # # # # # # # # # # # 
+    # # # # # # # # # # # # # # # # # # # # # # # # # # # # # # # # # # # # # # # # # # # # # # # # # # # # # # # # # # # # # # # 
+    # # # # # # # # # # # # # # # # # # # # # # # # # # # # # # # # # # # # # # # # # # # # # # # # # # # # # # # # # # # # # # # 
+    # # # # # # # # # # # # # # # # # # # # # # # # # # # # # # # # # # # # # # # # # # # # # # # # # # # # # # # # # # # # # # # 
+    def createJointMover_LeftHandModeUI(self):
+        
+        #create a new formLayout, parent it to the scrollLayout, and make it invisible
+        self.widgets["jointMover_leftHandModeFormLayout"] = cmds.formLayout(w = 400, h = 600, vis = False, parent = self.widgets["masterFormLayout"])
+        
+        
+        #setup the background image
+        backgroundImage = self.mayaToolsDir + "/General/Icons/ART/jointMover_lefthand.jpg"
+        self.widgets["jointMover_leftHandModeBackgroundImage"] = cmds.image(image = backgroundImage, w = 400, h = 600, parent = self.widgets["jointMover_leftHandModeFormLayout"])
+        
+        
+        #create widgets
+        backButton = cmds.symbolButton(w = 50, h = 40,  image = self.mayaToolsDir + "/General/Icons/ART/backButton.bmp", c = partial(self.switchMode, "jointMover_bodyFormLayout", None))        
+        self.widgets["jointMoverUI_thumb01_l"] = cmds.button(w = 20, h = 20, label = "", bgc = [0, .30, 1],parent = self.widgets["jointMover_leftHandModeFormLayout"], c = partial(self.jointMoverSelect, "thumb_01", "l", "global"))
+        self.widgets["jointMoverUI_thumb02_l"] = cmds.button(w = 20, h = 20, label = "", bgc = [0, .30, 1],parent = self.widgets["jointMover_leftHandModeFormLayout"], c = partial(self.jointMoverSelect, "thumb_02", "l", "global"))
+        self.widgets["jointMoverUI_thumb03_l"] = cmds.button(w = 20, h = 20, label = "", bgc = [0, .30, 1],parent = self.widgets["jointMover_leftHandModeFormLayout"], c = partial(self.jointMoverSelect, "thumb_03", "l", "global"))
+        
+        self.widgets["jointMoverUI_index01_l"] = cmds.button(w = 20, h = 20, label = "", bgc = [0, .30, 1],parent = self.widgets["jointMover_leftHandModeFormLayout"], c = partial(self.jointMoverSelect, "index_metacarpal", "l", "global"))
+        self.widgets["jointMoverUI_index02_l"] = cmds.button(w = 20, h = 20, label = "", bgc = [0, .30, 1],parent = self.widgets["jointMover_leftHandModeFormLayout"], c = partial(self.jointMoverSelect, "index_01", "l", "global"))
+        self.widgets["jointMoverUI_index03_l"] = cmds.button(w = 20, h = 20, label = "", bgc = [0, .30, 1],parent = self.widgets["jointMover_leftHandModeFormLayout"], c = partial(self.jointMoverSelect, "index_02", "l", "global"))
+        self.widgets["jointMoverUI_index04_l"] = cmds.button(w = 20, h = 20, label = "", bgc = [0, .30, 1],parent = self.widgets["jointMover_leftHandModeFormLayout"], c = partial(self.jointMoverSelect, "index_03", "l", "global"))
+        
+        self.widgets["jointMoverUI_middle01_l"] = cmds.button(w = 20, h = 20, label = "", bgc = [0, .30, 1],parent = self.widgets["jointMover_leftHandModeFormLayout"], c = partial(self.jointMoverSelect, "middle_metacarpal", "l", "global"))
+        self.widgets["jointMoverUI_middle02_l"] = cmds.button(w = 20, h = 20, label = "", bgc = [0, .30, 1],parent = self.widgets["jointMover_leftHandModeFormLayout"], c = partial(self.jointMoverSelect, "middle_01", "l", "global"))
+        self.widgets["jointMoverUI_middle03_l"] = cmds.button(w = 20, h = 20, label = "", bgc = [0, .30, 1],parent = self.widgets["jointMover_leftHandModeFormLayout"], c = partial(self.jointMoverSelect, "middle_02", "l", "global"))
+        self.widgets["jointMoverUI_middle04_l"] = cmds.button(w = 20, h = 20, label = "", bgc = [0, .30, 1],parent = self.widgets["jointMover_leftHandModeFormLayout"], c = partial(self.jointMoverSelect, "middle_03", "l", "global"))
+        
+        self.widgets["jointMoverUI_ring01_l"] = cmds.button(w = 20, h = 20, label = "", bgc = [0, .30, 1],parent = self.widgets["jointMover_leftHandModeFormLayout"], c = partial(self.jointMoverSelect, "ring_metacarpal", "l", "global"))
+        self.widgets["jointMoverUI_ring02_l"] = cmds.button(w = 20, h = 20, label = "", bgc = [0, .30, 1],parent = self.widgets["jointMover_leftHandModeFormLayout"], c = partial(self.jointMoverSelect, "ring_01", "l", "global"))
+        self.widgets["jointMoverUI_ring03_l"] = cmds.button(w = 20, h = 20, label = "", bgc = [0, .30, 1],parent = self.widgets["jointMover_leftHandModeFormLayout"], c = partial(self.jointMoverSelect, "ring_02", "l", "global"))
+        self.widgets["jointMoverUI_ring04_l"] = cmds.button(w = 20, h = 20, label = "", bgc = [0, .30, 1],parent = self.widgets["jointMover_leftHandModeFormLayout"], c = partial(self.jointMoverSelect, "ring_03", "l", "global"))
+        
+        self.widgets["jointMoverUI_pinky01_l"] = cmds.button(w = 20, h = 20, label = "", bgc = [0, .30, 1],parent = self.widgets["jointMover_leftHandModeFormLayout"], c = partial(self.jointMoverSelect, "pinky_metacarpal", "l", "global"))
+        self.widgets["jointMoverUI_pinky02_l"] = cmds.button(w = 20, h = 20, label = "", bgc = [0, .30, 1],parent = self.widgets["jointMover_leftHandModeFormLayout"], c = partial(self.jointMoverSelect, "pinky_01", "l", "global"))
+        self.widgets["jointMoverUI_pinky03_l"] = cmds.button(w = 20, h = 20, label = "", bgc = [0, .30, 1],parent = self.widgets["jointMover_leftHandModeFormLayout"], c = partial(self.jointMoverSelect, "pinky_02", "l", "global"))
+        self.widgets["jointMoverUI_pinky04_l"] = cmds.button(w = 20, h = 20, label = "", bgc = [0, .30, 1],parent = self.widgets["jointMover_leftHandModeFormLayout"], c = partial(self.jointMoverSelect, "pinky_03", "l", "global"))
+
+
+
+        #place widgets
+        cmds.formLayout(self.widgets["jointMover_leftHandModeFormLayout"], edit = True, af = [(backButton, 'top', 5), (backButton, 'left', 65)])
+        cmds.formLayout(self.widgets["jointMover_leftHandModeFormLayout"], edit = True, af = [(self.widgets["jointMoverUI_thumb01_l"], 'top', 318), (self.widgets["jointMoverUI_thumb01_l"], 'left', 258)])     
+        cmds.formLayout(self.widgets["jointMover_leftHandModeFormLayout"], edit = True, af = [(self.widgets["jointMoverUI_thumb02_l"], 'top', 280), (self.widgets["jointMoverUI_thumb02_l"], 'left', 298)])
+        cmds.formLayout(self.widgets["jointMover_leftHandModeFormLayout"], edit = True, af = [(self.widgets["jointMoverUI_thumb03_l"], 'top', 250), (self.widgets["jointMoverUI_thumb03_l"], 'left', 326)])
+        
+        cmds.formLayout(self.widgets["jointMover_leftHandModeFormLayout"], edit = True, af = [(self.widgets["jointMoverUI_index01_l"], 'top', 280), (self.widgets["jointMoverUI_index01_l"], 'left', 236)])
+        cmds.formLayout(self.widgets["jointMover_leftHandModeFormLayout"], edit = True, af = [(self.widgets["jointMoverUI_index02_l"], 'top', 180), (self.widgets["jointMoverUI_index02_l"], 'left', 252)]) 
+        cmds.formLayout(self.widgets["jointMover_leftHandModeFormLayout"], edit = True, af = [(self.widgets["jointMoverUI_index03_l"], 'top', 125), (self.widgets["jointMoverUI_index03_l"], 'left', 265)]) 
+        cmds.formLayout(self.widgets["jointMover_leftHandModeFormLayout"], edit = True, af = [(self.widgets["jointMoverUI_index04_l"], 'top', 85), (self.widgets["jointMoverUI_index04_l"], 'left', 271)]) 
+        
+        cmds.formLayout(self.widgets["jointMover_leftHandModeFormLayout"], edit = True, af = [(self.widgets["jointMoverUI_middle01_l"], 'top', 280), (self.widgets["jointMoverUI_middle01_l"], 'left', 212)])
+        cmds.formLayout(self.widgets["jointMover_leftHandModeFormLayout"], edit = True, af = [(self.widgets["jointMoverUI_middle02_l"], 'top', 164), (self.widgets["jointMoverUI_middle02_l"], 'left', 210)]) 
+        cmds.formLayout(self.widgets["jointMover_leftHandModeFormLayout"], edit = True, af = [(self.widgets["jointMoverUI_middle03_l"], 'top', 111), (self.widgets["jointMoverUI_middle03_l"], 'left', 208)]) 
+        cmds.formLayout(self.widgets["jointMover_leftHandModeFormLayout"], edit = True, af = [(self.widgets["jointMoverUI_middle04_l"], 'top', 64), (self.widgets["jointMoverUI_middle04_l"], 'left', 204)])
+        
+        cmds.formLayout(self.widgets["jointMover_leftHandModeFormLayout"], edit = True, af = [(self.widgets["jointMoverUI_ring01_l"], 'top', 280), (self.widgets["jointMoverUI_ring01_l"], 'left', 190)])
+        cmds.formLayout(self.widgets["jointMover_leftHandModeFormLayout"], edit = True, af = [(self.widgets["jointMoverUI_ring02_l"], 'top', 176), (self.widgets["jointMoverUI_ring02_l"], 'left', 172)]) 
+        cmds.formLayout(self.widgets["jointMover_leftHandModeFormLayout"], edit = True, af = [(self.widgets["jointMoverUI_ring03_l"], 'top', 119), (self.widgets["jointMoverUI_ring03_l"], 'left', 163)]) 
+        cmds.formLayout(self.widgets["jointMover_leftHandModeFormLayout"], edit = True, af = [(self.widgets["jointMoverUI_ring04_l"], 'top', 81), (self.widgets["jointMoverUI_ring04_l"], 'left', 162)]) 
+        
+        cmds.formLayout(self.widgets["jointMover_leftHandModeFormLayout"], edit = True, af = [(self.widgets["jointMoverUI_pinky01_l"], 'top', 280), (self.widgets["jointMoverUI_pinky01_l"], 'left', 167)])
+        cmds.formLayout(self.widgets["jointMover_leftHandModeFormLayout"], edit = True, af = [(self.widgets["jointMoverUI_pinky02_l"], 'top', 195), (self.widgets["jointMoverUI_pinky02_l"], 'left', 136)]) 
+        cmds.formLayout(self.widgets["jointMover_leftHandModeFormLayout"], edit = True, af = [(self.widgets["jointMoverUI_pinky03_l"], 'top', 154), (self.widgets["jointMoverUI_pinky03_l"], 'left', 127)]) 
+        cmds.formLayout(self.widgets["jointMover_leftHandModeFormLayout"], edit = True, af = [(self.widgets["jointMoverUI_pinky04_l"], 'top', 120), (self.widgets["jointMoverUI_pinky04_l"], 'left', 121)])
+	
+	
+        
+    # # # # # # # # # # # # # # # # # # # # # # # # # # # # # # # # # # # # # # # # # # # # # # # # # # # # # # # # # # # # # # # 
+    # # # # # # # # # # # # # # # # # # # # # # # # # # # # # # # # # # # # # # # # # # # # # # # # # # # # # # # # # # # # # # # 
+    # # # # # # # # # # # # # # # # # # # # # # # # # # # # # # # # # # # # # # # # # # # # # # # # # # # # # # # # # # # # # # # 
+    # # # # # # # # # # # # # # # # # # # # # # # # # # # # # # # # # # # # # # # # # # # # # # # # # # # # # # # # # # # # # # # 
+    # # # # # # # # # # # # # # # # # # # # # # # # # # # # # # # # # # # # # # # # # # # # # # # # # # # # # # # # # # # # # # # 
+    def createJointMover_RightHandModeUI(self):
+        
+        #create a new formLayout, parent it to the scrollLayout, and make it invisible
+        self.widgets["jointMover_rightHandModeFormLayout"] = cmds.formLayout(w = 400, h = 600, vis = False, parent = self.widgets["masterFormLayout"])
+        
+        
+        #setup the background image
+        backgroundImage = self.mayaToolsDir + "/General/Icons/ART/jointMover_righthand.jpg"
+        self.widgets["jointMover_rightHandModeBackgroundImage"] = cmds.image(image = backgroundImage, w = 400, h = 600, parent = self.widgets["jointMover_rightHandModeFormLayout"])
+        
+        
+        #create widgets
+        backButton = cmds.symbolButton(w = 50, h = 40,  image = self.mayaToolsDir + "/General/Icons/ART/backButton.bmp", c = partial(self.switchMode, "jointMover_bodyFormLayout", None))        
+        self.widgets["jointMoverUI_thumb01_r"] = cmds.button(w = 20, h = 20, label = "", bgc = [0, .60, .1],parent = self.widgets["jointMover_rightHandModeFormLayout"], c = partial(self.jointMoverSelect, "thumb_01", "r", "global"))
+        self.widgets["jointMoverUI_thumb02_r"] = cmds.button(w = 20, h = 20, label = "", bgc = [0, .60, .1],parent = self.widgets["jointMover_rightHandModeFormLayout"], c = partial(self.jointMoverSelect, "thumb_02", "r", "global"))
+        self.widgets["jointMoverUI_thumb03_r"] = cmds.button(w = 20, h = 20, label = "", bgc = [0, .60, .1],parent = self.widgets["jointMover_rightHandModeFormLayout"], c = partial(self.jointMoverSelect, "thumb_03", "r", "global"))
+        
+        self.widgets["jointMoverUI_index01_r"] = cmds.button(w = 20, h = 20, label = "", bgc = [0, .60, .1],parent = self.widgets["jointMover_rightHandModeFormLayout"], c = partial(self.jointMoverSelect, "index_metacarpal", "r", "global"))
+        self.widgets["jointMoverUI_index02_r"] = cmds.button(w = 20, h = 20, label = "", bgc = [0, .60, .1],parent = self.widgets["jointMover_rightHandModeFormLayout"], c = partial(self.jointMoverSelect, "index_01", "r", "global"))
+        self.widgets["jointMoverUI_index03_r"] = cmds.button(w = 20, h = 20, label = "", bgc = [0, .60, .1],parent = self.widgets["jointMover_rightHandModeFormLayout"], c = partial(self.jointMoverSelect, "index_02", "r", "global"))
+        self.widgets["jointMoverUI_index04_r"] = cmds.button(w = 20, h = 20, label = "", bgc = [0, .60, .1],parent = self.widgets["jointMover_rightHandModeFormLayout"], c = partial(self.jointMoverSelect, "index_03", "r", "global"))
+        
+        self.widgets["jointMoverUI_middle01_r"] = cmds.button(w = 20, h = 20, label = "", bgc = [0, .60, .1],parent = self.widgets["jointMover_rightHandModeFormLayout"], c = partial(self.jointMoverSelect, "middle_metacarpal", "r", "global"))
+        self.widgets["jointMoverUI_middle02_r"] = cmds.button(w = 20, h = 20, label = "", bgc = [0, .60, .1],parent = self.widgets["jointMover_rightHandModeFormLayout"], c = partial(self.jointMoverSelect, "middle_01", "r", "global"))
+        self.widgets["jointMoverUI_middle03_r"] = cmds.button(w = 20, h = 20, label = "", bgc = [0, .60, .1],parent = self.widgets["jointMover_rightHandModeFormLayout"], c = partial(self.jointMoverSelect, "middle_02", "r", "global"))
+        self.widgets["jointMoverUI_middle04_r"] = cmds.button(w = 20, h = 20, label = "", bgc = [0, .60, .1],parent = self.widgets["jointMover_rightHandModeFormLayout"], c = partial(self.jointMoverSelect, "middle_03", "r", "global"))
+        
+        self.widgets["jointMoverUI_ring01_r"] = cmds.button(w = 20, h = 20, label = "", bgc = [0, .60, .1],parent = self.widgets["jointMover_rightHandModeFormLayout"], c = partial(self.jointMoverSelect, "ring_metacarpal", "r", "global"))
+        self.widgets["jointMoverUI_ring02_r"] = cmds.button(w = 20, h = 20, label = "", bgc = [0, .60, .1],parent = self.widgets["jointMover_rightHandModeFormLayout"], c = partial(self.jointMoverSelect, "ring_01", "r", "global"))
+        self.widgets["jointMoverUI_ring03_r"] = cmds.button(w = 20, h = 20, label = "", bgc = [0, .60, .1],parent = self.widgets["jointMover_rightHandModeFormLayout"], c = partial(self.jointMoverSelect, "ring_02", "r", "global"))
+        self.widgets["jointMoverUI_ring04_r"] = cmds.button(w = 20, h = 20, label = "", bgc = [0, .60, .1],parent = self.widgets["jointMover_rightHandModeFormLayout"], c = partial(self.jointMoverSelect, "ring_03", "r", "global"))
+        
+        self.widgets["jointMoverUI_pinky01_r"] = cmds.button(w = 20, h = 20, label = "", bgc = [0, .60, .1],parent = self.widgets["jointMover_rightHandModeFormLayout"], c = partial(self.jointMoverSelect, "pinky_metacarpal", "r", "global"))
+        self.widgets["jointMoverUI_pinky02_r"] = cmds.button(w = 20, h = 20, label = "", bgc = [0, .60, .1],parent = self.widgets["jointMover_rightHandModeFormLayout"], c = partial(self.jointMoverSelect, "pinky_01", "r", "global"))
+        self.widgets["jointMoverUI_pinky03_r"] = cmds.button(w = 20, h = 20, label = "", bgc = [0, .60, .1],parent = self.widgets["jointMover_rightHandModeFormLayout"], c = partial(self.jointMoverSelect, "pinky_02", "r", "global"))
+        self.widgets["jointMoverUI_pinky04_r"] = cmds.button(w = 20, h = 20, label = "", bgc = [0, .60, .1],parent = self.widgets["jointMover_rightHandModeFormLayout"], c = partial(self.jointMoverSelect, "pinky_03", "r", "global"))
+
+
+
+        #place widgets
+        cmds.formLayout(self.widgets["jointMover_rightHandModeFormLayout"], edit = True, af = [(backButton, 'top', 5), (backButton, 'left', 65)])
+        cmds.formLayout(self.widgets["jointMover_rightHandModeFormLayout"], edit = True, af = [(self.widgets["jointMoverUI_thumb01_r"], 'top', 318), (self.widgets["jointMoverUI_thumb01_r"], 'left', 179)])     
+        cmds.formLayout(self.widgets["jointMover_rightHandModeFormLayout"], edit = True, af = [(self.widgets["jointMoverUI_thumb02_r"], 'top', 280), (self.widgets["jointMoverUI_thumb02_r"], 'left', 138)])
+        cmds.formLayout(self.widgets["jointMover_rightHandModeFormLayout"], edit = True, af = [(self.widgets["jointMoverUI_thumb03_r"], 'top', 250), (self.widgets["jointMoverUI_thumb03_r"], 'left', 111)])
+        
+        cmds.formLayout(self.widgets["jointMover_rightHandModeFormLayout"], edit = True, af = [(self.widgets["jointMoverUI_index01_r"], 'top', 280), (self.widgets["jointMoverUI_index01_r"], 'left', 202)])
+        cmds.formLayout(self.widgets["jointMover_rightHandModeFormLayout"], edit = True, af = [(self.widgets["jointMoverUI_index02_r"], 'top', 180), (self.widgets["jointMoverUI_index02_r"], 'left', 184)]) 
+        cmds.formLayout(self.widgets["jointMover_rightHandModeFormLayout"], edit = True, af = [(self.widgets["jointMoverUI_index03_r"], 'top', 125), (self.widgets["jointMoverUI_index03_r"], 'left', 173)]) 
+        cmds.formLayout(self.widgets["jointMover_rightHandModeFormLayout"], edit = True, af = [(self.widgets["jointMoverUI_index04_r"], 'top', 85), (self.widgets["jointMoverUI_index04_r"], 'left', 167)]) 
+        
+        cmds.formLayout(self.widgets["jointMover_rightHandModeFormLayout"], edit = True, af = [(self.widgets["jointMoverUI_middle01_r"], 'top', 280), (self.widgets["jointMoverUI_middle01_r"], 'left', 225)])
+        cmds.formLayout(self.widgets["jointMover_rightHandModeFormLayout"], edit = True, af = [(self.widgets["jointMoverUI_middle02_r"], 'top', 164), (self.widgets["jointMoverUI_middle02_r"], 'left', 226)]) 
+        cmds.formLayout(self.widgets["jointMover_rightHandModeFormLayout"], edit = True, af = [(self.widgets["jointMoverUI_middle03_r"], 'top', 111), (self.widgets["jointMoverUI_middle03_r"], 'left', 229)]) 
+        cmds.formLayout(self.widgets["jointMover_rightHandModeFormLayout"], edit = True, af = [(self.widgets["jointMoverUI_middle04_r"], 'top', 64), (self.widgets["jointMoverUI_middle04_r"], 'left', 232)])
+        
+        cmds.formLayout(self.widgets["jointMover_rightHandModeFormLayout"], edit = True, af = [(self.widgets["jointMoverUI_ring01_r"], 'top', 280), (self.widgets["jointMoverUI_ring01_r"], 'left', 247)])
+        cmds.formLayout(self.widgets["jointMover_rightHandModeFormLayout"], edit = True, af = [(self.widgets["jointMoverUI_ring02_r"], 'top', 176), (self.widgets["jointMoverUI_ring02_r"], 'left', 265)]) 
+        cmds.formLayout(self.widgets["jointMover_rightHandModeFormLayout"], edit = True, af = [(self.widgets["jointMoverUI_ring03_r"], 'top', 119), (self.widgets["jointMoverUI_ring03_r"], 'left', 273)]) 
+        cmds.formLayout(self.widgets["jointMover_rightHandModeFormLayout"], edit = True, af = [(self.widgets["jointMoverUI_ring04_r"], 'top', 81), (self.widgets["jointMoverUI_ring04_r"], 'left', 275)]) 
+        
+        cmds.formLayout(self.widgets["jointMover_rightHandModeFormLayout"], edit = True, af = [(self.widgets["jointMoverUI_pinky01_r"], 'top', 280), (self.widgets["jointMoverUI_pinky01_r"], 'left', 271)])
+        cmds.formLayout(self.widgets["jointMover_rightHandModeFormLayout"], edit = True, af = [(self.widgets["jointMoverUI_pinky02_r"], 'top', 195), (self.widgets["jointMoverUI_pinky02_r"], 'left', 300)]) 
+        cmds.formLayout(self.widgets["jointMover_rightHandModeFormLayout"], edit = True, af = [(self.widgets["jointMoverUI_pinky03_r"], 'top', 154), (self.widgets["jointMoverUI_pinky03_r"], 'left', 310)]) 
+        cmds.formLayout(self.widgets["jointMover_rightHandModeFormLayout"], edit = True, af = [(self.widgets["jointMoverUI_pinky04_r"], 'top', 120), (self.widgets["jointMoverUI_pinky04_r"], 'left', 317)])
+        
+  
+  
+    # # # # # # # # # # # # # # # # # # # # # # # # # # # # # # # # # # # # # # # # # # # # # # # # # # # # # # # # # # # # # # # 
+    # # # # # # # # # # # # # # # # # # # # # # # # # # # # # # # # # # # # # # # # # # # # # # # # # # # # # # # # # # # # # # # 
+    # # # # # # # # # # # # # # # # # # # # # # # # # # # # # # # # # # # # # # # # # # # # # # # # # # # # # # # # # # # # # # # 
+    # # # # # # # # # # # # # # # # # # # # # # # # # # # # # # # # # # # # # # # # # # # # # # # # # # # # # # # # # # # # # # # 
+    # # # # # # # # # # # # # # # # # # # # # # # # # # # # # # # # # # # # # # # # # # # # # # # # # # # # # # # # # # # # # # # 
+    def createJointMover_TorsoModeUI(self):
+        
+        #create a new formLayout, parent it to the scrollLayout, and make it invisible
+        self.widgets["jointMover_torsoModeFormLayout"] = cmds.formLayout(w = 400, h = 600, vis = False, parent = self.widgets["masterFormLayout"])
+        
+        
+        #setup the background image
+        backgroundImage = self.mayaToolsDir + "/General/Icons/ART/skeletonbuilder_torso.jpg"
+        self.widgets["jointMover_torsoModeBackgroundImage"] = cmds.image(image = backgroundImage, w = 400, h = 600, parent = self.widgets["jointMover_torsoModeFormLayout"])
+        
+        #create widgets
+        backButton = cmds.symbolButton(w = 50, h = 40,  image = self.mayaToolsDir + "/General/Icons/ART/backButton.bmp", c = partial(self.switchMode, "jointMover_bodyFormLayout", None))
+        self.widgets["jointMoverUI_pelvis"] = cmds.button(w = 50, h = 50, label = "", bgc = [.53, .56, .62], parent = self.widgets["jointMover_torsoModeFormLayout"], c = partial(self.jointMoverSelect, "pelvis", None, "global"))
+        self.widgets["jointMoverUI_head"] = cmds.button(w = 50, h = 50, label = "", bgc = [.53, .56, .62], parent = self.widgets["jointMover_torsoModeFormLayout"], c = partial(self.jointMoverSelect, "head", None, "global"))
+
+
+        #place widgets
+        cmds.formLayout(self.widgets["jointMover_torsoModeFormLayout"], edit = True, af = [(backButton, 'top', 5), (backButton, 'left', 65)])
+        cmds.formLayout(self.widgets["jointMover_torsoModeFormLayout"], edit = True, af = [(self.widgets["jointMoverUI_pelvis"], 'top', 330), (self.widgets["jointMoverUI_pelvis"], 'left', 200)])
+        cmds.formLayout(self.widgets["jointMover_torsoModeFormLayout"], edit = True, af = [(self.widgets["jointMoverUI_head"], 'top', 30), (self.widgets["jointMoverUI_head"], 'left', 200)])
+
+
+    # # # # # # # # # # # # # # # # # # # # # # # # # # # # # # # # # # # # # # # # # # # # # # # # # # # # # # # # # # # # # # # 
+    # # # # # # # # # # # # # # # # # # # # # # # # # # # # # # # # # # # # # # # # # # # # # # # # # # # # # # # # # # # # # # # 
+    # # # # # # # # # # # # # # # # # # # # # # # # # # # # # # # # # # # # # # # # # # # # # # # # # # # # # # # # # # # # # # # 
+    # # # # # # # # # # # # # # # # # # # # # # # # # # # # # # # # # # # # # # # # # # # # # # # # # # # # # # # # # # # # # # # 
+    # # # # # # # # # # # # # # # # # # # # # # # # # # # # # # # # # # # # # # # # # # # # # # # # # # # # # # # # # # # # # # # 
+    def createJointMover_LeftFootModeUI(self):
+        
+        #create a new formLayout, parent it to the scrollLayout, and make it invisible
+        self.widgets["jointMover_leftFootModeFormLayout"] = cmds.formLayout(w = 400, h = 600, vis = False, parent = self.widgets["masterFormLayout"])
+        
+        
+        #setup the background image
+        backgroundImage = self.mayaToolsDir + "/General/Icons/ART/jointmover_leftfoot.jpg"
+        self.widgets["jointMover_leftFootModeBackgroundImage"] = cmds.image(image = backgroundImage, w = 400, h = 600, parent = self.widgets["jointMover_leftFootModeFormLayout"])
+        
+        #create widgets
+        backButton = cmds.symbolButton(w = 50, h = 40,  image = self.mayaToolsDir + "/General/Icons/ART/backButton.bmp", c = partial(self.switchMode, "jointMover_bodyFormLayout", None))
+        self.widgets["jointMoverUI_ball_l"] = cmds.button(w = 40, h = 40, label = "", bgc = [0, .30, 1], parent = self.widgets["jointMover_leftFootModeFormLayout"], c = partial(self.jointMoverSelect, "ball", "l", "global"))
+        self.widgets["jointMoverUI_bigtoe01_l"] = cmds.button(w = 15, h = 15, label = "", bgc = [0, .30, 1], parent = self.widgets["jointMover_leftFootModeFormLayout"], c = partial(self.jointMoverSelect, "bigtoe_metatarsal", "l", "global"))
+        self.widgets["jointMoverUI_bigtoe02_l"] = cmds.button(w = 15, h = 15, label = "", bgc = [0, .30, 1], parent = self.widgets["jointMover_leftFootModeFormLayout"], c = partial(self.jointMoverSelect, "bigtoe_proximal_phalange", "l", "global"))
+        self.widgets["jointMoverUI_bigtoe03_l"] = cmds.button(w = 15, h = 15, label = "", bgc = [0, .30, 1], parent = self.widgets["jointMover_leftFootModeFormLayout"], c = partial(self.jointMoverSelect, "bigtoe_distal_phalange", "l", "global"))
+
+        self.widgets["jointMoverUI_indextoe01_l"] = cmds.button(w = 15, h = 15, label = "", bgc = [0, .30, 1], parent = self.widgets["jointMover_leftFootModeFormLayout"], c = partial(self.jointMoverSelect, "index_metatarsal", "l", "global"))
+        self.widgets["jointMoverUI_indextoe02_l"] = cmds.button(w = 15, h = 15, label = "", bgc = [0, .30, 1], parent = self.widgets["jointMover_leftFootModeFormLayout"], c = partial(self.jointMoverSelect, "index_proximal_phalange", "l", "global"))
+        self.widgets["jointMoverUI_indextoe03_l"] = cmds.button(w = 15, h = 15, label = "", bgc = [0, .30, 1], parent = self.widgets["jointMover_leftFootModeFormLayout"], c = partial(self.jointMoverSelect, "index_middle_phalange", "l", "global"))
+        self.widgets["jointMoverUI_indextoe04_l"] = cmds.button(w = 15, h = 15, label = "", bgc = [0, .30, 1], parent = self.widgets["jointMover_leftFootModeFormLayout"], c = partial(self.jointMoverSelect, "index_distal_phalange", "l", "global"))
+
+        self.widgets["jointMoverUI_middletoe01_l"] = cmds.button(w = 15, h = 15, label = "", bgc = [0, .30, 1], parent = self.widgets["jointMover_leftFootModeFormLayout"], c = partial(self.jointMoverSelect, "middle_metatarsal", "l", "global"))
+        self.widgets["jointMoverUI_middletoe02_l"] = cmds.button(w = 15, h = 15, label = "", bgc = [0, .30, 1], parent = self.widgets["jointMover_leftFootModeFormLayout"], c = partial(self.jointMoverSelect, "middle_proximal_phalange", "l", "global"))
+        self.widgets["jointMoverUI_middletoe03_l"] = cmds.button(w = 15, h = 15, label = "", bgc = [0, .30, 1], parent = self.widgets["jointMover_leftFootModeFormLayout"], c = partial(self.jointMoverSelect, "middle_middle_phalange", "l", "global"))
+        self.widgets["jointMoverUI_middletoe04_l"] = cmds.button(w = 15, h = 15, label = "", bgc = [0, .30, 1], parent = self.widgets["jointMover_leftFootModeFormLayout"], c = partial(self.jointMoverSelect, "middle_distal_phalange", "l", "global"))
+
+        self.widgets["jointMoverUI_ringtoe01_l"] = cmds.button(w = 15, h = 15, label = "", bgc = [0, .30, 1], parent = self.widgets["jointMover_leftFootModeFormLayout"], c = partial(self.jointMoverSelect, "ring_metatarsal", "l", "global"))
+        self.widgets["jointMoverUI_ringtoe02_l"] = cmds.button(w = 15, h = 15, label = "", bgc = [0, .30, 1], parent = self.widgets["jointMover_leftFootModeFormLayout"], c = partial(self.jointMoverSelect, "ring_proximal_phalange", "l", "global"))
+        self.widgets["jointMoverUI_ringtoe03_l"] = cmds.button(w = 15, h = 15, label = "", bgc = [0, .30, 1], parent = self.widgets["jointMover_leftFootModeFormLayout"], c = partial(self.jointMoverSelect, "ring_middle_phalange", "l", "global"))
+        self.widgets["jointMoverUI_ringtoe04_l"] = cmds.button(w = 15, h = 15, label = "", bgc = [0, .30, 1], parent = self.widgets["jointMover_leftFootModeFormLayout"], c = partial(self.jointMoverSelect, "ring_distal_phalange", "l", "global"))
+        
+        self.widgets["jointMoverUI_pinkytoe01_l"] = cmds.button(w = 15, h = 15, label = "", bgc = [0, .30, 1], parent = self.widgets["jointMover_leftFootModeFormLayout"], c = partial(self.jointMoverSelect, "pinky_metatarsal", "l", "global"))
+        self.widgets["jointMoverUI_pinkytoe02_l"] = cmds.button(w = 15, h = 15, label = "", bgc = [0, .30, 1], parent = self.widgets["jointMover_leftFootModeFormLayout"], c = partial(self.jointMoverSelect, "pinky_proximal_phalange", "l", "global"))
+        self.widgets["jointMoverUI_pinkytoe03_l"] = cmds.button(w = 15, h = 15, label = "", bgc = [0, .30, 1], parent = self.widgets["jointMover_leftFootModeFormLayout"], c = partial(self.jointMoverSelect, "pinky_middle_phalange", "l", "global"))
+        self.widgets["jointMoverUI_pinkytoe04_l"] = cmds.button(w = 15, h = 15, label = "", bgc = [0, .30, 1], parent = self.widgets["jointMover_leftFootModeFormLayout"], c = partial(self.jointMoverSelect, "pinky_distal_phalange", "l", "global"))
+        
+        
+        #place widgets
+        cmds.formLayout(self.widgets["jointMover_leftFootModeFormLayout"], edit = True, af = [(backButton, 'top', 5), (backButton, 'left', 65)])
+        cmds.formLayout(self.widgets["jointMover_leftFootModeFormLayout"], edit = True, af = [(self.widgets["jointMoverUI_ball_l"], 'top', 233), (self.widgets["jointMoverUI_ball_l"], 'left', 207)])
+        cmds.formLayout(self.widgets["jointMover_leftFootModeFormLayout"], edit = True, af = [(self.widgets["jointMoverUI_bigtoe01_l"], 'top', 154), (self.widgets["jointMoverUI_bigtoe01_l"], 'left', 270)])
+        cmds.formLayout(self.widgets["jointMover_leftFootModeFormLayout"], edit = True, af = [(self.widgets["jointMoverUI_bigtoe02_l"], 'top', 99), (self.widgets["jointMoverUI_bigtoe02_l"], 'left', 272)])
+        cmds.formLayout(self.widgets["jointMover_leftFootModeFormLayout"], edit = True, af = [(self.widgets["jointMoverUI_bigtoe03_l"], 'top', 63), (self.widgets["jointMoverUI_bigtoe03_l"], 'left', 267)])
+
+        cmds.formLayout(self.widgets["jointMover_leftFootModeFormLayout"], edit = True, af = [(self.widgets["jointMoverUI_indextoe01_l"], 'top', 175), (self.widgets["jointMoverUI_indextoe01_l"], 'left', 222)])
+        cmds.formLayout(self.widgets["jointMover_leftFootModeFormLayout"], edit = True, af = [(self.widgets["jointMoverUI_indextoe02_l"], 'top', 109), (self.widgets["jointMoverUI_indextoe02_l"], 'left', 218)])
+        cmds.formLayout(self.widgets["jointMover_leftFootModeFormLayout"], edit = True, af = [(self.widgets["jointMoverUI_indextoe03_l"], 'top', 81), (self.widgets["jointMoverUI_indextoe03_l"], 'left', 212)])
+        cmds.formLayout(self.widgets["jointMover_leftFootModeFormLayout"], edit = True, af = [(self.widgets["jointMoverUI_indextoe04_l"], 'top', 56), (self.widgets["jointMoverUI_indextoe04_l"], 'left', 206)])
+
+        cmds.formLayout(self.widgets["jointMover_leftFootModeFormLayout"], edit = True, af = [(self.widgets["jointMoverUI_middletoe01_l"], 'top', 185), (self.widgets["jointMoverUI_middletoe01_l"], 'left', 198)])
+        cmds.formLayout(self.widgets["jointMover_leftFootModeFormLayout"], edit = True, af = [(self.widgets["jointMoverUI_middletoe02_l"], 'top', 125), (self.widgets["jointMoverUI_middletoe02_l"], 'left', 190)])
+        cmds.formLayout(self.widgets["jointMover_leftFootModeFormLayout"], edit = True, af = [(self.widgets["jointMoverUI_middletoe03_l"], 'top', 103), (self.widgets["jointMoverUI_middletoe03_l"], 'left', 184)])
+        cmds.formLayout(self.widgets["jointMover_leftFootModeFormLayout"], edit = True, af = [(self.widgets["jointMoverUI_middletoe04_l"], 'top', 76), (self.widgets["jointMoverUI_middletoe04_l"], 'left', 178)])
+
+        cmds.formLayout(self.widgets["jointMover_leftFootModeFormLayout"], edit = True, af = [(self.widgets["jointMoverUI_ringtoe01_l"], 'top', 192), (self.widgets["jointMoverUI_ringtoe01_l"], 'left', 174)])
+        cmds.formLayout(self.widgets["jointMover_leftFootModeFormLayout"], edit = True, af = [(self.widgets["jointMoverUI_ringtoe02_l"], 'top', 137), (self.widgets["jointMoverUI_ringtoe02_l"], 'left', 166)])
+        cmds.formLayout(self.widgets["jointMover_leftFootModeFormLayout"], edit = True, af = [(self.widgets["jointMoverUI_ringtoe03_l"], 'top', 117), (self.widgets["jointMoverUI_ringtoe03_l"], 'left', 160)])
+        cmds.formLayout(self.widgets["jointMover_leftFootModeFormLayout"], edit = True, af = [(self.widgets["jointMoverUI_ringtoe04_l"], 'top', 97), (self.widgets["jointMoverUI_ringtoe04_l"], 'left', 154)])
+        
+        cmds.formLayout(self.widgets["jointMover_leftFootModeFormLayout"], edit = True, af = [(self.widgets["jointMoverUI_pinkytoe01_l"], 'top', 203), (self.widgets["jointMoverUI_pinkytoe01_l"], 'left', 151)])
+        cmds.formLayout(self.widgets["jointMover_leftFootModeFormLayout"], edit = True, af = [(self.widgets["jointMoverUI_pinkytoe02_l"], 'top', 165), (self.widgets["jointMoverUI_pinkytoe02_l"], 'left', 142)])
+        cmds.formLayout(self.widgets["jointMover_leftFootModeFormLayout"], edit = True, af = [(self.widgets["jointMoverUI_pinkytoe03_l"], 'top', 146), (self.widgets["jointMoverUI_pinkytoe03_l"], 'left', 140)])
+        cmds.formLayout(self.widgets["jointMover_leftFootModeFormLayout"], edit = True, af = [(self.widgets["jointMoverUI_pinkytoe04_l"], 'top', 127), (self.widgets["jointMoverUI_pinkytoe04_l"], 'left', 137)])
+        
+        
+        
+    # # # # # # # # # # # # # # # # # # # # # # # # # # # # # # # # # # # # # # # # # # # # # # # # # # # # # # # # # # # # # # # 
+    # # # # # # # # # # # # # # # # # # # # # # # # # # # # # # # # # # # # # # # # # # # # # # # # # # # # # # # # # # # # # # # 
+    # # # # # # # # # # # # # # # # # # # # # # # # # # # # # # # # # # # # # # # # # # # # # # # # # # # # # # # # # # # # # # # 
+    # # # # # # # # # # # # # # # # # # # # # # # # # # # # # # # # # # # # # # # # # # # # # # # # # # # # # # # # # # # # # # # 
+    # # # # # # # # # # # # # # # # # # # # # # # # # # # # # # # # # # # # # # # # # # # # # # # # # # # # # # # # # # # # # # # 
+    def createJointMover_RightFootModeUI(self):
+        
+        #create a new formLayout, parent it to the scrollLayout, and make it invisible
+        self.widgets["jointMover_rightFootModeFormLayout"] = cmds.formLayout(w = 400, h = 600, vis = False, parent = self.widgets["masterFormLayout"])
+        
+        
+        #setup the background image
+        backgroundImage = self.mayaToolsDir + "/General/Icons/ART/jointmover_rightfoot.jpg"
+        self.widgets["jointMover_rightFootModeBackgroundImage"] = cmds.image(image = backgroundImage, w = 400, h = 600, parent = self.widgets["jointMover_rightFootModeFormLayout"])
+        
+        #create widgets
+        backButton = cmds.symbolButton(w = 50, h = 40,  image = self.mayaToolsDir + "/General/Icons/ART/backButton.bmp", c = partial(self.switchMode, "jointMover_bodyFormLayout", None))
+        self.widgets["jointMoverUI_ball_r"] = cmds.button(w = 40, h = 40, label = "", bgc = [0, .60, .1], parent = self.widgets["jointMover_rightFootModeFormLayout"], c = partial(self.jointMoverSelect, "ball", "r", "global"))
+        self.widgets["jointMoverUI_bigtoe01_r"] = cmds.button(w = 15, h = 15, label = "", bgc = [0, .60, .1], parent = self.widgets["jointMover_rightFootModeFormLayout"], c = partial(self.jointMoverSelect, "bigtoe_metatarsal", "r", "global"))
+        self.widgets["jointMoverUI_bigtoe02_r"] = cmds.button(w = 15, h = 15, label = "", bgc = [0, .60, .1], parent = self.widgets["jointMover_rightFootModeFormLayout"], c = partial(self.jointMoverSelect, "bigtoe_proximal_phalange", "r", "global"))
+        self.widgets["jointMoverUI_bigtoe03_r"] = cmds.button(w = 15, h = 15, label = "", bgc = [0, .60, .1], parent = self.widgets["jointMover_rightFootModeFormLayout"], c = partial(self.jointMoverSelect, "bigtoe_distal_phalange", "r", "global"))
+
+        self.widgets["jointMoverUI_indextoe01_r"] = cmds.button(w = 15, h = 15, label = "", bgc = [0, .60, .1], parent = self.widgets["jointMover_rightFootModeFormLayout"], c = partial(self.jointMoverSelect, "index_metatarsal", "r", "global"))
+        self.widgets["jointMoverUI_indextoe02_r"] = cmds.button(w = 15, h = 15, label = "", bgc = [0, .60, .1], parent = self.widgets["jointMover_rightFootModeFormLayout"], c = partial(self.jointMoverSelect, "index_proximal_phalange", "r", "global"))
+        self.widgets["jointMoverUI_indextoe03_r"] = cmds.button(w = 15, h = 15, label = "", bgc = [0, .60, .1], parent = self.widgets["jointMover_rightFootModeFormLayout"], c = partial(self.jointMoverSelect, "index_middle_phalange", "r", "global"))
+        self.widgets["jointMoverUI_indextoe04_r"] = cmds.button(w = 15, h = 15, label = "", bgc = [0, .60, .1], parent = self.widgets["jointMover_rightFootModeFormLayout"], c = partial(self.jointMoverSelect, "index_distal_phalange", "r", "global"))
+
+        self.widgets["jointMoverUI_middletoe01_r"] = cmds.button(w = 15, h = 15, label = "", bgc = [0, .60, .1], parent = self.widgets["jointMover_rightFootModeFormLayout"], c = partial(self.jointMoverSelect, "middle_metatarsal", "r", "global"))
+        self.widgets["jointMoverUI_middletoe02_r"] = cmds.button(w = 15, h = 15, label = "", bgc = [0, .60, .1], parent = self.widgets["jointMover_rightFootModeFormLayout"], c = partial(self.jointMoverSelect, "middle_proximal_phalange", "r", "global"))
+        self.widgets["jointMoverUI_middletoe03_r"] = cmds.button(w = 15, h = 15, label = "", bgc = [0, .60, .1], parent = self.widgets["jointMover_rightFootModeFormLayout"], c = partial(self.jointMoverSelect, "middle_middle_phalange", "r", "global"))
+        self.widgets["jointMoverUI_middletoe04_r"] = cmds.button(w = 15, h = 15, label = "", bgc = [0, .60, .1], parent = self.widgets["jointMover_rightFootModeFormLayout"], c = partial(self.jointMoverSelect, "middle_distal_phalange", "r", "global"))
+
+        self.widgets["jointMoverUI_ringtoe01_r"] = cmds.button(w = 15, h = 15, label = "", bgc = [0, .60, .1], parent = self.widgets["jointMover_rightFootModeFormLayout"], c = partial(self.jointMoverSelect, "ring_metatarsal", "r", "global"))
+        self.widgets["jointMoverUI_ringtoe02_r"] = cmds.button(w = 15, h = 15, label = "", bgc = [0, .60, .1], parent = self.widgets["jointMover_rightFootModeFormLayout"], c = partial(self.jointMoverSelect, "ring_proximal_phalange", "r", "global"))
+        self.widgets["jointMoverUI_ringtoe03_r"] = cmds.button(w = 15, h = 15, label = "", bgc = [0, .60, .1], parent = self.widgets["jointMover_rightFootModeFormLayout"], c = partial(self.jointMoverSelect, "ring_middle_phalange", "r", "global"))
+        self.widgets["jointMoverUI_ringtoe04_r"] = cmds.button(w = 15, h = 15, label = "", bgc = [0, .60, .1], parent = self.widgets["jointMover_rightFootModeFormLayout"], c = partial(self.jointMoverSelect, "ring_distal_phalange", "r", "global"))
+        
+        self.widgets["jointMoverUI_pinkytoe01_r"] = cmds.button(w = 15, h = 15, label = "", bgc = [0, .60, .1], parent = self.widgets["jointMover_rightFootModeFormLayout"], c = partial(self.jointMoverSelect, "pinky_metatarsal", "r", "global"))
+        self.widgets["jointMoverUI_pinkytoe02_r"] = cmds.button(w = 15, h = 15, label = "", bgc = [0, .60, .1], parent = self.widgets["jointMover_rightFootModeFormLayout"], c = partial(self.jointMoverSelect, "pinky_proximal_phalange", "r", "global"))
+        self.widgets["jointMoverUI_pinkytoe03_r"] = cmds.button(w = 15, h = 15, label = "", bgc = [0, .60, .1], parent = self.widgets["jointMover_rightFootModeFormLayout"], c = partial(self.jointMoverSelect, "pinky_middle_phalange", "r", "global"))
+        self.widgets["jointMoverUI_pinkytoe04_r"] = cmds.button(w = 15, h = 15, label = "", bgc = [0, .60, .1], parent = self.widgets["jointMover_rightFootModeFormLayout"], c = partial(self.jointMoverSelect, "pinky_distal_phalange", "r", "global"))
+        
+        
+        #place widgets
+        cmds.formLayout(self.widgets["jointMover_rightFootModeFormLayout"], edit = True, af = [(backButton, 'top', 5), (backButton, 'left', 65)])
+        cmds.formLayout(self.widgets["jointMover_rightFootModeFormLayout"], edit = True, af = [(self.widgets["jointMoverUI_ball_r"], 'top', 233), (self.widgets["jointMoverUI_ball_r"], 'left', 200)])
+        cmds.formLayout(self.widgets["jointMover_rightFootModeFormLayout"], edit = True, af = [(self.widgets["jointMoverUI_bigtoe01_r"], 'top', 154), (self.widgets["jointMoverUI_bigtoe01_r"], 'left', 161)])
+        cmds.formLayout(self.widgets["jointMover_rightFootModeFormLayout"], edit = True, af = [(self.widgets["jointMoverUI_bigtoe02_r"], 'top', 99), (self.widgets["jointMoverUI_bigtoe02_r"], 'left', 160)])
+        cmds.formLayout(self.widgets["jointMover_rightFootModeFormLayout"], edit = True, af = [(self.widgets["jointMoverUI_bigtoe03_r"], 'top', 63), (self.widgets["jointMoverUI_bigtoe03_r"], 'left', 165)])
+
+        cmds.formLayout(self.widgets["jointMover_rightFootModeFormLayout"], edit = True, af = [(self.widgets["jointMoverUI_indextoe01_r"], 'top', 175), (self.widgets["jointMoverUI_indextoe01_r"], 'left', 209)])
+        cmds.formLayout(self.widgets["jointMover_rightFootModeFormLayout"], edit = True, af = [(self.widgets["jointMoverUI_indextoe02_r"], 'top', 109), (self.widgets["jointMoverUI_indextoe02_r"], 'left', 214)])
+        cmds.formLayout(self.widgets["jointMover_rightFootModeFormLayout"], edit = True, af = [(self.widgets["jointMoverUI_indextoe03_r"], 'top', 81), (self.widgets["jointMoverUI_indextoe03_r"], 'left', 220)])
+        cmds.formLayout(self.widgets["jointMover_rightFootModeFormLayout"], edit = True, af = [(self.widgets["jointMoverUI_indextoe04_r"], 'top', 56), (self.widgets["jointMoverUI_indextoe04_r"], 'left', 225)])
+
+        cmds.formLayout(self.widgets["jointMover_rightFootModeFormLayout"], edit = True, af = [(self.widgets["jointMoverUI_middletoe01_r"], 'top', 185), (self.widgets["jointMoverUI_middletoe01_r"], 'left', 234)])
+        cmds.formLayout(self.widgets["jointMover_rightFootModeFormLayout"], edit = True, af = [(self.widgets["jointMoverUI_middletoe02_r"], 'top', 125), (self.widgets["jointMoverUI_middletoe02_r"], 'left', 242)])
+        cmds.formLayout(self.widgets["jointMover_rightFootModeFormLayout"], edit = True, af = [(self.widgets["jointMoverUI_middletoe03_r"], 'top', 103), (self.widgets["jointMoverUI_middletoe03_r"], 'left', 248)])
+        cmds.formLayout(self.widgets["jointMover_rightFootModeFormLayout"], edit = True, af = [(self.widgets["jointMoverUI_middletoe04_r"], 'top', 76), (self.widgets["jointMoverUI_middletoe04_r"], 'left', 253)])
+
+        cmds.formLayout(self.widgets["jointMover_rightFootModeFormLayout"], edit = True, af = [(self.widgets["jointMoverUI_ringtoe01_r"], 'top', 192), (self.widgets["jointMoverUI_ringtoe01_r"], 'left', 257)])
+        cmds.formLayout(self.widgets["jointMover_rightFootModeFormLayout"], edit = True, af = [(self.widgets["jointMoverUI_ringtoe02_r"], 'top', 137), (self.widgets["jointMoverUI_ringtoe02_r"], 'left', 266)])
+        cmds.formLayout(self.widgets["jointMover_rightFootModeFormLayout"], edit = True, af = [(self.widgets["jointMoverUI_ringtoe03_r"], 'top', 117), (self.widgets["jointMoverUI_ringtoe03_r"], 'left', 273)])
+        cmds.formLayout(self.widgets["jointMover_rightFootModeFormLayout"], edit = True, af = [(self.widgets["jointMoverUI_ringtoe04_r"], 'top', 97), (self.widgets["jointMoverUI_ringtoe04_r"], 'left', 278)])
+        
+        cmds.formLayout(self.widgets["jointMover_rightFootModeFormLayout"], edit = True, af = [(self.widgets["jointMoverUI_pinkytoe01_r"], 'top', 203), (self.widgets["jointMoverUI_pinkytoe01_r"], 'left', 280)])
+        cmds.formLayout(self.widgets["jointMover_rightFootModeFormLayout"], edit = True, af = [(self.widgets["jointMoverUI_pinkytoe02_r"], 'top', 165), (self.widgets["jointMoverUI_pinkytoe02_r"], 'left', 290)])
+        cmds.formLayout(self.widgets["jointMover_rightFootModeFormLayout"], edit = True, af = [(self.widgets["jointMoverUI_pinkytoe03_r"], 'top', 146), (self.widgets["jointMoverUI_pinkytoe03_r"], 'left', 293)])
+        cmds.formLayout(self.widgets["jointMover_rightFootModeFormLayout"], edit = True, af = [(self.widgets["jointMoverUI_pinkytoe04_r"], 'top', 127), (self.widgets["jointMoverUI_pinkytoe04_r"], 'left', 295)])
+        
+        
+    # # # # # # # # # # # # # # # # # # # # # # # # # # # # # # # # # # # # # # # # # # # # # # # # # # # # # # # # # # # # # # # 
+    # # # # # # # # # # # # # # # # # # # # # # # # # # # # # # # # # # # # # # # # # # # # # # # # # # # # # # # # # # # # # # # 
+    # # # # # # # # # # # # # # # # # # # # # # # # # # # # # # # # # # # # # # # # # # # # # # # # # # # # # # # # # # # # # # # 
+    # # # # # # # # # # # # # # # # # # # # # # # # # # # # # # # # # # # # # # # # # # # # # # # # # # # # # # # # # # # # # # # 
+    # # # # # # # # # # # # # # # # # # # # # # # # # # # # # # # # # # # # # # # # # # # # # # # # # # # # # # # # # # # # # # # 
+    def createSkeletonToolbar(self, layout):
+        
+        #toolbar mode buttons
+        imagePath = self.mayaToolsDir + "/General/Icons/ART/"
+        self.widgets["templatesButton"] = cmds.symbolButton(w = 50, h = 50, parent = self.widgets[layout], image = imagePath + "skeletonbuilder_templates.bmp", ann = "Click to save or load a skeleton template.")
+	self.widgets["helpButton"] = cmds.symbolButton(w = 50, h = 50, parent = self.widgets[layout], image = imagePath + "help.bmp", ann = "Help", c = self.skeletonBuilderHelp)
+        
+        
+        #create popup menu for templates
+	self.widgets["templatesPopup"] = cmds.popupMenu(parent = self.widgets["templatesButton"], button = 1)
+        cmds.menuItem(label = "Save Template", parent = self.widgets["templatesPopup"], c = self.saveTemplate)
+        cmds.menuItem(label = "Load Template", parent = self.widgets["templatesPopup"], c = self.loadTemplateUI)
+        
+        
+        #place the toolbar mode buttons
+	cmds.formLayout(self.widgets[layout], edit = True, af = [(self.widgets["templatesButton"], 'top', 5), (self.widgets["templatesButton"], 'left', 8)])
+	cmds.formLayout(self.widgets[layout], edit = True, af = [(self.widgets["helpButton"], 'bottom', 5), (self.widgets["helpButton"], 'left', 8)])
+
+        
+        
+    # # # # # # # # # # # # # # # # # # # # # # # # # # # # # # # # # # # # # # # # # # # # # # # # # # # # # # # # # # # # # # # 
+    # # # # # # # # # # # # # # # # # # # # # # # # # # # # # # # # # # # # # # # # # # # # # # # # # # # # # # # # # # # # # # # 
+    # # # # # # # # # # # # # # # # # # # # # # # # # # # # # # # # # # # # # # # # # # # # # # # # # # # # # # # # # # # # # # # 
+    # # # # # # # # # # # # # # # # # # # # # # # # # # # # # # # # # # # # # # # # # # # # # # # # # # # # # # # # # # # # # # # 
+    # # # # # # # # # # # # # # # # # # # # # # # # # # # # # # # # # # # # # # # # # # # # # # # # # # # # # # # # # # # # # # # 
+    def createJointMoverToolbar(self, layout, globalMode, offsetMode, geoMode):
+        
+
+        #toolbar mode buttons
+        imagePath = self.mayaToolsDir + "/General/Icons/ART/"
+        self.widgets["aimMode"] = cmds.iconTextCheckBox( style='iconOnly', w = 50, h = 28, parent = self.widgets[layout], si = imagePath + "aimModeOn.bmp", image =  imagePath + "aimModeOff.bmp", ann = "Toggle Aim Mode", onc = self.jointMover_aimModeOn, ofc = self.jointMover_aimModeOff)	
+        self.widgets["symmetry"] = cmds.iconTextCheckBox( style='iconOnly', w = 50, h = 50, parent = self.widgets[layout], image = imagePath + "jointmover_mirror.bmp", onc = self.jointMoverUI_SymmetryMode, ofc = self.jointMoverUI_SymmetryModeOff, ann = "Toggle symmetry")
+        self.widgets["reset"] = cmds.symbolButton(w = 50, h = 50, parent = self.widgets[layout], image = imagePath + "skeletonbuilder_body.bmp", c = self.jointMoverUI_resetAll, ann = "Click to reset everything to the original position. Ctrl click to reset only selection. Alt click to reset the selection hierarchy.")
+	self.widgets["helpButton"] = cmds.symbolButton(w = 50, h = 50, parent = self.widgets[layout], image = imagePath + "help.bmp", ann = "Help", c = self.skeletonBuilderHelp)
+        self.widgets["templatesButton"] = cmds.symbolButton(w = 50, h = 50, parent = self.widgets[layout], image = imagePath + "skeletonbuilder_templates.bmp", ann = "Save or load a template for joint positions.")
+	
+	self.widgets["jointMover_Modes_Collection"] = cmds.iconTextRadioCollection(parent = self.widgets[layout])
+	
+	self.widgets["moverMode"] = cmds.iconTextRadioButton(sl = globalMode, onc = partial(self.jointMover_ChangeMoveMode, "Global", "on"), ofc = partial(self.jointMover_ChangeMoveMode, "Global", "off"), cl = self.widgets["jointMover_Modes_Collection"], w = 50, h = 50, parent = self.widgets[layout], selectionImage = imagePath + "globalMover_On.bmp", image = imagePath + "globalMover.bmp", ann = "Global Mover Mode: Controls will move heirarchy of selection.")
+        self.widgets["offsetMode"] = cmds.iconTextRadioButton(sl = offsetMode, onc = partial(self.jointMover_ChangeMoveMode, "Offset", "on"), ofc = partial(self.jointMover_ChangeMoveMode, "Offset", "off"), cl = self.widgets["jointMover_Modes_Collection"], w = 50, h = 50, parent = self.widgets[layout], selectionImage = imagePath + "offsetMover_On.bmp", image = imagePath + "offsetMover.bmp", ann = "Offset Mover Mode: Controls will only affect the joint selected and nothing down the chain.")
+        self.widgets["geoMoverMode"] = cmds.iconTextRadioButton(sl = geoMode, onc = partial(self.jointMover_ChangeMoveMode, "Geo", "on"), ofc = partial(self.jointMover_ChangeMoveMode, "Geo", "off"), cl = self.widgets["jointMover_Modes_Collection"], w = 50, h = 50, parent = self.widgets[layout], selectionImage = imagePath + "geoMover_On.bmp", image = imagePath + "geoMover.bmp", ann = "Geometry Mover Mode: Controls only affect proxy geometry and not joints. Useful for repositioning and scaling your proxy geometry.")
+        self.widgets["physiqueMode"] = cmds.symbolButton(c = self.physiqueSettings, parent = self.widgets[layout], image = imagePath + "physique_off.bmp", ann = "Physique Mode: Change the proxy mesh's overall physique")
+        self.widgets["meshMode"] = cmds.symbolButton(w = 50, h = 50, parent = self.widgets[layout], image = imagePath + "proxyGeo.bmp", c = partial(self.toggleJointMoverControls_Vis, "Mesh", "meshMode"), ann = "Toggle the visibility of the proxy geometry")
+	
+	
+	#create the skin weight page toolbar buttons
+        self.widgets["paintWeightsButton"] = cmds.symbolButton(w = 50, h = 50, parent = self.widgets[layout], image = imagePath + "paintWeightsMode.bmp", visible = False, c = self.paintWeights, ann = "Paint Skin Weights")
+        self.widgets["romButton"] = cmds.symbolButton(w = 50, h = 50, parent = self.widgets[layout], image = imagePath + "rom.bmp", visible = False, c = self.rangeOfMotion_UI, ann = "Range of Motion Tool. Generate a range of motion for your skeleton. Right click to clear all range of motion keys.")
+        self.widgets["tabletPressure"] = cmds.iconTextCheckBox(w = 50, h = 50, parent = self.widgets[layout], image = imagePath + "tabletPressure_Off.bmp", visible = False, onc = self.paintWeightsTabletPressureOn, ofc = self.paintWeightsTabletPressureOff, ann = "Toggle using tablet pressure on and off")
+        self.widgets["mirrorWeightsButton"] = cmds.symbolButton(w = 50, h = 50, parent = self.widgets[layout], image = imagePath + "mirrorWeights.bmp", visible = False, c = self.mirrorSkinWeights, ann = "Mirror Skin Weights")
+        self.widgets["pruneWeightsButton"] = cmds.symbolButton(w = 50, h = 50, parent = self.widgets[layout], image = imagePath + "pruneWeights.bmp", visible = False, c = self.pruneSkinWeights, ann = "Prune Skin Weights")
+        self.widgets["addInfluenceButton"] = cmds.symbolButton(w = 50, h = 50, parent = self.widgets[layout], image = imagePath + "addInf.bmp", visible = False, c = self.paintWeightsAddInfluence, enable = False, ann = "Add influences tool")
+        self.widgets["removeInfluenceButton"] = cmds.symbolButton(w = 50, h = 50, parent = self.widgets[layout], image = imagePath + "removeUnused.bmp", visible = False, enable = False, c = self.paintWeightsRemoveInfluence, ann = "Remove influences tool. Also includes remove unused influences.")
+        self.widgets["exportSkinWeightsButton"] = cmds.symbolButton(w = 50, h = 50, parent = self.widgets[layout], image = imagePath + "exportSkinWeights.bmp", visible = False, enable = True, c = self.exportSkinWeights, ann = "Export skin weights for selected mesh.")
+        self.widgets["importSkinWeightsButton"] = cmds.symbolButton(w = 50, h = 50, parent = self.widgets[layout], image = imagePath + "importSkinWeights.bmp", visible = False, enable = True, c = self.importSkinWeights, ann = "Import skin weights for selected mesh.")
+
+        
+        #create the popup menu for the templates button
+        self.widgets["templatesPopup"] = cmds.popupMenu(parent = self.widgets["templatesButton"], button = 1)
+        cmds.menuItem(label = "Save Template", parent = self.widgets["templatesPopup"], c = self.jointMoverUI_SaveTemplateUI)
+        cmds.menuItem(label = "Load Template", parent = self.widgets["templatesPopup"], c = self.jointMoverUI_LoadTemplateUI)
+        
+        
+	#create the popup menu for the ROM tool
+        self.widgets["rom_popup"] = cmds.popupMenu(parent = self.widgets["romButton"], button = 3)
+        cmds.menuItem(label = "Clear all Range of Motion Keys", parent = self.widgets["rom_popup"], c = self.rangeOfMotion_Clear)
+	
+
+	
+        #place the toolbar mode buttons
+	cmds.formLayout(self.widgets[layout], edit = True, af = [(self.widgets["templatesButton"], 'top', 5), (self.widgets["templatesButton"], 'left', 8)])
+	cmds.formLayout(self.widgets[layout], edit = True, af = [(self.widgets["aimMode"], 'top', 56), (self.widgets["aimMode"], 'left', 8)])
+        cmds.formLayout(self.widgets[layout], edit = True, af = [(self.widgets["symmetry"], 'top', 90), (self.widgets["symmetry"], 'left', 8)])
+        cmds.formLayout(self.widgets[layout], edit = True, af = [(self.widgets["reset"], 'top', 150), (self.widgets["reset"], 'left', 8)])
+        cmds.formLayout(self.widgets[layout], edit = True, af = [(self.widgets["moverMode"], 'top', 210), (self.widgets["moverMode"], 'left', 8)])
+        cmds.formLayout(self.widgets[layout], edit = True, af = [(self.widgets["offsetMode"], 'top', 270), (self.widgets["offsetMode"], 'left', 8)])
+	cmds.formLayout(self.widgets[layout], edit = True, af = [(self.widgets["geoMoverMode"], 'top', 330), (self.widgets["geoMoverMode"], 'left', 8)])
+	cmds.formLayout(self.widgets[layout], edit = True, af = [(self.widgets["physiqueMode"], 'top', 390), (self.widgets["physiqueMode"], 'left', 8)])
+        cmds.formLayout(self.widgets[layout], edit = True, af = [(self.widgets["meshMode"], 'top', 450), (self.widgets["meshMode"], 'left', 8)])
+	cmds.formLayout(self.widgets[layout], edit = True, af = [(self.widgets["helpButton"], 'top', 540), (self.widgets["helpButton"], 'left', 8)])
+	
+	cmds.formLayout(self.widgets[layout], edit = True, af = [(self.widgets["paintWeightsButton"], 'top', 5), (self.widgets["paintWeightsButton"], 'left', 8)])
+	cmds.formLayout(self.widgets[layout], edit = True, af = [(self.widgets["romButton"], 'top', 65), (self.widgets["romButton"], 'left', 8)])
+	cmds.formLayout(self.widgets[layout], edit = True, af = [(self.widgets["tabletPressure"], 'top', 125), (self.widgets["tabletPressure"], 'left', 8)])
+	cmds.formLayout(self.widgets[layout], edit = True, af = [(self.widgets["mirrorWeightsButton"], 'top', 185), (self.widgets["mirrorWeightsButton"], 'left', 8)])
+	cmds.formLayout(self.widgets[layout], edit = True, af = [(self.widgets["pruneWeightsButton"], 'top', 245), (self.widgets["pruneWeightsButton"], 'left', 8)])
+	cmds.formLayout(self.widgets[layout], edit = True, af = [(self.widgets["addInfluenceButton"], 'top', 305), (self.widgets["addInfluenceButton"], 'left', 8)])
+	cmds.formLayout(self.widgets[layout], edit = True, af = [(self.widgets["removeInfluenceButton"], 'top', 365), (self.widgets["removeInfluenceButton"], 'left', 8)])
+	cmds.formLayout(self.widgets[layout], edit = True, af = [(self.widgets["exportSkinWeightsButton"], 'top', 425), (self.widgets["exportSkinWeightsButton"], 'left', 8)])
+	cmds.formLayout(self.widgets[layout], edit = True, af = [(self.widgets["importSkinWeightsButton"], 'top', 485), (self.widgets["importSkinWeightsButton"], 'left', 8)])
+
+	if cmds.objExists("Symmetry_Nodes"):
+	    cmds.iconTextCheckBox(self.widgets["symmetry"] , edit = True, v = True)
+	    cmds.symbolButton(self.widgets["reset"], edit = True, enable = False)
+	    cmds.symbolButton(self.widgets["templatesButton"], edit = True, enable = False)
+	    
+	if cmds.objExists("AimModeNodes"):
+	    cmds.iconTextCheckBox(self.widgets["aimMode"] , edit = True, v = True)
+	    cmds.symbolButton(self.widgets["reset"], edit = True, enable = False)
+	    cmds.symbolButton(self.widgets["templatesButton"], edit = True, enable = False)
+	    cmds.iconTextCheckBox(self.widgets["symmetry"], edit = True, enable = False)	    
+
+
+
+    # # # # # # # # # # # # # # # # # # # # # # # # # # # # # # # # # # # # # # # # # # # # # # # # # # # # # # # # # # # # # # # 
+    # # # # # # # # # # # # # # # # # # # # # # # # # # # # # # # # # # # # # # # # # # # # # # # # # # # # # # # # # # # # # # # 
+    # # # # # # # # # # # # # # # # # # # # # # # # # # # # # # # # # # # # # # # # # # # # # # # # # # # # # # # # # # # # # # # 
+    # # # # # # # # # # # # # # # # # # # # # # # # # # # # # # # # # # # # # # # # # # # # # # # # # # # # # # # # # # # # # # # 
+    # # # # # # # # # # # # # # # # # # # # # # # # # # # # # # # # # # # # # # # # # # # # # # # # # # # # # # # # # # # # # # # 
+    def extraJointInfo_UI(self, jointType, *args):
+	
+	#if window exists, delete it
+	if cmds.window("extraJointInfo_UI", exists = True):
+	    cmds.deleteUI("extraJointInfo_UI")
+	    
+	#create the window and the layouts
+	window = cmds.window("extraJointInfo_UI", w = 300, h = 150, title = "Joint Info", tbm = False, sizeable = False)
+	mainLayout = cmds.columnLayout(w = 300, h = 150)
+	formLayout = cmds.formLayout(w = 300, h = 150, parent = mainLayout)
+	
+	#widgets
+	label1 = cmds.text(label = "Enter Joint Name:    ")
+	label2 = cmds.text(label = "Choose Joint Parent: ")
+	
+	name = cmds.textField("extraJointInfo_UI_JointName", w = 150)
+	parent = cmds.optionMenu("extraJointInfo_UI_JointParent", w = 150)
+	
+	if jointType == "leaf":
+	    #add checkboxes for translation and rotation
+	    translationCB = cmds.checkBox("extraJointInfo_UItranslationCB", label = "Translation", v = True)
+	    rotationB = cmds.checkBox("extraJointInfo_UIrotationCB", label = "Rotation", v = True)
+	    
+	
+	okButton = cmds.button(label = "OK", w = 145, h = 50, c = partial(self.extraJointInfo_UI_Accept, jointType))
+	cancelButton = cmds.button(label = "Cancel", w = 145, h = 50, c = self.extraJointInfo_UI_Cancel)
+	
+	
+	#place widgets
+	cmds.formLayout(formLayout, edit = True, af = [(label1, 'top', 10), (label1, 'left', 10)])
+	cmds.formLayout(formLayout, edit = True, af = [(label2, 'top', 45), (label2, 'left', 10)])
+	
+	cmds.formLayout(formLayout, edit = True, af = [(name, 'top', 10), (name, 'left', 130)])
+	cmds.formLayout(formLayout, edit = True, af = [(parent, 'top', 45), (parent, 'left', 130)])
+	
+	cmds.formLayout(formLayout, edit = True, af = [(okButton, 'top', 90), (okButton, 'left', 5)])
+	cmds.formLayout(formLayout, edit = True, af = [(cancelButton, 'top', 90), (cancelButton, 'left', 150)])
+	
+	if jointType == "leaf":
+	    cmds.formLayout(formLayout, edit = True, af = [(translationCB, 'top', 70), (translationCB, 'left', 5)])
+	    cmds.formLayout(formLayout, edit = True, af = [(rotationB, 'top', 70), (rotationB, 'left', 150)])
+	
+	
+	#show window
+	cmds.showWindow(window)
+	self.extraJointInfo_UI_GetJoints("extraJointInfo_UI_JointParent")
+	
+	
+    # # # # # # # # # # # # # # # # # # # # # # # # # # # # # # # # # # # # # # # # # # # # # # # # # # # # # # # # # # # # # # # 
+    # # # # # # # # # # # # # # # # # # # # # # # # # # # # # # # # # # # # # # # # # # # # # # # # # # # # # # # # # # # # # # # 
+    # # # # # # # # # # # # # # # # # # # # # # # # # # # # # # # # # # # # # # # # # # # # # # # # # # # # # # # # # # # # # # # 
+    # # # # # # # # # # # # # # # # # # # # # # # # # # # # # # # # # # # # # # # # # # # # # # # # # # # # # # # # # # # # # # # 
+    # # # # # # # # # # # # # # # # # # # # # # # # # # # # # # # # # # # # # # # # # # # # # # # # # # # # # # # # # # # # # # # 
+    def dynamicJointInfo_UI(self, *args):
+	
+	#if window exists, delete it
+	if cmds.window("dynamicJointInfo_UI", exists = True):
+	    cmds.deleteUI("dynamicJointInfo_UI")
+	    
+	#create the window and the layouts
+	window = cmds.window("dynamicJointInfo_UI", w = 300, h = 200, title = "Joint Info", tbm = False, sizeable = False)
+	mainLayout = cmds.columnLayout(w = 300, h = 200)
+	formLayout = cmds.formLayout(w = 300, h = 200, parent = mainLayout)
+	
+	#widgets
+	label1 = cmds.text(label = "Enter Joint Name:")
+	label2 = cmds.text(label = "Choose Joint Parent:")
+	label3 = cmds.text(label = "Number of Joints:")
+	
+	name = cmds.textField("dynamicJointInfo_UI_JointName", w = 150)
+	parent = cmds.optionMenu("dynamicJointInfo_UI_JointParent", w = 150)
+	numJoints = cmds.textField("dynamicJointInfo_UI_NumJoints", w = 150)
+	
+	
+	okButton = cmds.button(label = "OK", w = 145, h = 50, c = self.dynamicJointInfo_Accept)
+	cancelButton = cmds.button(label = "Cancel", w = 145, h = 50, c = self.dynamicJointInfo_Cancel)
+	
+	
+	#place widgets
+	cmds.formLayout(formLayout, edit = True, af = [(label1, 'top', 10), (label1, 'left', 10)])
+	cmds.formLayout(formLayout, edit = True, af = [(label3, 'top', 45), (label3, 'left', 10)])
+	cmds.formLayout(formLayout, edit = True, af = [(label2, 'top', 80), (label2, 'left', 10)])
+	
+	cmds.formLayout(formLayout, edit = True, af = [(name, 'top', 10), (name, 'left', 130)])
+	cmds.formLayout(formLayout, edit = True, af = [(numJoints, 'top', 45), (numJoints, 'left', 130)])
+	cmds.formLayout(formLayout, edit = True, af = [(parent, 'top', 80), (parent, 'left', 130)])
+	
+	cmds.formLayout(formLayout, edit = True, af = [(okButton, 'top', 125), (okButton, 'left', 5)])
+	cmds.formLayout(formLayout, edit = True, af = [(cancelButton, 'top', 125), (cancelButton, 'left', 150)])
+	
+	
+	#show window
+	cmds.showWindow(window)
+	self.extraJointInfo_UI_GetJoints("dynamicJointInfo_UI_JointParent")
+	
+	
+	
+    # # # # # # # # # # # # # # # # # # # # # # # # # # # # # # # # # # # # # # # # # # # # # # # # # # # # # # # # # # # # # # # 
+    # # # # # # # # # # # # # # # # # # # # # # # # # # # # # # # # # # # # # # # # # # # # # # # # # # # # # # # # # # # # # # # 
+    # # # # # # # # # # # # # # # # # # # # # # # # # # # # # # # # # # # # # # # # # # # # # # # # # # # # # # # # # # # # # # # 
+    # # # # # # # # # # # # # # # # # # # # # # # # # # # # # # # # # # # # # # # # # # # # # # # # # # # # # # # # # # # # # # # 
+    # # # # # # # # # # # # # # # # # # # # # # # # # # # # # # # # # # # # # # # # # # # # # # # # # # # # # # # # # # # # # # # 
+    def dynamicJointInfo_Accept(self, *args):
+	#On Accept, create the entry into the table. Get the information from the extraJointInfo_UI first.
+	
+	name = cmds.textField("dynamicJointInfo_UI_JointName", q = True, text = True)
+	parent = cmds.optionMenu("dynamicJointInfo_UI_JointParent", q = True, value = True)
+	num = cmds.textField("dynamicJointInfo_UI_NumJoints", q = True, text = True)
+	
+	#check to see if name already exists in the list for extra joints
+	extraJoints = cmds.textScrollList(self.widgets["addedRigModulesList"], q = True, allItems = True)
+	if extraJoints != None:
+	    for joint in extraJoints:
+		niceName = joint.rpartition("] ")[2]
+		if niceName == name:
+		    cmds.confirmDialog(title = "Error", icon = "warning", message = "A rig module already exists with that name.")
+		    return
+		
+		
+	if name != "":
+	    if num > 0:
+
+		#add to the list of added rig modules
+		cmds.textScrollList(self.widgets["addedRigModulesList"], edit = True, append = "[chain] " + name, sc = self.rigMod_changeSettingsForms)
+		#select the newly added item
+		cmds.textScrollList(self.widgets["addedRigModulesList"], edit = True, si = "[chain] " + name)
+		
+		#create the entry info and parent to the self.widgets["rigModuleSettingsFrame"]
+		#for leaf joints, we need name, parent, translate/rotate
+		
+		#hide existing formLayouts
+		self.rigMod_hideSettingsForms()
+		
+		self.widgets["rigMod_" + name + "_settingsForm"] = cmds.formLayout(w = 150, h = 220, parent = self.widgets["rigModuleSettingsFrame"])
+		label_name = cmds.text(label = "Name: ")
+		self.widgets["rigMod_" + name + "_settings_name"] = cmds.textField(w = 100, parent = self.widgets["rigMod_" + name + "_settingsForm"], text = name)
+		label_parent = cmds.text(label = "Parent: ")
+		self.widgets["rigMod_" + name + "_settings_parent"] = cmds.textField(w = 80, parent = self.widgets["rigMod_" + name + "_settingsForm"], text = parent, editable = False)
+		self.widgets["rigMod_" + name + "_settings_changeParent"] = cmds.button(label = "...", w = 20, parent = self.widgets["rigMod_" + name + "_settingsForm"], ann = "Change Parent", c = partial(self.extraJointInfo_UI_ParentList, name))
+		label_number = cmds.text(label = "Count: ")
+		self.widgets["rigMod_" + name + "_settings_numberInChain"] = cmds.textField(w = 100, parent = self.widgets["rigMod_" + name + "_settingsForm"], text = str(num), ann = "Number of joints in the chain.")
+		
+		cmds.formLayout(self.widgets["rigMod_" + name + "_settingsForm"], edit = True, af = [(label_name, "top", 5),(label_name, "left", 5)])
+		cmds.formLayout(self.widgets["rigMod_" + name + "_settingsForm"], edit = True, af = [(self.widgets["rigMod_" + name + "_settings_name"], "top", 5),(self.widgets["rigMod_" + name + "_settings_name"], "right", 5)])
+		cmds.formLayout(self.widgets["rigMod_" + name + "_settingsForm"], edit = True, af = [(label_parent, "top", 35),(label_parent, "left", 5)])
+		cmds.formLayout(self.widgets["rigMod_" + name + "_settingsForm"], edit = True, af = [(self.widgets["rigMod_" + name + "_settings_parent"], "top", 35),(self.widgets["rigMod_" + name + "_settings_parent"], "right", 25)])	    
+		cmds.formLayout(self.widgets["rigMod_" + name + "_settingsForm"], edit = True, af = [(self.widgets["rigMod_" + name + "_settings_changeParent"], "top", 35),(self.widgets["rigMod_" + name + "_settings_changeParent"], "right", 5)])	    
+		cmds.formLayout(self.widgets["rigMod_" + name + "_settingsForm"], edit = True, af = [(label_number, "top", 65),(label_number, "left", 5)])
+		cmds.formLayout(self.widgets["rigMod_" + name + "_settingsForm"], edit = True, af = [(self.widgets["rigMod_" + name + "_settings_numberInChain"], "top", 65),(self.widgets["rigMod_" + name + "_settings_numberInChain"], "right", 5)])
+
+
+		cmds.deleteUI("dynamicJointInfo_UI")
+		
+		
+	    else:
+		cmds.warning("Please enter a valid number of joints")
+		             
+	else:
+	    cmds.warning("Please enter a valid name")
+	
+	
+	
+    # # # # # # # # # # # # # # # # # # # # # # # # # # # # # # # # # # # # # # # # # # # # # # # # # # # # # # # # # # # # # # # 
+    # # # # # # # # # # # # # # # # # # # # # # # # # # # # # # # # # # # # # # # # # # # # # # # # # # # # # # # # # # # # # # # 
+    # # # # # # # # # # # # # # # # # # # # # # # # # # # # # # # # # # # # # # # # # # # # # # # # # # # # # # # # # # # # # # # 
+    # # # # # # # # # # # # # # # # # # # # # # # # # # # # # # # # # # # # # # # # # # # # # # # # # # # # # # # # # # # # # # # 
+    # # # # # # # # # # # # # # # # # # # # # # # # # # # # # # # # # # # # # # # # # # # # # # # # # # # # # # # # # # # # # # # 
+    def dynamicJointInfo_Cancel(self, *args):
+	cmds.deleteUI("dynamicJointInfo_UI")
+	
+	
+    # # # # # # # # # # # # # # # # # # # # # # # # # # # # # # # # # # # # # # # # # # # # # # # # # # # # # # # # # # # # # # # 
+    # # # # # # # # # # # # # # # # # # # # # # # # # # # # # # # # # # # # # # # # # # # # # # # # # # # # # # # # # # # # # # # 
+    # # # # # # # # # # # # # # # # # # # # # # # # # # # # # # # # # # # # # # # # # # # # # # # # # # # # # # # # # # # # # # # 
+    # # # # # # # # # # # # # # # # # # # # # # # # # # # # # # # # # # # # # # # # # # # # # # # # # # # # # # # # # # # # # # # 
+    # # # # # # # # # # # # # # # # # # # # # # # # # # # # # # # # # # # # # # # # # # # # # # # # # # # # # # # # # # # # # # # 
+    def removeExtraDynJoint(self, name, parent, type, num, removeButton, *args):
+	
+	for each in [name, parent, type, num]:
+	    cmds.deleteUI(each)
+	    
+	cmds.deleteUI(removeButton)
+	
+    # # # # # # # # # # # # # # # # # # # # # # # # # # # # # # # # # # # # # # # # # # # # # # # # # # # # # # # # # # # # # # # 
+    # # # # # # # # # # # # # # # # # # # # # # # # # # # # # # # # # # # # # # # # # # # # # # # # # # # # # # # # # # # # # # # 
+    # # # # # # # # # # # # # # # # # # # # # # # # # # # # # # # # # # # # # # # # # # # # # # # # # # # # # # # # # # # # # # # 
+    # # # # # # # # # # # # # # # # # # # # # # # # # # # # # # # # # # # # # # # # # # # # # # # # # # # # # # # # # # # # # # # 
+    # # # # # # # # # # # # # # # # # # # # # # # # # # # # # # # # # # # # # # # # # # # # # # # # # # # # # # # # # # # # # # # 
+    def extraJointInfo_UI_Sort(self,  *args):
+	
+	method = cmds.optionMenu(self.widgets["rigModuleSortListMenu"], q = True, v = True)
+	
+	if method == "Alphabetically":
+	    #get all of the items in the list
+	    items = cmds.textScrollList(self.widgets["addedRigModulesList"], q = True, allItems = True)
+	    
+	    
+	    namesAndTypes = []
+	    for item in items:
+		niceName = item.partition("] ")[2]
+		modType = item.partition(" ")[0]
+		namesAndTypes.append([niceName, str(modType)])
+		
+		
+		
+	    
+	    #sort that list
+	    sortedList = sorted(namesAndTypes)
+	    
+	    #clear all items and re-add
+	    cmds.textScrollList(self.widgets["addedRigModulesList"], edit = True, removeAll = True)
+	    
+	    for item in sortedList:
+		cmds.textScrollList(self.widgets["addedRigModulesList"], edit = True, append =  item[1] + " " + item[0], sc = self.rigMod_changeSettingsForms)
+		
+	    #select the first index
+	    try:
+		cmds.textScrollList(self.widgets["addedRigModulesList"], edit = True, si = items[0])
+		
+	    except:
+		pass
+	    
+	    
+	if method == "By Type":
+	    #get all of the items in the list
+	    items = cmds.textScrollList(self.widgets["addedRigModulesList"], q = True, allItems = True)
+	
+	
+	    #sort that list
+	    sortedList = sorted(items)
+	    
+	    #clear all items and re-add
+	    cmds.textScrollList(self.widgets["addedRigModulesList"], edit = True, removeAll = True)
+	    
+	    for item in sortedList:
+		cmds.textScrollList(self.widgets["addedRigModulesList"], edit = True, append = item, sc = self.rigMod_changeSettingsForms)
+		
+	    #select the first index
+	    try:
+		cmds.textScrollList(self.widgets["addedRigModulesList"], edit = True, si = items[0])
+		
+	    except:
+		pass
+		
+	    
+	
+    # # # # # # # # # # # # # # # # # # # # # # # # # # # # # # # # # # # # # # # # # # # # # # # # # # # # # # # # # # # # # # # 
+    # # # # # # # # # # # # # # # # # # # # # # # # # # # # # # # # # # # # # # # # # # # # # # # # # # # # # # # # # # # # # # # 
+    # # # # # # # # # # # # # # # # # # # # # # # # # # # # # # # # # # # # # # # # # # # # # # # # # # # # # # # # # # # # # # # 
+    # # # # # # # # # # # # # # # # # # # # # # # # # # # # # # # # # # # # # # # # # # # # # # # # # # # # # # # # # # # # # # # 
+    # # # # # # # # # # # # # # # # # # # # # # # # # # # # # # # # # # # # # # # # # # # # # # # # # # # # # # # # # # # # # # # 
+    def extraJointInfo_UI_GetJoints(self, layout, *args):
+	
+	#set a few of the option menu items to the defaults, but also get neck, spine, and extra bones added
+	returnList = []
+	
+	
+	#lower body
+	returnList.append("root")
+	returnList.append("pelvis")
+	returnList.append("thigh_l")
+	returnList.append("calf_l")
+	returnList.append("foot_l")
+	returnList.append("thigh_r")
+	returnList.append("calf_r")
+	returnList.append("foot_r")
+	
+
+	if cmds.checkBox(self.widgets["leftLegIncludeBallJoint_CB"], q = True, value = True):
+	    returnList.append("ball_l")
+	    
+	if cmds.checkBox(self.widgets["rightLegIncludeBallJoint_CB"], q = True, value = True):
+	    returnList.append("ball_r")
+	
+	#spine
+	numSpineBones = cmds.intField(self.widgets["torsoOptions_numSpineBones"], q = True, value = True)
+	for i in range(numSpineBones):
+	    returnList.append("spine_0" + str(i + 1))
+
+	#neck
+	numNeckBones = cmds.intField(self.widgets["torsoOptions_numNeckBones"], q = True, value = True)
+	for i in range(numNeckBones):
+	    returnList.append("neck_0" + str(i + 1))
+
+	    
+	#head
+	returnList.append("head")
+	
+
+	#arms
+	returnList.append("clavicle_l")
+	returnList.append("upperarm_l")
+	returnList.append("lowerarm_l")
+	returnList.append("hand_l")
+	returnList.append("clavicle_r")
+	returnList.append("upperarm_r")
+	returnList.append("lowerarm_r")
+	returnList.append("hand_r")
+
+	    
+	
+	#twist bones and other
+        twistFields = [["leftArmOptions_numUpArmTwistBones", "upperarm_twist_l"],["leftArmOptions_numLowArmTwistBones", "lowerarm_twist_l"], ["rightArmOptions_numUpArmTwistBones", "upperarm_twist_r"], ["rightArmOptions_numLowArmTwistBones", "lowerarm_twist_r"], ["leftLegOptions_numThighTwistBones", "thigh_twist_l"], ["leftLegOptions_numCalfTwistBones", "calf_twist_l"], ["rightLegOptions_numThighTwistBones", "thigh_twist_r"], ["rightLegOptions_numCalfTwistBones", "calf_twist_r"]]
+	
+	#lowerarm_twist_l upperarm_twist_l calf_twist_l thigh_twist_l
+	for field in twistFields:
+	    value = cmds.intField(self.widgets[field[0]], q = True, v = True)
+	    for i in range(int(value)):
+		side = field[1].rpartition("_")[2]
+		prefix = field[1].rpartition("_")[0]
+		returnList.append(prefix + "_0" + str(i + 1) + "_" + side)
+		
+		
+		
+	#get any extra joints
+	extras = cmds.textScrollList(self.widgets["addedRigModulesList"], q = True, allItems = True)
+	if extras != None:
+	    for extra in extras:
+		name = extra.rpartition("] ")[2]
+		returnList.append(name)
+
+	
+	if layout != None:
+	    for joint in returnList:
+		cmds.menuItem(joint, parent = layout)
+		
+	return returnList
+		
+	
+	
+
+	
+    # # # # # # # # # # # # # # # # # # # # # # # # # # # # # # # # # # # # # # # # # # # # # # # # # # # # # # # # # # # # # # # 
+    # # # # # # # # # # # # # # # # # # # # # # # # # # # # # # # # # # # # # # # # # # # # # # # # # # # # # # # # # # # # # # # 
+    # # # # # # # # # # # # # # # # # # # # # # # # # # # # # # # # # # # # # # # # # # # # # # # # # # # # # # # # # # # # # # # 
+    # # # # # # # # # # # # # # # # # # # # # # # # # # # # # # # # # # # # # # # # # # # # # # # # # # # # # # # # # # # # # # # 
+    # # # # # # # # # # # # # # # # # # # # # # # # # # # # # # # # # # # # # # # # # # # # # # # # # # # # # # # # # # # # # # # 
+    def removeSelectedModule(self, *args):
+	
+	#get the selected module in the list
+	selected = cmds.textScrollList(self.widgets["addedRigModulesList"], q = True, si = True)[0]
+	niceName = selected.rpartition("] ")[2]
+	cmds.textScrollList(self.widgets["addedRigModulesList"], edit = True, ri = selected)
+	cmds.deleteUI(self.widgets["rigMod_" + niceName + "_settingsForm"])
+	
+	try:
+	    cmds.textScrollList(self.widgets["addedRigModulesList"], edit = True, sii = 1)
+	    self.rigMod_hideSettingsForms()
+	    
+	    selected = cmds.textScrollList(self.widgets["addedRigModulesList"], q = True, si = True)[0]
+	    niceName = selected.rpartition("] ")[2]
+	    cmds.formLayout(self.widgets["rigMod_" + niceName + "_settingsForm"], edit = True, visible = True)
+	    
+	    
+	except:
+	    cmds.formLayout(self.widgets["rigModuleDefaultSettingsLayout"], edit = True, visible = True)
+	    
+	    
+	
+    # # # # # # # # # # # # # # # # # # # # # # # # # # # # # # # # # # # # # # # # # # # # # # # # # # # # # # # # # # # # # # # 
+    # # # # # # # # # # # # # # # # # # # # # # # # # # # # # # # # # # # # # # # # # # # # # # # # # # # # # # # # # # # # # # # 
+    # # # # # # # # # # # # # # # # # # # # # # # # # # # # # # # # # # # # # # # # # # # # # # # # # # # # # # # # # # # # # # # 
+    # # # # # # # # # # # # # # # # # # # # # # # # # # # # # # # # # # # # # # # # # # # # # # # # # # # # # # # # # # # # # # # 
+    # # # # # # # # # # # # # # # # # # # # # # # # # # # # # # # # # # # # # # # # # # # # # # # # # # # # # # # # # # # # # # # 
+    def extraJointInfo_UI_Accept(self, jointType, *args):
+	#On Accept, create the entry into the tree lister. Get the information from the extraJointInfo_UI first.
+	imagePath = self.mayaToolsDir + "/General/Icons/ART/"
+	
+	name = cmds.textField("extraJointInfo_UI_JointName", q = True, text = True)
+	parent = cmds.optionMenu("extraJointInfo_UI_JointParent", q = True, value = True)
+	suffix = ""
+	
+	#check to see if name already exists in the list for extra joints
+	extraJoints = cmds.textScrollList(self.widgets["addedRigModulesList"], q = True, allItems = True)
+	if extraJoints != None:
+	    for joint in extraJoints:
+		niceName = joint.rpartition("] ")[2]
+		if niceName == name:
+		    cmds.confirmDialog(title = "Error", icon = "warning", message = "A rig module already exists with that name.")
+		    return
+	    
+	
+	if jointType == "leaf":
+	    #get checkbox values
+	    translation = cmds.checkBox("extraJointInfo_UItranslationCB", q = True, v = True)
+	    rotation = cmds.checkBox("extraJointInfo_UIrotationCB", q = True, v = True)
+	    
+	    
+	    #add to the list of added rig modules
+	    cmds.textScrollList(self.widgets["addedRigModulesList"], edit = True, append = "[leaf] " + name, sc = self.rigMod_changeSettingsForms)
+	    #select the newly added item
+	    cmds.textScrollList(self.widgets["addedRigModulesList"], edit = True, si = "[leaf] " + name)
+	    
+	    #create the entry info and parent to the self.widgets["rigModuleSettingsFrame"]
+	    #for leaf joints, we need name, parent, translate/rotate
+	    
+	    #hide existing formLayouts
+	    self.rigMod_hideSettingsForms()
+	    
+	    self.widgets["rigMod_" + name + "_settingsForm"] = cmds.formLayout(w = 150, h = 220, parent = self.widgets["rigModuleSettingsFrame"])
+	    label_name = cmds.text(label = "Name: ")
+	    self.widgets["rigMod_" + name + "_settings_name"] = cmds.textField(w = 100, parent = self.widgets["rigMod_" + name + "_settingsForm"], text = name, cc = partial(self.updateRigModuleList, name, "leaf"))
+	    label_parent = cmds.text(label = "Parent: ")
+	    self.widgets["rigMod_" + name + "_settings_parent"] = cmds.textField(w = 80, parent = self.widgets["rigMod_" + name + "_settingsForm"], text = parent, editable = False)
+	    self.widgets["rigMod_" + name + "_settings_changeParent"] = cmds.button(label = "...", w = 20, parent = self.widgets["rigMod_" + name + "_settingsForm"], ann = "Change Parent", c = partial(self.extraJointInfo_UI_ParentList, name))
+	    self.widgets["rigMod_" + name + "_settings_translate"] = cmds.checkBox(label = "Translate", v = translation, ann = "Joint can Translate")
+	    self.widgets["rigMod_" + name + "_settings_rotate"] = cmds.checkBox(label = "Rotate", v = rotation, ann = "Joint can Rotate")
+	    
+	    cmds.formLayout(self.widgets["rigMod_" + name + "_settingsForm"], edit = True, af = [(label_name, "top", 5),(label_name, "left", 5)])
+	    cmds.formLayout(self.widgets["rigMod_" + name + "_settingsForm"], edit = True, af = [(self.widgets["rigMod_" + name + "_settings_name"], "top", 5),(self.widgets["rigMod_" + name + "_settings_name"], "right", 5)])
+	    cmds.formLayout(self.widgets["rigMod_" + name + "_settingsForm"], edit = True, af = [(label_parent, "top", 35),(label_parent, "left", 5)])
+	    cmds.formLayout(self.widgets["rigMod_" + name + "_settingsForm"], edit = True, af = [(self.widgets["rigMod_" + name + "_settings_parent"], "top", 35),(self.widgets["rigMod_" + name + "_settings_parent"], "right", 25)])	    
+	    cmds.formLayout(self.widgets["rigMod_" + name + "_settingsForm"], edit = True, af = [(self.widgets["rigMod_" + name + "_settings_changeParent"], "top", 35),(self.widgets["rigMod_" + name + "_settings_changeParent"], "right", 5)])	    
+	    cmds.formLayout(self.widgets["rigMod_" + name + "_settingsForm"], edit = True, af = [(self.widgets["rigMod_" + name + "_settings_translate"], "top", 65),(self.widgets["rigMod_" + name + "_settings_translate"], "left", 5)])	    
+	    cmds.formLayout(self.widgets["rigMod_" + name + "_settingsForm"], edit = True, af = [(self.widgets["rigMod_" + name + "_settings_rotate"], "top", 85),(self.widgets["rigMod_" + name + "_settings_rotate"], "left", 5)])	    
+	    
+
+
+
+	else:
+	    
+	    #add to the list of added rig modules
+	    cmds.textScrollList(self.widgets["addedRigModulesList"], edit = True, append = "[jiggle] " + name, sc = self.rigMod_changeSettingsForms)
+	    
+	    #select the newly added item
+	    cmds.textScrollList(self.widgets["addedRigModulesList"], edit = True, si = "[jiggle] " + name)
+	    
+	    #create the entry info and parent to the self.widgets["rigModuleSettingsFrame"]
+	    
+	    #hide existing formLayouts
+	    self.rigMod_hideSettingsForms()
+	    
+	    self.widgets["rigMod_" + name + "_settingsForm"] = cmds.formLayout(w = 150, h = 220, parent = self.widgets["rigModuleSettingsFrame"])
+	    label_name = cmds.text(label = "Name: ")
+	    self.widgets["rigMod_" + name + "_settings_name"] = cmds.textField(w = 100, parent = self.widgets["rigMod_" + name + "_settingsForm"], text = name)
+	    label_parent = cmds.text(label = "Parent: ")
+	    self.widgets["rigMod_" + name + "_settings_parent"] = cmds.textField(w = 80, parent = self.widgets["rigMod_" + name + "_settingsForm"], text = parent, editable = False)
+	    self.widgets["rigMod_" + name + "_settings_changeParent"] = cmds.button(label = "...", w = 20, parent = self.widgets["rigMod_" + name + "_settingsForm"], ann = "Change Parent", c = partial(self.extraJointInfo_UI_ParentList, name))
+	    
+	    cmds.formLayout(self.widgets["rigMod_" + name + "_settingsForm"], edit = True, af = [(label_name, "top", 5),(label_name, "left", 5)])
+	    cmds.formLayout(self.widgets["rigMod_" + name + "_settingsForm"], edit = True, af = [(self.widgets["rigMod_" + name + "_settings_name"], "top", 5),(self.widgets["rigMod_" + name + "_settings_name"], "right", 5)])
+	    cmds.formLayout(self.widgets["rigMod_" + name + "_settingsForm"], edit = True, af = [(label_parent, "top", 35),(label_parent, "left", 5)])
+	    cmds.formLayout(self.widgets["rigMod_" + name + "_settingsForm"], edit = True, af = [(self.widgets["rigMod_" + name + "_settings_parent"], "top", 35),(self.widgets["rigMod_" + name + "_settings_parent"], "right", 25)])	    
+	    cmds.formLayout(self.widgets["rigMod_" + name + "_settingsForm"], edit = True, af = [(self.widgets["rigMod_" + name + "_settings_changeParent"], "top", 35),(self.widgets["rigMod_" + name + "_settings_changeParent"], "right", 5)])	    
+	    
+	    
+  
+	cmds.deleteUI("extraJointInfo_UI")
+	
+    # # # # # # # # # # # # # # # # # # # # # # # # # # # # # # # # # # # # # # # # # # # # # # # # # # # # # # # # # # # # # # # 
+    # # # # # # # # # # # # # # # # # # # # # # # # # # # # # # # # # # # # # # # # # # # # # # # # # # # # # # # # # # # # # # # 
+    # # # # # # # # # # # # # # # # # # # # # # # # # # # # # # # # # # # # # # # # # # # # # # # # # # # # # # # # # # # # # # # 
+    # # # # # # # # # # # # # # # # # # # # # # # # # # # # # # # # # # # # # # # # # # # # # # # # # # # # # # # # # # # # # # # 
+    # # # # # # # # # # # # # # # # # # # # # # # # # # # # # # # # # # # # # # # # # # # # # # # # # # # # # # # # # # # # # # # 
+    def updateRigModuleList(self, name, module, *args):
+	
+	#get the new name
+	newName = cmds.textField(self.widgets["rigMod_" + name + "_settings_name"], q = True, text = True)
+	
+	#remove that item in the textScrollList
+	cmds.textScrollList(self.widgets["addedRigModulesList"], edit = True, ri = "[" + str(module) + "] " + name)
+	
+	#add the newly named item to the textScrollList
+	cmds.textScrollList(self.widgets["addedRigModulesList"], edit = True, append = "[" + str(module) + "] " + newName)
+	
+	#add new dictionary entries
+	self.widgets["rigMod_" + newName + "_settingsForm"] = self.widgets["rigMod_" + name + "_settingsForm"]
+	self.widgets["rigMod_" + newName + "_settings_parent"] = self.widgets["rigMod_" + name + "_settings_parent"]
+	
+	if module == "leaf":
+	    self.widgets["rigMod_" + newName + "_settings_translate"] = self.widgets["rigMod_" + name + "_settings_translate"]
+	    self.widgets["rigMod_" + newName + "_settings_rotate"] = self.widgets["rigMod_" + name + "_settings_rotate"]
+	    
+	if module == "dynamic":
+	    self.widgets["rigMod_" + newName + "_settings_numberInChain"] = self.widgets["rigMod_" + name + "_settings_numberInChain"]
+	    
+	
+	#show correct form
+	self.rigMod_hideSettingsForms()
+	cmds.formLayout(self.widgets["rigMod_" + newName + "_settingsForm"], edit = True, visible = True)
+	cmds.textScrollList(self.widgets["addedRigModulesList"], edit = True, si = "[" + str(module) + "] " + newName)
+	
+	#search through all forms, looking for any parents that may have had the old name
+	for widget in self.widgets:
+	    if widget.find("rigMod_") == 0:
+		if widget.find("_settings_parent") != -1:
+		    
+		    if widget.partition("rigMod_")[2].find(name) != 0:
+			parent = cmds.textField(self.widgets[widget], q = True, text = True)
+			if parent == name:
+			    cmds.textField(self.widgets[widget], edit = True, text = newName)
+
+	
+    # # # # # # # # # # # # # # # # # # # # # # # # # # # # # # # # # # # # # # # # # # # # # # # # # # # # # # # # # # # # # # # 
+    # # # # # # # # # # # # # # # # # # # # # # # # # # # # # # # # # # # # # # # # # # # # # # # # # # # # # # # # # # # # # # # 
+    # # # # # # # # # # # # # # # # # # # # # # # # # # # # # # # # # # # # # # # # # # # # # # # # # # # # # # # # # # # # # # # 
+    # # # # # # # # # # # # # # # # # # # # # # # # # # # # # # # # # # # # # # # # # # # # # # # # # # # # # # # # # # # # # # # 
+    # # # # # # # # # # # # # # # # # # # # # # # # # # # # # # # # # # # # # # # # # # # # # # # # # # # # # # # # # # # # # # # 
+    def extraJointInfo_UI_ParentList(self, name, *args):
+	
+	#create a list of all the joints in a UI
+	joints = self.extraJointInfo_UI_GetJoints(None)
+	
+
+	
+	window = cmds.window("extraJoint_ChooseNewParent_Window", w = 200, h = 300, title = "Select Parent", mnb = False, mxb = False)
+	mainLayout = cmds.columnLayout(rs = 10, co = ["both", 10])
+	jointList = cmds.textScrollList("extraJoint_ParentTextScrollList",  w = 180, h = 250, docTag = name)
+	
+	for joint in joints:
+	    if joint != name:
+		cmds.textScrollList(jointList, edit = True, append = joint)
+	    
+	cmds.button(w = 180, label = "Select", c = self.extraJointInfo_UI_ParentList_Accept)
+	
+	
+	cmds.showWindow(window)
+	
+	
+    # # # # # # # # # # # # # # # # # # # # # # # # # # # # # # # # # # # # # # # # # # # # # # # # # # # # # # # # # # # # # # # 
+    # # # # # # # # # # # # # # # # # # # # # # # # # # # # # # # # # # # # # # # # # # # # # # # # # # # # # # # # # # # # # # # 
+    # # # # # # # # # # # # # # # # # # # # # # # # # # # # # # # # # # # # # # # # # # # # # # # # # # # # # # # # # # # # # # # 
+    # # # # # # # # # # # # # # # # # # # # # # # # # # # # # # # # # # # # # # # # # # # # # # # # # # # # # # # # # # # # # # # 
+    # # # # # # # # # # # # # # # # # # # # # # # # # # # # # # # # # # # # # # # # # # # # # # # # # # # # # # # # # # # # # # # 
+    def extraJointInfo_UI_ParentList_Accept(self, *args):
+	
+	try:
+	    newParent = cmds.textScrollList("extraJoint_ParentTextScrollList", q = True, si = True)[0]
+	except:
+	    cmds.warning("Nothing Selected.")
+	    cmds.deleteUI("extraJoint_ChooseNewParent_Window")
+	    return
+	
+	name = cmds.textScrollList("extraJoint_ParentTextScrollList", q = True, docTag = True)
+	cmds.textField(self.widgets["rigMod_" + name + "_settings_parent"], edit = True, text = newParent)
+	
+	
+	
+	cmds.deleteUI("extraJoint_ChooseNewParent_Window")
+	
+	
+	
+    # # # # # # # # # # # # # # # # # # # # # # # # # # # # # # # # # # # # # # # # # # # # # # # # # # # # # # # # # # # # # # # 
+    # # # # # # # # # # # # # # # # # # # # # # # # # # # # # # # # # # # # # # # # # # # # # # # # # # # # # # # # # # # # # # # 
+    # # # # # # # # # # # # # # # # # # # # # # # # # # # # # # # # # # # # # # # # # # # # # # # # # # # # # # # # # # # # # # # 
+    # # # # # # # # # # # # # # # # # # # # # # # # # # # # # # # # # # # # # # # # # # # # # # # # # # # # # # # # # # # # # # # 
+    # # # # # # # # # # # # # # # # # # # # # # # # # # # # # # # # # # # # # # # # # # # # # # # # # # # # # # # # # # # # # # # 
+    def rigMod_hideSettingsForms(self, *args):
+	
+	cmds.formLayout(self.widgets["rigModuleDefaultSettingsLayout"], edit = True, visible = False)
+	
+	for widget in self.widgets:
+	    if widget.find("rigMod_") == 0:
+		if widget.find("_settingsForm") != -1:
+		    try:
+			cmds.formLayout(self.widgets[widget], edit = True, visible = False)
+		    except:
+			pass
+			
+		    
+		    
+    # # # # # # # # # # # # # # # # # # # # # # # # # # # # # # # # # # # # # # # # # # # # # # # # # # # # # # # # # # # # # # # 
+    # # # # # # # # # # # # # # # # # # # # # # # # # # # # # # # # # # # # # # # # # # # # # # # # # # # # # # # # # # # # # # # 
+    # # # # # # # # # # # # # # # # # # # # # # # # # # # # # # # # # # # # # # # # # # # # # # # # # # # # # # # # # # # # # # # 
+    # # # # # # # # # # # # # # # # # # # # # # # # # # # # # # # # # # # # # # # # # # # # # # # # # # # # # # # # # # # # # # # 
+    # # # # # # # # # # # # # # # # # # # # # # # # # # # # # # # # # # # # # # # # # # # # # # # # # # # # # # # # # # # # # # # 
+    def rigMod_changeSettingsForms(self, *args):
+	
+	#get selected item from the list
+	selected = cmds.textScrollList(self.widgets["addedRigModulesList"], q = True, si = True)[0]
+	selected = selected.rpartition("] ")[2]
+	self.rigMod_hideSettingsForms()
+	
+	cmds.formLayout(self.widgets["rigMod_" + selected + "_settingsForm"], edit = True, visible = True)
+    # # # # # # # # # # # # # # # # # # # # # # # # # # # # # # # # # # # # # # # # # # # # # # # # # # # # # # # # # # # # # # # 
+    # # # # # # # # # # # # # # # # # # # # # # # # # # # # # # # # # # # # # # # # # # # # # # # # # # # # # # # # # # # # # # # 
+    # # # # # # # # # # # # # # # # # # # # # # # # # # # # # # # # # # # # # # # # # # # # # # # # # # # # # # # # # # # # # # # 
+    # # # # # # # # # # # # # # # # # # # # # # # # # # # # # # # # # # # # # # # # # # # # # # # # # # # # # # # # # # # # # # # 
+    # # # # # # # # # # # # # # # # # # # # # # # # # # # # # # # # # # # # # # # # # # # # # # # # # # # # # # # # # # # # # # # 
+    def extraJointInfo_UI_Cancel(self, *args):
+	
+	cmds.deleteUI("extraJointInfo_UI")
+	
+
+     
+
+    # # # # # # # # # # # # # # # # # # # # # # # # # # # # # # # # # # # # # # # # # # # # # # # # # # # # # # # # # # # # # # # 
+    # # # # # # # # # # # # # # # # # # # # # # # # # # # # # # # # # # # # # # # # # # # # # # # # # # # # # # # # # # # # # # # 
+    # # # # # # # # # # # # # # # # # # # # # # # # # # # # # # # # # # # # # # # # # # # # # # # # # # # # # # # # # # # # # # # 
+    # # # # # # # # # # # # # # # # # # # # # # # # # # # # # # # # # # # # # # # # # # # # # # # # # # # # # # # # # # # # # # # 
+    # # # # # # # # # # # # # # # # # # # # # # # # # # # # # # # # # # # # # # # # # # # # # # # # # # # # # # # # # # # # # # # 
+    def saveTemplate(self, *args):
+        
+        if not os.path.exists(self.mayaToolsDir + "/General/ART/SkeletonTemplates"):
+	    os.makedirs(self.mayaToolsDir + "/General/ART/SkeletonTemplates")
+	    
+	fileName = cmds.fileDialog2(startingDirectory = self.mayaToolsDir + "/General/ART/SkeletonTemplates/", ff = "*.txt", fm = 0, okCaption = "Save")
+	
+	if fileName != None:
+	    #check path to make sure the path is in fact, in the JointMoverTemplates directory
+	    if fileName[0].find("SkeletonTemplates") != -1:
+		
+		if fileName[0].rpartition(".")[2] != "txt":
+		    self.writeToTemplate(fileName[0] + ".txt")
+		    
+		else:
+		    self.writeToTemplate(fileName[0])
+		    
+		
+		
+	    else:
+		cmds.confirmDialog(title = "Save Template", icon = "warning", message = "Please Save the template in the SkeletonTemplates directory.")
+		return
+	    
+	    
+    
+            
+    # # # # # # # # # # # # # # # # # # # # # # # # # # # # # # # # # # # # # # # # # # # # # # # # # # # # # # # # # # # # # # # 
+    # # # # # # # # # # # # # # # # # # # # # # # # # # # # # # # # # # # # # # # # # # # # # # # # # # # # # # # # # # # # # # # 
+    # # # # # # # # # # # # # # # # # # # # # # # # # # # # # # # # # # # # # # # # # # # # # # # # # # # # # # # # # # # # # # # 
+    # # # # # # # # # # # # # # # # # # # # # # # # # # # # # # # # # # # # # # # # # # # # # # # # # # # # # # # # # # # # # # # 
+    # # # # # # # # # # # # # # # # # # # # # # # # # # # # # # # # # # # # # # # # # # # # # # # # # # # # # # # # # # # # # # # 
+    def writeToTemplate(self, path):
+	
+	sourceControl = False
+	#check if user has source control on
+	settingsLocation = self.mayaToolsDir + "/General/Scripts/projectSettings.txt"
+	if os.path.exists(settingsLocation):
+	    f = open(settingsLocation, 'r')
+	    settings = cPickle.load(f)
+	    f.close()	    
+	    sourceControl = settings.get("UseSourceControl")
+	    
+	    
+	fileCheckedOut = False
+	
+        #write out a dictionary with the following format uiWidget(key):[value, bgc]
+	try:
+	    f = open(path, 'w')
+	
+	except:
+	    
+	    if sourceControl == False:
+		cmds.confirmDialog(title = "Save Template", icon = "critical", message = "Could not save file. File may exist already and be marked as read only.")
+		
+		return
+	    
+	    else:
+		#check to see if the file is currently checked out
+		result = cmds.confirmDialog(title = "Save Template", icon = "critical", message = "Could not save file. File may exist already and be marked as read only.", button = ["Check Out File", "Cancel"])
+		
+		if result == "Check Out File":
+		    import perforceUtils
+		    reload(perforceUtils)
+		    perforceUtils.p4_checkOutCurrentFile(path)
+		    
+		    try:
+			f = open(path, 'w')
+			fileCheckedOut = True
+			
+		    except:
+			cmds.confirmDialog(title = "Save Template", icon = "critical", message = "Perforce operation unsucessful. Could not save file.")
+			
+
+	
+	
+        
+        #create a dictionary for the values of our UI widgets
+	uiControls = {}
+	
+	#make a list of all of the skeleton builder ui widgets
+	intFields = [
+	    "torsoOptions_numSpineBones", "torsoOptions_numNeckBones",
+	    "leftArmOptions_numUpArmTwistBones", "leftArmOptions_numLowArmTwistBones", "leftHandOptions_numThumbBones", "leftHandOptions_numIndexBones", "leftHandOptions_numMiddleBones", "leftHandOptions_numRingBones", "leftHandOptions_numPinkyBones",
+	    "rightArmOptions_numUpArmTwistBones", "rightArmOptions_numLowArmTwistBones", "rightHandOptions_numThumbBones", "rightHandOptions_numIndexBones", "rightHandOptions_numMiddleBones", "rightHandOptions_numRingBones", "rightHandOptions_numPinkyBones",
+	    "leftLegOptions_numThighTwistBones", "leftLegOptions_numCalfTwistBones", "leftFootOptions_numThumbBones", "leftFootOptions_numIndexBones", "leftFootOptions_numMiddleBones", "leftFootOptions_numRingBones", "leftFootOptions_numPinkyBones",
+	    "rightLegOptions_numThighTwistBones", "rightLegOptions_numCalfTwistBones", "rightFootOptions_numThumbBones", "rightFootOptions_numIndexBones", "rightFootOptions_numMiddleBones", "rightFootOptions_numRingBones", "rightFootOptions_numPinkyBones"]
+	
+	checkBoxes = ["leftLegIncludeBallJoint_CB", "rightLegIncludeBallJoint_CB"]
+	
+	
+	#get values
+	for field in intFields:
+	    value = cmds.intField(self.widgets[field], q = True, v = True)
+	    color = cmds.intField(self.widgets[field], q = True, bgc = True)
+	    uiControls[field] = [value, color]
+	    
+	for box in checkBoxes:
+	    value = cmds.checkBox(self.widgets[box], q = True, value = True)
+	    uiControls[box] = value
+	    
+	    
+	    
+	    
+        #get any extra joints
+        #find children parent/type/name
+	extraJoints = cmds.textScrollList(self.widgets["addedRigModulesList"], q = True, allItems = True)
+	name = ""
+	type = ""
+	parent = ""
+	num = -1
+	translate = ""
+	rotate = ""
+	
+	try:
+	    for joint in extraJoints:
+		enumEntry = ""
+		
+		#get extra joint information
+		name = joint.rpartition("] ")[2]
+		type = joint.partition("]")[0].partition("[")[2]
+		
+		#get joint parent
+		parent = cmds.textField(self.widgets["rigMod_" + name + "_settings_parent"], q = True, text = True)
+		
+		if type == "leaf":
+		    #get translate and rotate values
+		    translate = cmds.checkBox(self.widgets["rigMod_" + name + "_settings_translate"], q = True, value = True)
+		    rotate = cmds.checkBox(self.widgets["rigMod_" + name + "_settings_rotate"], q = True, value = True)
+		    
+	    
+		if type == "chain":
+		    #get count
+		    count = cmds.textField(self.widgets["rigMod_" + name + "_settings_numberInChain"], q = True, text = True)
+		    
+		    
+		
+	    
+		if name != "":
+		    if parent != "":
+			if type != "":
+			    
+			    #LEAF JOINTS
+			    if type == "leaf":
+				if translate != "":
+				    if rotate != "":
+					if translate == True and rotate == True:
+					    enumEntry = parent + "/" + type + "/" + name + " (TR)"
+					    
+					if translate == False and rotate == True:
+					    enumEntry = parent + "/" + type + "/" + name + " (R)"
+					    
+					if translate == True and rotate == False:
+					    enumEntry = parent + "/" + type + "/" + name + " (T)"
+					    
+					if translate == False and rotate == False:
+					    enumEntry = parent + "/" + type + "/" + name + " (TR)"
+					    cmds.warning(message= "Cannot create a leaf joint with no attributes available for animation. Defaulting to Translate and Rotate.")
+		
+					uiControls["extraJoint/" + enumEntry] = enumEntry
+					name = ""
+					type = ""
+					parent = ""
+					count = -1
+					translate = ""
+					rotate = ""
+					
+			    #JIGGLE JOINTS
+			    if type == "jiggle":
+				enumEntry = parent + "/" + type + "/" + name		    
+				uiControls["extraJoint/" + enumEntry] = enumEntry
+				name = ""
+				type = ""
+				parent = ""
+				count = -1
+				translate = ""
+				rotate = ""				
+	    
+			    #DYNAMIC CHAINS
+			    if type == "chain":
+				if count > 0:
+				    enumEntry = parent + "/" + type + "/" + name + " (" + str(count) + ")"
+				    uiControls["extraJoint/" + enumEntry] = enumEntry
+				    name = ""
+				    type = ""
+				    parent = ""
+				    count = -1
+				    translate = ""
+				    rotate = ""	
+
+	except:
+	    pass
+	
+
+	    
+	#write our dictionary to file
+	cPickle.dump(uiControls, f)
+	f.close()
+	
+
+	if sourceControl == True:
+	    
+	    try:
+		import perforceUtils
+		reload(perforceUtils)
+		
+		if fileCheckedOut == True:
+		    result = cmds.promptDialog(title = "Perforce", message = "Please Enter a Description..", button = ["Accept", "Cancel"], defaultButton = "Accept", dismissString = "Cancel", cancelButton = "Cancel")
+		    if result == "Accept":
+			desc = cmds.promptDialog(q = True, text = True)
+			
+			result = perforceUtils.p4_submitCurrentFile(path, desc)
+			
+			
+			if result == False:
+			    result = cmds.confirmDialog(title = "Save Template", icon = "question", message = "Would you like to add this template to Perforce?", button = ["Yes", "No"])
+			    
+			    if result == "Yes":
+				import perforceUtils
+				reload(perforceUtils)
+				perforceUtils.p4_addAndSubmitCurrentFile(path, desc)
+				
+			
+		else:
+		    if sourceControl == True:
+			result = cmds.confirmDialog(title = "Save Template", icon = "question", message = "Would you like to add this template to Perforce?", button = ["Yes", "No"])
+			
+			if result == "Yes":
+			    
+			    result = cmds.promptDialog(title = "Perforce", message = "Please Enter a Description..", button = ["Accept", "Cancel"], defaultButton = "Accept", dismissString = "Cancel", cancelButton = "Cancel")
+			    if result == "Accept":
+				desc = cmds.promptDialog(q = True, text = True)			    
+			    
+			    import perforceUtils
+			    reload(perforceUtils)
+			    perforceUtils.p4_addAndSubmitCurrentFile(path, desc)
+			    
+			    
+			    
+	    
+	    
+	
+	    except:
+		cmds.confirmDialog(title = "Save Template", icon = "critical", message = "Could not save file. File may exist already and be marked as read only.")
+		return
+	
+	
+	cmds.headsUpMessage( 'Template Saved' )
+	
+        
+        
+
+	
+    # # # # # # # # # # # # # # # # # # # # # # # # # # # # # # # # # # # # # # # # # # # # # # # # # # # # # # # # # # # # # # # 
+    # # # # # # # # # # # # # # # # # # # # # # # # # # # # # # # # # # # # # # # # # # # # # # # # # # # # # # # # # # # # # # # 
+    # # # # # # # # # # # # # # # # # # # # # # # # # # # # # # # # # # # # # # # # # # # # # # # # # # # # # # # # # # # # # # # 
+    # # # # # # # # # # # # # # # # # # # # # # # # # # # # # # # # # # # # # # # # # # # # # # # # # # # # # # # # # # # # # # # 
+    # # # # # # # # # # # # # # # # # # # # # # # # # # # # # # # # # # # # # # # # # # # # # # # # # # # # # # # # # # # # # # # 
+    def loadTemplateUI(self, *args):
+        
+        if cmds.window("loadTemplateUI", exists = True):
+            cmds.deleteUI("loadTemplateUI")
+            
+        window = cmds.window("loadTemplateUI", w = 300, h = 200, title = "Load Template", titleBarMenu = False, sizeable = False)
+        mainLayout = cmds.columnLayout(w = 300, h = 300)
+        menuBarLayout = cmds.menuBarLayout(h = 30, w = 300, ebg = True, bgc = [.1, .1, .1], parent = mainLayout)
+        cmds.menu(label = "File")
+        cmds.menuItem(label = "Remove Selected Template", c = self.removeSelectedTemplate)
+        
+        #banner image
+	cmds.image(w = 300, h = 50, image = self.mayaToolsDir + "/General/Icons/ART/artBanner300px.bmp", parent = mainLayout)
+        
+        cmds.textScrollList("templateTextScrollList", w = 300, h = 150, numberOfRows = 100, allowMultiSelection = False, parent = mainLayout)
+	
+	#buttons + layout
+	cmds.separator(h = 10, parent = mainLayout)
+	buttonLayout = cmds.rowColumnLayout(parent = mainLayout, w = 300, rs = [1, 10], nc = 2, cat = [(1, "left", 20), (2, "both", 20)])
+        cmds.button(label = "Load Template", w = 120, h = 50, c = self.loadTemplate, parent = buttonLayout)
+	cmds.button(label = "Cancel", w = 120, h = 50, c = self.exitLoadTemplateUI, parent = buttonLayout)
+        cmds.separator(h = 10, parent = mainLayout)
+        
+        #get out template files
+	if not os.path.exists(self.mayaToolsDir + "/General/ART/SkeletonTemplates"):
+	    os.makedirs(self.mayaToolsDir + "/General/ART/SkeletonTemplates")
+	    
+        files = os.listdir(self.mayaToolsDir + "/General/ART/SkeletonTemplates/")
+        
+        for file in files:
+            niceName = file.rpartition(".")[0]
+            cmds.textScrollList("templateTextScrollList", edit = True, append = niceName)
+        
+        cmds.showWindow(window)
+        
+        
+    # # # # # # # # # # # # # # # # # # # # # # # # # # # # # # # # # # # # # # # # # # # # # # # # # # # # # # # # # # # # # # # 
+    # # # # # # # # # # # # # # # # # # # # # # # # # # # # # # # # # # # # # # # # # # # # # # # # # # # # # # # # # # # # # # # 
+    # # # # # # # # # # # # # # # # # # # # # # # # # # # # # # # # # # # # # # # # # # # # # # # # # # # # # # # # # # # # # # # 
+    # # # # # # # # # # # # # # # # # # # # # # # # # # # # # # # # # # # # # # # # # # # # # # # # # # # # # # # # # # # # # # # 
+    # # # # # # # # # # # # # # # # # # # # # # # # # # # # # # # # # # # # # # # # # # # # # # # # # # # # # # # # # # # # # # # 
+    def exitLoadTemplateUI(self,*args):
+        cmds.deleteUI("loadTemplateUI")
+        
+        
+    # # # # # # # # # # # # # # # # # # # # # # # # # # # # # # # # # # # # # # # # # # # # # # # # # # # # # # # # # # # # # # # 
+    # # # # # # # # # # # # # # # # # # # # # # # # # # # # # # # # # # # # # # # # # # # # # # # # # # # # # # # # # # # # # # # 
+    # # # # # # # # # # # # # # # # # # # # # # # # # # # # # # # # # # # # # # # # # # # # # # # # # # # # # # # # # # # # # # # 
+    # # # # # # # # # # # # # # # # # # # # # # # # # # # # # # # # # # # # # # # # # # # # # # # # # # # # # # # # # # # # # # # 
+    # # # # # # # # # # # # # # # # # # # # # # # # # # # # # # # # # # # # # # # # # # # # # # # # # # # # # # # # # # # # # # # 
+    def removeSelectedTemplate(self, *args):
+        
+        #find the selected template 
+        selected = cmds.textScrollList("templateTextScrollList", q = True, si = True)[0]
+        cmds.textScrollList("templateTextScrollList", edit = True, ri = selected)
+        
+        templatePath = self.mayaToolsDir + "/General/ART/SkeletonTemplates/" + selected + ".txt"
+        os.remove(templatePath)
+        
+        
+    # # # # # # # # # # # # # # # # # # # # # # # # # # # # # # # # # # # # # # # # # # # # # # # # # # # # # # # # # # # # # # # 
+    # # # # # # # # # # # # # # # # # # # # # # # # # # # # # # # # # # # # # # # # # # # # # # # # # # # # # # # # # # # # # # # 
+    # # # # # # # # # # # # # # # # # # # # # # # # # # # # # # # # # # # # # # # # # # # # # # # # # # # # # # # # # # # # # # # 
+    # # # # # # # # # # # # # # # # # # # # # # # # # # # # # # # # # # # # # # # # # # # # # # # # # # # # # # # # # # # # # # # 
+    # # # # # # # # # # # # # # # # # # # # # # # # # # # # # # # # # # # # # # # # # # # # # # # # # # # # # # # # # # # # # # # 
+    def loadTemplate(self, *args):
+    
+        #get the selected template from the UI
+        if cmds.textScrollList("templateTextScrollList", q = True, exists = True):
+	    try:
+		selectedFile = cmds.textScrollList("templateTextScrollList", q = True, si = True)[0]
+		path = self.mayaToolsDir + "/General/ART/SkeletonTemplates/" + str(selectedFile) + ".txt"
+		
+	    except:
+		cmds.confirmDialog(title = "Load Template", icon = "warning", message = "Please select a template from the list.")
+		return
+		
+            
+        else:
+            #if the load UI doesn't exist, that means we are calling this from the edit function so we'll be using the cache file
+            path = self.mayaToolsDir + "/General/ART/templateCache.txt"
+        
+        
+        #delete the ui
+        if cmds.window("loadTemplateUI", exists = True):
+            cmds.deleteUI("loadTemplateUI")
+        
+        
+        #open the text file and start reading info
+        file = open(path, 'r')
+	uiControls = cPickle.load(file)
+        
+        #set values
+	for control in uiControls:
+	    if control.find("extraJoint") != 0:
+		if cmds.intField(self.widgets[control], q = True,exists = True):
+		    cmds.intField(self.widgets[control], edit = True, value = uiControls[control][0])
+		    cmds.intField(self.widgets[control], edit = True, bgc = uiControls[control][1])
+        
+	for control in uiControls:
+	    if control.find("extraJoint") != 0:
+		if cmds.checkBox(self.widgets[control], q = True,exists = True):
+		    cmds.checkBox(self.widgets[control], edit = True, value = uiControls[control])
+
+        
+        
+        
+        
+	#first, remove any extra rig bones in the UI
+	cmds.textScrollList(self.widgets["addedRigModulesList"], e = True, removeAll = True)
+	
+	cmds.formLayout(self.widgets["rigModuleDefaultSettingsLayout"], edit = True, visible = True)
+	for widget in self.widgets:
+	    if widget.find("rigMod_") == 0:
+		if widget.find("_settingsForm") != -1:
+		    try:
+			cmds.deleteUI(self.widgets[widget])
+		    except:
+			pass
+		    
+	    
+        
+        #add any extra child bones
+	extraJoints = []
+        for item in uiControls:
+            if item.find("extraJoint") == 0:
+                extraJoints.append(item)
+		
+	for joint in extraJoints:
+	    info = uiControls.get(joint)
+	    type = info.partition("/")[2].partition("/")[0]
+	    parent = info.partition("/")[0]
+	    jointName = info.rpartition("/")[2]
+	    
+	    	    
+	    if type == "leaf":
+		name = jointName.partition(" (")[0]
+		data = jointName.partition(" (")[2].partition(")")[0]
+		
+		translation = False
+		rotation = False
+		
+		#get checkbox values
+		if data.find("T") != -1:
+		    translation = True
+		if data.find("R") != -1:
+		    rotation = True
+	
+		
+		
+		#add to the list of added rig modules
+		cmds.textScrollList(self.widgets["addedRigModulesList"], edit = True, append = "[leaf] " + name, sc = self.rigMod_changeSettingsForms)
+		#select the newly added item
+		cmds.textScrollList(self.widgets["addedRigModulesList"], edit = True, si = "[leaf] " + name)
+		
+		#create the entry info and parent to the self.widgets["rigModuleSettingsFrame"]
+		#for leaf joints, we need name, parent, translate/rotate
+		
+		#hide existing formLayouts
+		self.rigMod_hideSettingsForms()
+		
+		self.widgets["rigMod_" + name + "_settingsForm"] = cmds.formLayout(w = 150, h = 220, parent = self.widgets["rigModuleSettingsFrame"])
+		label_name = cmds.text(label = "Name: ")
+		self.widgets["rigMod_" + name + "_settings_name"] = cmds.textField(w = 100, parent = self.widgets["rigMod_" + name + "_settingsForm"], text = name, cc = partial(self.updateRigModuleList, name, "leaf"))
+		label_parent = cmds.text(label = "Parent: ")
+		self.widgets["rigMod_" + name + "_settings_parent"] = cmds.textField(w = 80, parent = self.widgets["rigMod_" + name + "_settingsForm"], text = parent, editable = False)
+		self.widgets["rigMod_" + name + "_settings_changeParent"] = cmds.button(label = "...", w = 20, parent = self.widgets["rigMod_" + name + "_settingsForm"], ann = "Change Parent", c = partial(self.extraJointInfo_UI_ParentList, name))
+		self.widgets["rigMod_" + name + "_settings_translate"] = cmds.checkBox(label = "Translate", v = translation, ann = "Joint can Translate")
+		self.widgets["rigMod_" + name + "_settings_rotate"] = cmds.checkBox(label = "Rotate", v = rotation, ann = "Joint can Rotate")
+		
+		cmds.formLayout(self.widgets["rigMod_" + name + "_settingsForm"], edit = True, af = [(label_name, "top", 5),(label_name, "left", 5)])
+		cmds.formLayout(self.widgets["rigMod_" + name + "_settingsForm"], edit = True, af = [(self.widgets["rigMod_" + name + "_settings_name"], "top", 5),(self.widgets["rigMod_" + name + "_settings_name"], "right", 5)])
+		cmds.formLayout(self.widgets["rigMod_" + name + "_settingsForm"], edit = True, af = [(label_parent, "top", 35),(label_parent, "left", 5)])
+		cmds.formLayout(self.widgets["rigMod_" + name + "_settingsForm"], edit = True, af = [(self.widgets["rigMod_" + name + "_settings_parent"], "top", 35),(self.widgets["rigMod_" + name + "_settings_parent"], "right", 25)])	    
+		cmds.formLayout(self.widgets["rigMod_" + name + "_settingsForm"], edit = True, af = [(self.widgets["rigMod_" + name + "_settings_changeParent"], "top", 35),(self.widgets["rigMod_" + name + "_settings_changeParent"], "right", 5)])	    
+		cmds.formLayout(self.widgets["rigMod_" + name + "_settingsForm"], edit = True, af = [(self.widgets["rigMod_" + name + "_settings_translate"], "top", 65),(self.widgets["rigMod_" + name + "_settings_translate"], "left", 5)])	    
+		cmds.formLayout(self.widgets["rigMod_" + name + "_settingsForm"], edit = True, af = [(self.widgets["rigMod_" + name + "_settings_rotate"], "top", 85),(self.widgets["rigMod_" + name + "_settings_rotate"], "left", 5)])	    
+		
+	
+	
+	    if type == "jiggle":
+		name = jointName.partition(" (")[0]
+		
+		#add to the list of added rig modules
+		cmds.textScrollList(self.widgets["addedRigModulesList"], edit = True, append = "[jiggle] " + name, sc = self.rigMod_changeSettingsForms)
+		
+		#select the newly added item
+		cmds.textScrollList(self.widgets["addedRigModulesList"], edit = True, si = "[jiggle] " + name)
+		
+		#create the entry info and parent to the self.widgets["rigModuleSettingsFrame"]
+		
+		#hide existing formLayouts
+		self.rigMod_hideSettingsForms()
+		
+		self.widgets["rigMod_" + name + "_settingsForm"] = cmds.formLayout(w = 150, h = 220, parent = self.widgets["rigModuleSettingsFrame"])
+		label_name = cmds.text(label = "Name: ")
+		self.widgets["rigMod_" + name + "_settings_name"] = cmds.textField(w = 100, parent = self.widgets["rigMod_" + name + "_settingsForm"], text = name)
+		label_parent = cmds.text(label = "Parent: ")
+		self.widgets["rigMod_" + name + "_settings_parent"] = cmds.textField(w = 80, parent = self.widgets["rigMod_" + name + "_settingsForm"], text = parent, editable = False)
+		self.widgets["rigMod_" + name + "_settings_changeParent"] = cmds.button(label = "...", w = 20, parent = self.widgets["rigMod_" + name + "_settingsForm"], ann = "Change Parent", c = partial(self.extraJointInfo_UI_ParentList, name))
+		
+		cmds.formLayout(self.widgets["rigMod_" + name + "_settingsForm"], edit = True, af = [(label_name, "top", 5),(label_name, "left", 5)])
+		cmds.formLayout(self.widgets["rigMod_" + name + "_settingsForm"], edit = True, af = [(self.widgets["rigMod_" + name + "_settings_name"], "top", 5),(self.widgets["rigMod_" + name + "_settings_name"], "right", 5)])
+		cmds.formLayout(self.widgets["rigMod_" + name + "_settingsForm"], edit = True, af = [(label_parent, "top", 35),(label_parent, "left", 5)])
+		cmds.formLayout(self.widgets["rigMod_" + name + "_settingsForm"], edit = True, af = [(self.widgets["rigMod_" + name + "_settings_parent"], "top", 35),(self.widgets["rigMod_" + name + "_settings_parent"], "right", 25)])	    
+		cmds.formLayout(self.widgets["rigMod_" + name + "_settingsForm"], edit = True, af = [(self.widgets["rigMod_" + name + "_settings_changeParent"], "top", 35),(self.widgets["rigMod_" + name + "_settings_changeParent"], "right", 5)])	    
+		
+		
+	    if type == "chain":
+		name = jointName.partition(" (")[0]
+		num = jointName.partition(" (")[2].partition(")")[0]
+		
+		
+		#add to the list of added rig modules
+		cmds.textScrollList(self.widgets["addedRigModulesList"], edit = True, append = "[chain] " + name, sc = self.rigMod_changeSettingsForms)
+		#select the newly added item
+		cmds.textScrollList(self.widgets["addedRigModulesList"], edit = True, si = "[chain] " + name)
+		
+		#create the entry info and parent to the self.widgets["rigModuleSettingsFrame"]
+		#for leaf joints, we need name, parent, translate/rotate
+		
+		#hide existing formLayouts
+		self.rigMod_hideSettingsForms()
+		
+		self.widgets["rigMod_" + name + "_settingsForm"] = cmds.formLayout(w = 150, h = 220, parent = self.widgets["rigModuleSettingsFrame"])
+		label_name = cmds.text(label = "Name: ")
+		self.widgets["rigMod_" + name + "_settings_name"] = cmds.textField(w = 100, parent = self.widgets["rigMod_" + name + "_settingsForm"], text = name)
+		label_parent = cmds.text(label = "Parent: ")
+		self.widgets["rigMod_" + name + "_settings_parent"] = cmds.textField(w = 80, parent = self.widgets["rigMod_" + name + "_settingsForm"], text = parent, editable = False)
+		self.widgets["rigMod_" + name + "_settings_changeParent"] = cmds.button(label = "...", w = 20, parent = self.widgets["rigMod_" + name + "_settingsForm"], ann = "Change Parent", c = partial(self.extraJointInfo_UI_ParentList, name))
+		label_number = cmds.text(label = "Count: ")
+		self.widgets["rigMod_" + name + "_settings_numberInChain"] = cmds.textField(w = 100, parent = self.widgets["rigMod_" + name + "_settingsForm"], text = str(num), ann = "Number of joints in the chain.")
+		
+		cmds.formLayout(self.widgets["rigMod_" + name + "_settingsForm"], edit = True, af = [(label_name, "top", 5),(label_name, "left", 5)])
+		cmds.formLayout(self.widgets["rigMod_" + name + "_settingsForm"], edit = True, af = [(self.widgets["rigMod_" + name + "_settings_name"], "top", 5),(self.widgets["rigMod_" + name + "_settings_name"], "right", 5)])
+		cmds.formLayout(self.widgets["rigMod_" + name + "_settingsForm"], edit = True, af = [(label_parent, "top", 35),(label_parent, "left", 5)])
+		cmds.formLayout(self.widgets["rigMod_" + name + "_settingsForm"], edit = True, af = [(self.widgets["rigMod_" + name + "_settings_parent"], "top", 35),(self.widgets["rigMod_" + name + "_settings_parent"], "right", 25)])	    
+		cmds.formLayout(self.widgets["rigMod_" + name + "_settingsForm"], edit = True, af = [(self.widgets["rigMod_" + name + "_settings_changeParent"], "top", 35),(self.widgets["rigMod_" + name + "_settings_changeParent"], "right", 5)])	    
+		cmds.formLayout(self.widgets["rigMod_" + name + "_settingsForm"], edit = True, af = [(label_number, "top", 65),(label_number, "left", 5)])
+		cmds.formLayout(self.widgets["rigMod_" + name + "_settingsForm"], edit = True, af = [(self.widgets["rigMod_" + name + "_settings_numberInChain"], "top", 65),(self.widgets["rigMod_" + name + "_settings_numberInChain"], "right", 5)])
+
+
+    # # # # # # # # # # # # # # # # # # # # # # # # # # # # # # # # # # # # # # # # # # # # # # # # # # # # # # # # # # # # # # # 
+    # # # # # # # # # # # # # # # # # # # # # # # # # # # # # # # # # # # # # # # # # # # # # # # # # # # # # # # # # # # # # # # 
+    # # # # # # # # # # # # # # # # # # # # # # # # # # # # # # # # # # # # # # # # # # # # # # # # # # # # # # # # # # # # # # # 
+    # # # # # # # # # # # # # # # # # # # # # # # # # # # # # # # # # # # # # # # # # # # # # # # # # # # # # # # # # # # # # # # 
+    # # # # # # # # # # # # # # # # # # # # # # # # # # # # # # # # # # # # # # # # # # # # # # # # # # # # # # # # # # # # # # # 
+    def build(self, *args):
+        
+        result = cmds.confirmDialog(icon = "question", title = "Build Skeleton", message = "Are you sure you want to build the skeleton with your chosen parameters?", button = ["Yes", "No"], defaultButton = "Yes", cancelButton = "No", dismissString = "No")
+        edit = False
+        
+        if result == "Yes":
+            
+            if cmds.objExists("SkeletonSettings_Cache"):
+		cmds.delete("SkeletonSettings_Cache")
+		edit = True
+		
+            if cmds.objExists("Skeleton_Settings"):
+		cmds.lockNode("Skeleton_Settings", lock = False)
+                cmds.delete("Skeleton_Settings")
+                edit = True
+                
+
+            #change flow chart buttons:
+            imagePath = self.mayaToolsDir + "/General/Icons/ART/"
+	    
+	    cmds.symbolButton(self.widgets["previousStep"], edit = True, visible = True, c = self.edit)
+	    cmds.symbolButton(self.widgets["nextStep"], edit = True, visible = True, image = imagePath + "defSetup.bmp", c = self.lock_phase1)
+	
+            
+            #create skeleton settings node
+            skeletonSettings = cmds.group(empty = True, name = "Skeleton_Settings")
+
+            cmds.setAttr(skeletonSettings + ".tx", lock = True, keyable = False, channelBox = False)
+            cmds.setAttr(skeletonSettings + ".ty", lock = True, keyable = False, channelBox = False)
+            cmds.setAttr(skeletonSettings + ".tz", lock = True, keyable = False, channelBox = False)
+            cmds.setAttr(skeletonSettings + ".rx", lock = True, keyable = False, channelBox = False)
+            cmds.setAttr(skeletonSettings + ".ry", lock = True, keyable = False, channelBox = False)
+            cmds.setAttr(skeletonSettings + ".rz", lock = True, keyable = False, channelBox = False)
+            cmds.setAttr(skeletonSettings + ".sx", lock = True, keyable = False, channelBox = False)
+            cmds.setAttr(skeletonSettings + ".sy", lock = True, keyable = False, channelBox = False)
+            cmds.setAttr(skeletonSettings + ".sz", lock = True, keyable = False, channelBox = False)
+            cmds.setAttr(skeletonSettings + ".v", lock = True, keyable = False, channelBox = False)
+            
+            
+
+	    # # # # # # # # # # # # # # # # # # # # # # # # # # # # # # # # # # # # # # # # # # # # # # # # # # # # # # # # # # # # # # # 
+	    # # # # # # # # # # # # # # # # # # # # # # # # # # # # # # # # # # # # # # # # # # # # # # # # # # # # # # # # # # # # # # # 
+            #gather information
+            
+            #get the number of spine and neck bones
+	    numSpineBones = cmds.intField(self.widgets["torsoOptions_numSpineBones"], q = True, value = True)
+            numNeckBones = cmds.intField(self.widgets["torsoOptions_numNeckBones"], q = True, value = True)
+                    
+
+            #get the number of fingers for each hand
+            numLeftFingers = 0
+            numRightFingers = 0
+	    
+	    for field in ["HandOptions_numThumbBones", "HandOptions_numIndexBones", "HandOptions_numMiddleBones", "HandOptions_numRingBones", "HandOptions_numPinkyBones"]:
+		if cmds.intField(self.widgets["left" + field], q = True, value = True) > 0:
+		    numLeftFingers = numLeftFingers + 1
+		    
+	    for field in ["HandOptions_numThumbBones", "HandOptions_numIndexBones", "HandOptions_numMiddleBones", "HandOptions_numRingBones", "HandOptions_numPinkyBones"]:
+		if cmds.intField(self.widgets["right" + field], q = True, value = True) > 0:
+		    numRightFingers = numRightFingers + 1
+			
+            
+            #get the number of toes for each foot
+            numLeftToes = 0
+            numRightToes = 0
+	    
+	    for field in ["FootOptions_numThumbBones", "FootOptions_numIndexBones", "FootOptions_numMiddleBones", "FootOptions_numRingBones", "FootOptions_numPinkyBones"]:
+		if cmds.intField(self.widgets["left" + field], q = True, value = True) > 0:
+		    numLeftToes = numLeftToes + 1
+		    
+	    for field in ["FootOptions_numThumbBones", "FootOptions_numIndexBones", "FootOptions_numMiddleBones", "FootOptions_numRingBones", "FootOptions_numPinkyBones"]:
+		if cmds.intField(self.widgets["right" + field], q = True, value = True) > 0:
+		    numRightToes = numRightToes + 1
+	    
+
+            #add current info to settings
+            cmds.select(skeletonSettings)
+            cmds.addAttr(longName='numSpineBones', defaultValue= float(numSpineBones), minValue=1, maxValue=5, keyable = True)
+            cmds.addAttr(longName='numNeckBones', defaultValue= float(numNeckBones), minValue=1, maxValue=3, keyable = True)
+            cmds.addAttr(longName='numLeftFingers', enumName = numLeftFingers, at = "enum", keyable = True)
+            cmds.addAttr(longName='numRightFingers', enumName = numRightFingers, at = "enum", keyable = True)
+            cmds.addAttr(longName='numLeftToes', enumName = numLeftToes, at = "enum", keyable = True)
+            cmds.addAttr(longName='numRightToes', enumName = numRightToes, at = "enum", keyable = True)
+	    
+	    
+	    #get the twist joint info for the arms
+	    leftUpperArmTwist = cmds.intField(self.widgets["leftArmOptions_numUpArmTwistBones"], q = True, value = True)
+	    leftLowerArmTwist = cmds.intField(self.widgets["leftArmOptions_numLowArmTwistBones"], q = True, value = True)
+	    rightUpperArmTwist = cmds.intField(self.widgets["rightArmOptions_numUpArmTwistBones"], q = True, value = True)
+	    rightLowerArmTwist = cmds.intField(self.widgets["rightArmOptions_numLowArmTwistBones"], q = True, value = True)
+	    
+	    cmds.addAttr(longName= "leftUpperArmTwist", defaultValue= leftUpperArmTwist, keyable = True)
+	    cmds.addAttr(longName= "leftLowerArmTwist", defaultValue= leftLowerArmTwist, keyable = True)
+	    cmds.addAttr(longName= "rightUpperArmTwist", defaultValue= rightUpperArmTwist, keyable = True)
+	    cmds.addAttr(longName= "rightLowerArmTwist", defaultValue= rightLowerArmTwist, keyable = True)
+	    
+
+	    # # # # # # # # # # # # # # # # # # # # # # # # # # # # # # # # # # # # # # # # # # # # # # # # # # # # # # # # # # # # # # # 
+	    # # # # # # # # # # # # # # # # # # # # # # # # # # # # # # # # # # # # # # # # # # # # # # # # # # # # # # # # # # # # # # # 
+	    #get the finger info for the hands
+	    for attr in ["leftpinkymeta", "leftpinky1", "leftpinky2", "leftpinky3",
+                          "leftringmeta", "leftring1", "leftring2", "leftring3",
+                          "leftmiddlemeta", "leftmiddle1", "leftmiddle2", "leftmiddle3",
+                          "leftindexmeta", "leftindex1", "leftindex2", "leftindex3",
+                          "leftthumb1", "leftthumb2", "leftthumb3",
+                          "rightpinkymeta", "rightpinky1", "rightpinky2", "rightpinky3",
+                          "rightringmeta", "rightring1", "rightring2", "rightring3",
+                          "rightmiddlemeta", "rightmiddle1", "rightmiddle2", "rightmiddle3",
+                          "rightindexmeta", "rightindex1", "rightindex2", "rightindex3",
+                          "rightthumb1", "rightthumb2", "rightthumb3"]:
+		
+		cmds.addAttr(ln = attr, at = 'bool', dv = False, keyable = True)
+		
+		
+	    lIndexValue = cmds.intField(self.widgets["leftHandOptions_numIndexBones"], q = True, value = True)
+	    if lIndexValue >= 1:
+		cmds.setAttr("Skeleton_Settings.leftindex1", True)
+	    
+	    if lIndexValue >= 2:
+		cmds.setAttr("Skeleton_Settings.leftindex2", True)
+		
+	    if lIndexValue >= 3:
+		cmds.setAttr("Skeleton_Settings.leftindex3", True)
+		
+	    if lIndexValue == 4:
+		cmds.setAttr("Skeleton_Settings.leftindexmeta", True)
+		
+
+	    lMiddleValue = cmds.intField(self.widgets["leftHandOptions_numMiddleBones"], q = True, value = True)
+	    if lMiddleValue >= 1:
+		cmds.setAttr("Skeleton_Settings.leftmiddle1", True)
+	    
+	    if lMiddleValue >= 2:
+		cmds.setAttr("Skeleton_Settings.leftmiddle2", True)
+		
+	    if lMiddleValue >= 3:
+		cmds.setAttr("Skeleton_Settings.leftmiddle3", True)
+		
+	    if lMiddleValue == 4:
+		cmds.setAttr("Skeleton_Settings.leftmiddlemeta", True)
+		
+		
+	    lRingValue = cmds.intField(self.widgets["leftHandOptions_numRingBones"], q = True, value = True)
+	    if lRingValue >= 1:
+		cmds.setAttr("Skeleton_Settings.leftring1", True)
+	    
+	    if lRingValue >= 2:
+		cmds.setAttr("Skeleton_Settings.leftring2", True)
+		
+	    if lRingValue >= 3:
+		cmds.setAttr("Skeleton_Settings.leftring3", True)
+		
+	    if lRingValue == 4:
+		cmds.setAttr("Skeleton_Settings.leftringmeta", True)
+		
+		
+	    lPinkyValue = cmds.intField(self.widgets["leftHandOptions_numPinkyBones"], q = True, value = True)
+	    if lPinkyValue >= 1:
+		cmds.setAttr("Skeleton_Settings.leftpinky1", True)
+	    
+	    if lPinkyValue >= 2:
+		cmds.setAttr("Skeleton_Settings.leftpinky2", True)
+		
+	    if lPinkyValue >= 3:
+		cmds.setAttr("Skeleton_Settings.leftpinky3", True)
+		
+	    if lPinkyValue == 4:
+		cmds.setAttr("Skeleton_Settings.leftpinkymeta", True)
+		
+
+	    lThumbValue = cmds.intField(self.widgets["leftHandOptions_numThumbBones"], q = True, value = True)
+	    if lThumbValue >= 1:
+		cmds.setAttr("Skeleton_Settings.leftthumb1", True)
+	    
+	    if lThumbValue >= 2:
+		cmds.setAttr("Skeleton_Settings.leftthumb2", True)
+		
+	    if lThumbValue == 3:
+		cmds.setAttr("Skeleton_Settings.leftthumb3", True)
+	    # # # # # # # # # # # # # # # # # # # # # # # # # # # # # # # # # # # # # # # # # # # # # # # # # # # # # # # # # # # # # # # 
+	    # # # # # # # # # # # # # # # # # # # # # # # # # # # # # # # # # # # # # # # # # # # # # # # # # # # # # # # # # # # # # # # 
+	    rIndexValue = cmds.intField(self.widgets["rightHandOptions_numIndexBones"], q = True, value = True)
+	    if rIndexValue >= 1:
+		cmds.setAttr("Skeleton_Settings.rightindex1", True)
+	    
+	    if rIndexValue >= 2:
+		cmds.setAttr("Skeleton_Settings.rightindex2", True)
+		
+	    if rIndexValue >= 3:
+		cmds.setAttr("Skeleton_Settings.rightindex3", True)
+		
+	    if rIndexValue == 4:
+		cmds.setAttr("Skeleton_Settings.rightindexmeta", True)
+		
+
+	    rMiddleValue = cmds.intField(self.widgets["rightHandOptions_numMiddleBones"], q = True, value = True)
+	    if rMiddleValue >= 1:
+		cmds.setAttr("Skeleton_Settings.rightmiddle1", True)
+	    
+	    if rMiddleValue >= 2:
+		cmds.setAttr("Skeleton_Settings.rightmiddle2", True)
+		
+	    if rMiddleValue >= 3:
+		cmds.setAttr("Skeleton_Settings.rightmiddle3", True)
+		
+	    if rMiddleValue == 4:
+		cmds.setAttr("Skeleton_Settings.rightmiddlemeta", True)
+		
+		
+	    rRingValue = cmds.intField(self.widgets["rightHandOptions_numRingBones"], q = True, value = True)
+	    if rRingValue >= 1:
+		cmds.setAttr("Skeleton_Settings.rightring1", True)
+	    
+	    if rRingValue >= 2:
+		cmds.setAttr("Skeleton_Settings.rightring2", True)
+		
+	    if rRingValue >= 3:
+		cmds.setAttr("Skeleton_Settings.rightring3", True)
+		
+	    if rRingValue == 4:
+		cmds.setAttr("Skeleton_Settings.rightringmeta", True)
+		
+		
+	    rPinkyValue = cmds.intField(self.widgets["rightHandOptions_numPinkyBones"], q = True, value = True)
+	    if rPinkyValue >= 1:
+		cmds.setAttr("Skeleton_Settings.rightpinky1", True)
+	    
+	    if rPinkyValue >= 2:
+		cmds.setAttr("Skeleton_Settings.rightpinky2", True)
+		
+	    if rPinkyValue >= 3:
+		cmds.setAttr("Skeleton_Settings.rightpinky3", True)
+		
+	    if rPinkyValue == 4:
+		cmds.setAttr("Skeleton_Settings.rightpinkymeta", True)
+		
+
+	    rThumbValue = cmds.intField(self.widgets["rightHandOptions_numThumbBones"], q = True, value = True)
+	    if rThumbValue >= 1:
+		cmds.setAttr("Skeleton_Settings.rightthumb1", True)
+	    
+	    if rThumbValue >= 2:
+		cmds.setAttr("Skeleton_Settings.rightthumb2", True)
+		
+	    if rThumbValue == 3:
+		cmds.setAttr("Skeleton_Settings.rightthumb3", True)
+		
+		
+		
+	    #Leg Style
+	    legStyle = cmds.optionMenu(self.widgets["legStyleOptionMenu"], q = True, v = True)
+	    if legStyle == "Standard Biped":
+		val = 0
+	    if legStyle == "Hind Leg":
+		val = 1
+	    if legStyle == "Foreleg":
+		val = 2
+		
+	    cmds.addAttr(ln = "legStyle", dv = val, keyable = True)
+	    
+	    
+	    #Leg Twists and ball joints ["LegOptions_numThighTwistBones", "LegOptions_numCalfTwistBones"]
+	    leftUpperLegTwist = cmds.intField(self.widgets["leftLegOptions_numThighTwistBones"], q = True, value = True)
+	    leftLowerLegTwist = cmds.intField(self.widgets["leftLegOptions_numCalfTwistBones"], q = True, value = True)
+	    rightUpperLegTwist = cmds.intField(self.widgets["rightLegOptions_numThighTwistBones"], q = True, value = True)
+	    rightLowerLegTwist = cmds.intField(self.widgets["rightLegOptions_numCalfTwistBones"], q = True, value = True)
+	    leftball = cmds.checkBox(self.widgets["leftLegIncludeBallJoint_CB"], q = True, value = True)
+	    rightball = cmds.checkBox(self.widgets["rightLegIncludeBallJoint_CB"], q = True, value = True)
+	    
+	    cmds.addAttr(ln = "leftUpperLegTwist", dv = leftUpperLegTwist, keyable = True)
+	    cmds.addAttr(ln = "leftLowerLegTwist", dv = leftLowerLegTwist, keyable = True)
+	    cmds.addAttr(ln = "rightUpperLegTwist", dv = rightUpperLegTwist, keyable = True)
+	    cmds.addAttr(ln = "rightLowerLegTwist", dv = rightLowerLegTwist, keyable = True)
+	    cmds.addAttr(ln = "leftball", at = 'bool', dv = leftball, keyable = True)
+	    cmds.addAttr(ln = "rightball", at = 'bool', dv = rightball, keyable = True)
+	    # # # # # # # # # # # # # # # # # # # # # # # # # # # # # # # # # # # # # # # # # # # # # # # # # # # # # # # # # # # # # # # 
+	    # # # # # # # # # # # # # # # # # # # # # # # # # # # # # # # # # # # # # # # # # # # # # # # # # # # # # # # # # # # # # # # 
+            for attr in  [
+                          "leftFootpinkymeta", "leftFootpinky1", "leftFootpinky2", "leftFootpinky3",
+                          "leftFootringmeta", "leftFootring1", "leftFootring2", "leftFootring3",
+                          "leftFootmiddlemeta", "leftFootmiddle1", "leftFootmiddle2", "leftFootmiddle3",
+                          "leftFootindexmeta", "leftFootindex1", "leftFootindex2", "leftFootindex3",
+                          "leftFootbigtoemeta", "leftFootbigtoe1", "leftFootbigtoe2", 
+                          "rightFootpinkymeta", "rightFootpinky1", "rightFootpinky2", "rightFootpinky3",
+                          "rightFootringmeta", "rightFootring1", "rightFootring2", "rightFootring3",
+                          "rightFootmiddlemeta", "rightFootmiddle1", "rightFootmiddle2", "rightFootmiddle3",
+                          "rightFootindexmeta", "rightFootindex1", "rightFootindex2", "rightFootindex3",
+                          "rightFootbigtoemeta", "rightFootbigtoe1", "rightFootbigtoe2"
+                          ]:
+		
+		cmds.addAttr(ln = attr, at = 'bool', dv = False)
+	    # # # # # # # # # # # # # # # # # # # # # # # # # # # # # # # # # # # # # # # # # # # # # # # # # # # # # # # # # # # # # # # 
+	    # # # # # # # # # # # # # # # # # # # # # # # # # # # # # # # # # # # # # # # # # # # # # # # # # # # # # # # # # # # # # # # 
+	    lIndexToeVal = cmds.intField(self.widgets["leftFootOptions_numIndexBones"], q = True, value = True)
+	    if lIndexToeVal >= 1:
+		cmds.setAttr("Skeleton_Settings.leftFootindex1", True)
+	    
+	    if lIndexToeVal >= 2:
+		cmds.setAttr("Skeleton_Settings.leftFootindex2", True)
+		
+	    if lIndexToeVal >= 3:
+		cmds.setAttr("Skeleton_Settings.leftFootindex3", True)
+		
+	    if lIndexToeVal == 4:
+		cmds.setAttr("Skeleton_Settings.leftFootindexmeta", True)
+		
+
+	    lMiddleToeVal = cmds.intField(self.widgets["leftFootOptions_numMiddleBones"], q = True, value = True)
+	    if lMiddleToeVal >= 1:
+		cmds.setAttr("Skeleton_Settings.leftFootmiddle1", True)
+	    
+	    if lMiddleToeVal >= 2:
+		cmds.setAttr("Skeleton_Settings.leftFootmiddle2", True)
+		
+	    if lMiddleToeVal >= 3:
+		cmds.setAttr("Skeleton_Settings.leftFootmiddle3", True)
+		
+	    if lMiddleToeVal == 4:
+		cmds.setAttr("Skeleton_Settings.leftFootmiddlemeta", True)
+		
+		
+	    lRingToeVal = cmds.intField(self.widgets["leftFootOptions_numRingBones"], q = True, value = True)
+	    if lRingToeVal >= 1:
+		cmds.setAttr("Skeleton_Settings.leftFootring1", True)
+	    
+	    if lRingToeVal >= 2:
+		cmds.setAttr("Skeleton_Settings.leftFootring2", True)
+		
+	    if lRingToeVal >= 3:
+		cmds.setAttr("Skeleton_Settings.leftFootring3", True)
+		
+	    if lRingToeVal == 4:
+		cmds.setAttr("Skeleton_Settings.leftFootringmeta", True)
+		
+		
+	    lPinkyToeVal = cmds.intField(self.widgets["leftFootOptions_numPinkyBones"], q = True, value = True)
+	    if lPinkyToeVal >= 1:
+		cmds.setAttr("Skeleton_Settings.leftFootpinky1", True)
+	    
+	    if lPinkyToeVal >= 2:
+		cmds.setAttr("Skeleton_Settings.leftFootpinky2", True)
+		
+	    if lPinkyToeVal >= 3:
+		cmds.setAttr("Skeleton_Settings.leftFootpinky3", True)
+		
+	    if lPinkyToeVal == 4:
+		cmds.setAttr("Skeleton_Settings.leftFootpinkymeta", True)
+		
+		
+	    lBigToeVal = cmds.intField(self.widgets["leftFootOptions_numThumbBones"], q = True, value = True)
+	    if lBigToeVal >= 1:
+		cmds.setAttr("Skeleton_Settings.leftFootbigtoe1", True)
+	    
+	    if lBigToeVal >= 2:
+		cmds.setAttr("Skeleton_Settings.leftFootbigtoe2", True)
+		
+	    if lBigToeVal == 3:
+		cmds.setAttr("Skeleton_Settings.leftFootbigtoemeta", True)
+	    # # # # # # # # # # # # # # # # # # # # # # # # # # # # # # # # # # # # # # # # # # # # # # # # # # # # # # # # # # # # # # # 
+	    # # # # # # # # # # # # # # # # # # # # # # # # # # # # # # # # # # # # # # # # # # # # # # # # # # # # # # # # # # # # # # # 
+	    rIndexToeVal = cmds.intField(self.widgets["rightFootOptions_numIndexBones"], q = True, value = True)
+	    if rIndexToeVal >= 1:
+		cmds.setAttr("Skeleton_Settings.rightFootindex1", True)
+	    
+	    if rIndexToeVal >= 2:
+		cmds.setAttr("Skeleton_Settings.rightFootindex2", True)
+		
+	    if rIndexToeVal >= 3:
+		cmds.setAttr("Skeleton_Settings.rightFootindex3", True)
+		
+	    if rIndexToeVal == 4:
+		cmds.setAttr("Skeleton_Settings.rightFootindexmeta", True)
+		
+
+	    rMiddleToeVal = cmds.intField(self.widgets["rightFootOptions_numMiddleBones"], q = True, value = True)
+	    if rMiddleToeVal >= 1:
+		cmds.setAttr("Skeleton_Settings.rightFootmiddle1", True)
+	    
+	    if rMiddleToeVal >= 2:
+		cmds.setAttr("Skeleton_Settings.rightFootmiddle2", True)
+		
+	    if rMiddleToeVal >= 3:
+		cmds.setAttr("Skeleton_Settings.rightFootmiddle3", True)
+		
+	    if rMiddleToeVal == 4:
+		cmds.setAttr("Skeleton_Settings.rightFootmiddlemeta", True)
+		
+		
+	    rRingToeVal = cmds.intField(self.widgets["rightFootOptions_numRingBones"], q = True, value = True)
+	    if rRingToeVal >= 1:
+		cmds.setAttr("Skeleton_Settings.rightFootring1", True)
+	    
+	    if rRingToeVal >= 2:
+		cmds.setAttr("Skeleton_Settings.rightFootring2", True)
+		
+	    if rRingToeVal >= 3:
+		cmds.setAttr("Skeleton_Settings.rightFootring3", True)
+		
+	    if rRingToeVal == 4:
+		cmds.setAttr("Skeleton_Settings.rightFootringmeta", True)
+		
+		
+	    rPinkyToeVal = cmds.intField(self.widgets["rightFootOptions_numPinkyBones"], q = True, value = True)
+	    if rPinkyToeVal >= 1:
+		cmds.setAttr("Skeleton_Settings.rightFootpinky1", True)
+	    
+	    if rPinkyToeVal >= 2:
+		cmds.setAttr("Skeleton_Settings.rightFootpinky2", True)
+		
+	    if rPinkyToeVal >= 3:
+		cmds.setAttr("Skeleton_Settings.rightFootpinky3", True)
+		
+	    if rPinkyToeVal == 4:
+		cmds.setAttr("Skeleton_Settings.rightFootpinkymeta", True)
+		
+		
+	    rBigToeVal = cmds.intField(self.widgets["rightFootOptions_numThumbBones"], q = True, value = True)
+	    if rBigToeVal >= 1:
+		cmds.setAttr("Skeleton_Settings.rightFootbigtoe1", True)
+	    
+	    if rBigToeVal >= 2:
+		cmds.setAttr("Skeleton_Settings.rightFootbigtoe2", True)
+		
+	    if rBigToeVal == 3:
+		cmds.setAttr("Skeleton_Settings.rightFootbigtoemeta", True)
+	    
+	    
+	    # # # # # # # # # # # # # # # # # # # # # # # # # # # # # # # # # # # # # # # # # # # # # # # # # # # # # # # # # # # # # # # 
+	    # # # # # # # # # # # # # # # # # # # # # # # # # # # # # # # # # # # # # # # # # # # # # # # # # # # # # # # # # # # # # # # 
+	    #add extra joint info
+	    extraJoints = cmds.textScrollList(self.widgets["addedRigModulesList"], q = True, allItems = True)
+	    name = ""
+	    type = ""
+	    parent = ""
+	    num = -1
+	    translate = ""
+	    rotate = ""
+	    
+	    if extraJoints != None:
+		for joint in extraJoints:
+		    enumEntry = ""
+		    
+		    #get extra joint information
+		    name = joint.rpartition("] ")[2]
+		    type = joint .partition("]")[0].partition("[")[2]
+		    
+		    #get joint parent
+		    parent = cmds.textField(self.widgets["rigMod_" + name + "_settings_parent"], q = True, text = True)
+		    
+		    if type == "leaf":
+			#get translate and rotate values
+			translate = cmds.checkBox(self.widgets["rigMod_" + name + "_settings_translate"], q = True, value = True)
+			rotate = cmds.checkBox(self.widgets["rigMod_" + name + "_settings_rotate"], q = True, value = True)
+			
+    
+		    if type == "chain":
+			#get count
+			count = cmds.textField(self.widgets["rigMod_" + name + "_settings_numberInChain"], q = True, text = True)
+			
+			
+		    
+     
+		    if name != "":
+			if parent != "":
+			    if type != "":
+				
+				#LEAF JOINTS
+				if type == "leaf":
+				    if translate != "":
+					if rotate != "":
+					    if translate == True and rotate == True:
+						enumEntry = parent + "/" + type + "/" + name + " (TR)"
+						
+					    if translate == False and rotate == True:
+						enumEntry = parent + "/" + type + "/" + name + " (R)"
+						
+					    if translate == True and rotate == False:
+						enumEntry = parent + "/" + type + "/" + name + " (T)"
+						
+					    if translate == False and rotate == False:
+						enumEntry = parent + "/" + type + "/" + name + " (TR)"
+						cmds.warning(message= "Cannot create a leaf joint with no attributes available for animation. Defaulting to Translate and Rotate.")
+		    
+					    cmds.addAttr("Skeleton_Settings", longName = "extraJoint" + name, enumName = enumEntry, at = "enum", keyable = True)
+					    name = ""
+					    type = ""
+					    parent = ""
+					    count = -1
+					    translate = ""
+					    rotate = ""
+					    
+				#JIGGLE JOINTS
+				if type == "jiggle":
+				    enumEntry = parent + "/" + type + "/" + name		    
+				    cmds.addAttr("Skeleton_Settings", longName = "extraJoint" + name, enumName = enumEntry, at = "enum", keyable = True)
+				    name = ""
+				    type = ""
+				    parent = ""
+				    count = -1
+				    translate = ""
+				    rotate = ""				
+    
+				#DYNAMIC CHAINS
+				if type == "chain":
+				    if count > 0:
+					enumEntry = parent + "/" + type + "/" + name + " (" + str(count) + ")"
+					cmds.addAttr("Skeleton_Settings", longName = "extraJoint" + name, enumName = enumEntry, at = "enum", keyable = True)
+					name = ""
+					type = ""
+					parent = ""
+					count = -1
+					translate = ""
+					rotate = ""	
+
+		
+		
+		
+
+
+
+
+	    # # # # # # # # # # # # # # # # # # # # # # # # # # # # # # # # # # # # # # # # # # # # # # # # # # # # # # # # # # # # # # # 
+	    # # # # # # # # # # # # # # # # # # # # # # # # # # # # # # # # # # # # # # # # # # # # # # # # # # # # # # # # # # # # # # # 
+	    #lock down skeleton settings
+	    attrs = cmds.listAttr("Skeleton_Settings", keyable = True)
+	    
+	    for attr in attrs:
+		cmds.setAttr("Skeleton_Settings." + attr, lock = True)
+	    # # # # # # # # # # # # # # # # # # # # # # # # # # # # # # # # # # # # # # # # # # # # # # # # # # # # # # # # # # # # # # # 
+	    # # # # # # # # # # # # # # # # # # # # # # # # # # # # # # # # # # # # # # # # # # # # # # # # # # # # # # # # # # # # # # # 
+            #import joint mover
+            import ART_jointMover
+            reload(ART_jointMover)
+            ART_jointMover.JointMover_Build()
+            
+            #set default scale
+	    cmds.setAttr("root_mover.scaleZ", 0.4)
+	    
+	    
+	    
+	    # # # # # # # # # # # # # # # # # # # # # # # # # # # # # # # # # # # # # # # # # # # # # # # # # # # # # # # # # # # # # # # 
+	    # # # # # # # # # # # # # # # # # # # # # # # # # # # # # # # # # # # # # # # # # # # # # # # # # # # # # # # # # # # # # # # 
+            #create lists for global movers, offset movers, geo movers, and lras
+            self.globalMovers = []
+            self.offsetMovers = []
+            self.geoMovers = []
+            self.lraControls = []
+            self.proxyGeo = []
+	    self.skelMeshRefGeo = []
+            
+            cmds.select("*_mover*")
+            selection = cmds.ls(sl = True)
+            
+            for each in selection:
+		try:
+		    color = cmds.getAttr(each + ".overrideColor")
+		    
+		    if color == 17:
+			shape = cmds.listRelatives(each, children = True)[0]
+			self.globalMovers.append(each + "|" + shape)
+			
+		    if color == 18:
+			shape = cmds.listRelatives(each, children = True)[0]
+			self.offsetMovers.append(each + "|" + shape)
+			
+		    if color == 20:
+			baseName = each.partition("_geo_mover")[0]
+			suffix = each.partition("_geo_mover")[2]
+			mover = baseName + "_mover" + suffix
+			if cmds.objExists(mover):
+			    moverShape = cmds.listRelatives(mover, shapes = True)[0]
+			    vis = cmds.getAttr(mover + "|" + moverShape + ".v")
+			    if vis == True:
+				shape = cmds.listRelatives(each, shapes = True)[0]
+				self.geoMovers.append(each + "|" + shape)
+				
+		except:
+		    pass
+		
+	    #setup default visibility
+	    for mover in self.offsetMovers:
+		cmds.lockNode(mover, lock = False)
+		cmds.setAttr(mover + ".visibility", lock = False, keyable = True)
+		cmds.setAttr(mover + ".visibility", 0)
+		cmds.lockNode(mover, lock = True)
+		
+                    
+            cmds.select("*_lra")
+            selection = cmds.ls(sl = True)
+            for each in selection:
+                self.lraControls.append(each)
+                
+            
+            cmds.select("proxy_geo_*")
+            selection = cmds.ls(sl = True, exactType = "transform", shapes = False)
+            for each in selection:
+                if each.find("grp") == -1:
+                    self.proxyGeo.append(each)
+		    
+		    
+
+                    
+                    
+                    
+	    #lock down the skeleton mesh reference meshes
+	    for each in self.skelMeshRefGeo:
+		cmds.setAttr(each + ".overrideEnabled", 1)
+                cmds.setAttr(each + ".overrideDisplayType", 2)
+		
+		
+            #clear selection and fit view
+            cmds.select(clear = True)
+            cmds.viewFit()
+            panels = cmds.getPanel(type = 'modelPanel')
+
+            #turn on smooth shading
+            for panel in panels:
+                editor = cmds.modelPanel(panel, q = True, modelEditor = True)
+                cmds.modelEditor(editor, edit = True, displayAppearance = "smoothShaded")
+            
+            
+            #switch modes
+            #hide all formLayouts except the main mode layout
+            cmds.formLayout(self.widgets["formLayout"], edit = True, vis = False)
+            cmds.formLayout(self.widgets["jointMover_bodyFormLayout"], edit = True, vis = True)
+            cmds.formLayout(self.widgets["jointMover_torsoModeFormLayout"], edit = True, vis = False)
+            cmds.formLayout(self.widgets["jointMover_leftHandModeFormLayout"], edit = True, vis = False)
+            cmds.formLayout(self.widgets["jointMover_rightHandModeFormLayout"], edit = True, vis = False)
+            cmds.formLayout(self.widgets["jointMover_leftFootModeFormLayout"], edit = True, vis = False)
+            cmds.formLayout(self.widgets["jointMover_rightFootModeFormLayout"], edit = True, vis = False)
+            cmds.formLayout(self.widgets["jointMover_physiqueFormLayout"], edit = True, vis = False)
+            
+
+	    
+            #repopulate the spine and neck bone buttons
+            self.jointMoverUI_spineBones()
+            self.jointMoverUI_neckBones()
+            self.jointMoverUI_updateFingerControlVis("numLeftFingers", "left")
+            self.jointMoverUI_updateFingerControlVis("numRightFingers", "right")
+            self.jointMoverUI_updateToeControlVis("numLeftToes", "left")
+            self.jointMoverUI_updateToeControlVis("numRightToes", "right")
+	    
+
+            #if this is a rebuild from the edit function, we need to look at the edit cache and set attrs on our movers so they are in the positions they were before the edit
+            if edit == True:
+                path = self.mayaToolsDir + "/General/ART/editCache.txt"
+                f = open(path, 'r')
+                lines = f.readlines()
+                
+                for line in lines:
+                    #convert the string that is a list, to a python list
+                    list = eval(line)
+                    control = list[0] 
+                    tx = list[1]
+                    ty = list[2]
+                    tz = list[3]
+                    rx = list[4]
+                    ry = list[5]
+                    rz = list[6]
+                    sx = list[7]
+                    sy = list[8]
+                    sz = list[9]
+                    
+                    if cmds.objExists(control):
+			for attr in [[".tx", tx], [".ty", ty], [".tz",tz], [".rx", rx], [".ry", ry], [".rz", rz], [".sx", sx], [".sy", sy], [".sz", sz]]:
+			    if cmds.getAttr(control + attr[0], keyable = True) == True:
+				cmds.setAttr(control + attr[0], attr[1])
+				
+				
+	    #for global movers and offset movers, connect scale x and y to z. Then lock x and y
+	    for mover in self.globalMovers:
+		try:
+		    mover = mover.partition("|")[0]
+		    cmds.connectAttr(mover + ".scaleZ", mover + ".scaleX")
+		    cmds.connectAttr(mover + ".scaleZ", mover + ".scaleY")
+		    cmds.setAttr(mover + ".scaleX",  keyable = False, channelBox = False)
+		    cmds.setAttr(mover + ".scaleY",  keyable = False, channelBox = False)
+		    
+		except:
+		    pass
+		
+	    for mover in self.offsetMovers:
+		try:
+		    mover = mover.partition("|")[0]
+		    cmds.connectAttr(mover + ".scaleZ", mover + ".scaleX")
+		    cmds.connectAttr(mover + ".scaleZ", mover + ".scaleY")
+		    cmds.setAttr(mover + ".scaleX", keyable = False, channelBox = False)
+		    cmds.setAttr(mover + ".scaleY",  keyable = False, channelBox = False)
+		    
+		except:
+		    pass
+		
+
+	    # # # # # # # # # # # # # # # # # # # # # # # # # # # # # # # # # # # # # # # # # # # # # # # # # # # # # # # # # # # # # # # 
+	    # # # # # # # # # # # # # # # # # # # # # # # # # # # # # # # # # # # # # # # # # # # # # # # # # # # # # # # # # # # # # # # 
+	    #save out a template for this file
+	    skeletonSettingsTemplateNode = cmds.createNode("unknown", name = "SkeletonSettings_Cache")
+
+	    
+	    #make a list of all of the skeleton builder ui widgets
+	    intFields = [
+		"torsoOptions_numSpineBones", "torsoOptions_numNeckBones",
+		"leftArmOptions_numUpArmTwistBones", "leftArmOptions_numLowArmTwistBones", "leftHandOptions_numThumbBones", "leftHandOptions_numIndexBones", "leftHandOptions_numMiddleBones", "leftHandOptions_numRingBones", "leftHandOptions_numPinkyBones",
+		"rightArmOptions_numUpArmTwistBones", "rightArmOptions_numLowArmTwistBones", "rightHandOptions_numThumbBones", "rightHandOptions_numIndexBones", "rightHandOptions_numMiddleBones", "rightHandOptions_numRingBones", "rightHandOptions_numPinkyBones",
+		"leftLegOptions_numThighTwistBones", "leftLegOptions_numCalfTwistBones", "leftFootOptions_numThumbBones", "leftFootOptions_numIndexBones", "leftFootOptions_numMiddleBones", "leftFootOptions_numRingBones", "leftFootOptions_numPinkyBones",
+		"rightLegOptions_numThighTwistBones", "rightLegOptions_numCalfTwistBones", "rightFootOptions_numThumbBones", "rightFootOptions_numIndexBones", "rightFootOptions_numMiddleBones", "rightFootOptions_numRingBones", "rightFootOptions_numPinkyBones"]
+	    
+	    checkBoxes = ["leftLegIncludeBallJoint_CB", "rightLegIncludeBallJoint_CB"]
+	    cmds.select(skeletonSettingsTemplateNode)
+	    
+	    #get values
+	    for field in intFields:
+		value = cmds.intField(self.widgets[field], q = True, v = True)
+		color = cmds.intField(self.widgets[field], q = True, bgc = True)
+		if cmds.objExists(skeletonSettingsTemplateNode + "." + field) == False:
+		    cmds.addAttr(skeletonSettingsTemplateNode, dt = "string", ln = field)
+		    cmds.setAttr(skeletonSettingsTemplateNode + "." + field, str([value, color]), type = "string")
+		
+	    for box in checkBoxes:
+		value = cmds.checkBox(self.widgets[box], q = True, value = True)
+		if cmds.objExists(skeletonSettingsTemplateNode + "." + box) == False:
+		    cmds.addAttr(skeletonSettingsTemplateNode, dt = "string", ln = box)
+		    cmds.setAttr(skeletonSettingsTemplateNode + "." + box, value, type = "string")
+		    
+		    
+	    #leg style
+	    style = cmds.optionMenu(self.widgets["legStyleOptionMenu"], q = True, value = True)
+	    if cmds.objExists(skeletonSettingsTemplateNode + ".legStyle") == False:
+		cmds.addAttr(skeletonSettingsTemplateNode, dt = "string", ln = "legStyle")
+		cmds.setAttr(skeletonSettingsTemplateNode + ".legStyle", style, type = "string")
+		
+		
+	    # # # # # # # # # # # # # # # # # # # # # # # # # # # # # # # # # # # # # # # # # # # # # # # # # # # # # # # # # # # # # # # 
+	    # # # # # # # # # # # # # # # # # # # # # # # # # # # # # # # # # # # # # # # # # # # # # # # # # # # # # # # # # # # # # # # 
+	    #add extra joint info
+	    extraJoints = cmds.textScrollList(self.widgets["addedRigModulesList"], q = True, allItems = True)
+	    name = ""
+	    type = ""
+	    parent = ""
+	    count = -1
+	    translate = ""
+	    rotate = ""
+	    
+	    if extraJoints != None:
+		for joint in extraJoints:
+		    enumEntry = ""
+		    
+		    #get extra joint information
+		    name = joint.rpartition("] ")[2]
+		    type = joint .partition("]")[0].partition("[")[2]
+    
+		    
+		    #get joint parent
+		    parent = cmds.textField(self.widgets["rigMod_" + name + "_settings_parent"], q = True, text = True)
+		    
+		    if type == "leaf":
+			#get translate and rotate values
+			translate = cmds.checkBox(self.widgets["rigMod_" + name + "_settings_translate"], q = True, value = True)
+			rotate = cmds.checkBox(self.widgets["rigMod_" + name + "_settings_rotate"], q = True, value = True)
+			
+    
+		    if type == "chain":
+			#get count
+			count = cmds.textField(self.widgets["rigMod_" + name + "_settings_numberInChain"], q = True, text = True)
+				
+				
+		    
+				
+		    if name != "":
+			if parent != "":
+			    if type != "":
+				
+				#LEAF JOINTS
+				if type == "leaf":
+				    if translate != "":
+					if rotate != "":
+					    if translate == True and rotate == True:
+						enumEntry = parent + "/" + type + "/" + name + " (TR)"
+						
+					    if translate == False and rotate == True:
+						enumEntry = parent + "/" + type + "/" + name + " (R)"
+						
+					    if translate == True and rotate == False:
+						enumEntry = parent + "/" + type + "/" + name + " (T)"
+						
+					    if translate == False and rotate == False:
+						enumEntry = parent + "/" + type + "/" + name + " (TR)"
+						cmds.warning(message= "Cannot create a leaf joint with no attributes available for animation. Defaulting to Translate and Rotate.")
+		    
+					    cmds.addAttr(skeletonSettingsTemplateNode, longName = "extraJoint" + name, enumName = enumEntry, at = "enum", keyable = True)
+					    name = ""
+					    type = ""
+					    parent = ""
+					    count = -1
+					    translate = ""
+					    rotate = ""
+					    
+					    
+				#JIGGLE JOINTS
+				if type == "jiggle":
+				    enumEntry = parent + "/" + type + "/" + name		    
+				    cmds.addAttr(skeletonSettingsTemplateNode, longName = "extraJoint" + name, enumName = enumEntry, at = "enum", keyable = True)
+				    name = ""
+				    type = ""
+				    parent = ""
+				    count = -1
+				    translate = ""
+				    rotate = ""				
+    
+    
+    
+				#DYNAMIC CHAINS
+				if type == "chain":
+				    if int(count) > 0:
+					enumEntry = parent + "/" + type + "/" + name + " (" + str(count) + ")"
+					cmds.addAttr(skeletonSettingsTemplateNode, longName = "extraJoint" + name, enumName = enumEntry, at = "enum", keyable = True)
+					name = ""
+					type = ""
+					parent = ""
+					count = -1
+					translate = ""
+					rotate = ""	
+
+		
+		
+	    #add mirror settings to the cache node
+	    cmds.addAttr(skeletonSettingsTemplateNode, dt = "string", ln = "mirrorSetting")
+	    cmds.addAttr(skeletonSettingsTemplateNode, dt = "string", ln = "mirrorLeft")
+	    cmds.addAttr(skeletonSettingsTemplateNode, dt = "string", ln = "mirrorRight")
+	    
+	    mirrorSetting = cmds.optionMenu(self.widgets["rigModuleMirrorMode"], q = True, v = True)
+	    mirrorLeft = cmds.textField(self.widgets["rigModuleMirrorTextL"], q = True, text = True)
+	    mirrorRight = cmds.textField(self.widgets["rigModuleMirrorTextR"], q = True, text = True)
+	    
+	    cmds.setAttr(skeletonSettingsTemplateNode + ".mirrorSetting", mirrorSetting, type = "string")
+	    cmds.setAttr(skeletonSettingsTemplateNode + ".mirrorLeft", mirrorLeft, type = "string")
+	    cmds.setAttr(skeletonSettingsTemplateNode + ".mirrorRight", mirrorRight, type = "string")
+	    
+	    
+	    
+	    cmds.container("JointMover", edit = True, addNode = skeletonSettings)
+	    
+	    
+	    #change mode buttons to default
+	    self.jointMover_ChangeMoveMode("Geo", "off")
+	    self.jointMover_ChangeMoveMode("Offset", "off")
+	    self.jointMover_ChangeMoveMode("Global", "off")
+	    self.jointMover_ChangeMoveMode("Global", "on")
+	    
+	    cmds.iconTextRadioButton(self.widgets["moverMode"], edit = True, sl = True)
+	    cmds.iconTextRadioButton(self.widgets["offsetMode"], edit = True, sl = False)
+	    cmds.iconTextRadioButton(self.widgets["geoMoverMode"], edit = True, sl = False)
+	    
+	    
+	    #rebuild the physique UI
+	    if cmds.formLayout(self.widgets["jointMover_physiqueFormLayout"], q = True, exists = True):
+		cmds.deleteUI(self.widgets["jointMover_physiqueFormLayout"])
+	    
+	    self.createJointMover_PhysiqueUI()
+	    
+	    #unlock menu
+	    cmds.menu(self.widgets["jointMoverTools"], edit = True, enable = True)
+	    
+	    
+	    #show tabs in tab Layout
+	    cmds.tabLayout(self.widgets["tabLayout"], edit = True, tv = True)
+	    
+	    #create list view
+	    self.createListView()
+	    
+
+	    
+	    
+    # # # # # # # # # # # # # # # # # # # # # # # # # # # # # # # # # # # # # # # # # # # # # # # # # # # # # # # # # # # # # # # 
+    # # # # # # # # # # # # # # # # # # # # # # # # # # # # # # # # # # # # # # # # # # # # # # # # # # # # # # # # # # # # # # # 
+    # # # # # # # # # # # # # # # # # # # # # # # # # # # # # # # # # # # # # # # # # # # # # # # # # # # # # # # # # # # # # # # 
+    # # # # # # # # # # # # # # # # # # # # # # # # # # # # # # # # # # # # # # # # # # # # # # # # # # # # # # # # # # # # # # # 
+    # # # # # # # # # # # # # # # # # # # # # # # # # # # # # # # # # # # # # # # # # # # # # # # # # # # # # # # # # # # # # # # 
+    
+    def jointMoverUI_updateToeControlVis(self, attr, side):
+        
+        numToes = cmds.getAttr("Skeleton_Settings." + attr, asString = True)
+        flag = side[0]
+        
+        #Turn visibility on for all toe buttons and reset page button
+        cmds.button(self.widgets["jointMoverUI_" + side + "FootPage"], edit = True, bgc = [.92, .93, .14],c = partial(self.switchMode, "jointMover_" + side + "FootModeFormLayout", None))
+        
+        for button in ["jointMoverUI_indextoe01_" + flag, "jointMoverUI_indextoe02_" + flag, "jointMoverUI_indextoe03_" + flag, "jointMoverUI_indextoe04_" + flag, "jointMoverUI_ringtoe01_" + flag, "jointMoverUI_ringtoe02_" + flag, "jointMoverUI_ringtoe03_" + flag, "jointMoverUI_ringtoe04_" + flag, "jointMoverUI_pinkytoe01_" + flag, "jointMoverUI_pinkytoe02_" + flag, "jointMoverUI_pinkytoe03_" + flag, "jointMoverUI_pinkytoe04_" + flag, "jointMoverUI_middletoe01_" + flag, "jointMoverUI_middletoe02_" + flag, "jointMoverUI_middletoe03_" + flag, "jointMoverUI_middletoe04_" + flag, "jointMoverUI_bigtoe01_" + flag, "jointMoverUI_bigtoe02_" + flag, "jointMoverUI_bigtoe03_" + flag]:
+            cmds.button(self.widgets[button], edit = True, vis = True)
+
+
+
+        #set visibility based on number of toes
+        if numToes == "None":
+            cmds.button(self.widgets["jointMoverUI_" + side + "FootPage"], edit = True, bgc = [0, .30, 1],c = partial(self.jointMoverSelect, "ball", "l", "global"))
+           
+        #otherwise, check each attr on skeleton_settings and if the attr is False, hide that toe button 
+        else:
+            toeAttrInfo = [[side + "Footindexmeta", "jointMoverUI_indextoe01_" + flag], [side + "Footindex1", "jointMoverUI_indextoe02_" + flag], [side + "Footindex2","jointMoverUI_indextoe03_" + flag], [side + "Footindex3","jointMoverUI_indextoe04_" + flag],
+                           [side + "Footringmeta","jointMoverUI_ringtoe01_" + flag], [side + "Footring1", "jointMoverUI_ringtoe02_" + flag], [side + "Footring2","jointMoverUI_ringtoe03_" + flag], [side + "Footring3","jointMoverUI_ringtoe04_" + flag],
+                           [side + "Footpinkymeta", "jointMoverUI_pinkytoe01_" + flag], [side + "Footpinky1", "jointMoverUI_pinkytoe02_" + flag], [side + "Footpinky2", "jointMoverUI_pinkytoe03_" + flag], [side + "Footpinky3", "jointMoverUI_pinkytoe04_" + flag],
+                           [side + "Footindexmeta", "jointMoverUI_middletoe01_" + flag], [side + "Footindex1", "jointMoverUI_middletoe02_" + flag], [side + "Footindex2", "jointMoverUI_middletoe03_" + flag], [side + "Footindex3","jointMoverUI_middletoe04_" + flag],
+                           [side + "Footbigtoemeta", "jointMoverUI_bigtoe01_" + flag], [side + "Footbigtoe1", "jointMoverUI_bigtoe02_" + flag], [side + "Footbigtoe2", "jointMoverUI_bigtoe03_" + flag]]
+
+            for entry in toeAttrInfo:
+                attr = entry[0]
+                button = entry[1]
+                
+                if cmds.getAttr("Skeleton_Settings." + attr) == False:
+                    cmds.button(self.widgets[button], edit = True, vis = False)
+                
+
+
+    # # # # # # # # # # # # # # # # # # # # # # # # # # # # # # # # # # # # # # # # # # # # # # # # # # # # # # # # # # # # # # # 
+    # # # # # # # # # # # # # # # # # # # # # # # # # # # # # # # # # # # # # # # # # # # # # # # # # # # # # # # # # # # # # # # 
+    # # # # # # # # # # # # # # # # # # # # # # # # # # # # # # # # # # # # # # # # # # # # # # # # # # # # # # # # # # # # # # # 
+    # # # # # # # # # # # # # # # # # # # # # # # # # # # # # # # # # # # # # # # # # # # # # # # # # # # # # # # # # # # # # # # 
+    # # # # # # # # # # # # # # # # # # # # # # # # # # # # # # # # # # # # # # # # # # # # # # # # # # # # # # # # # # # # # # # 
+    def jointMoverUI_updateFingerControlVis(self, attr, side):
+        
+        numFingers = cmds.getAttr("Skeleton_Settings." + attr, asString = True)
+        flag = side[0]
+        
+        cmds.button(self.widgets["jointMoverUI_" + side + "HandPage"], edit = True, vis = True )
+        for button in ["jointMoverUI_index01_" + flag, "jointMoverUI_index02_" + flag, "jointMoverUI_index03_" + flag, "jointMoverUI_index04_" + flag, "jointMoverUI_ring01_" + flag, "jointMoverUI_ring02_" + flag, "jointMoverUI_ring03_" + flag, "jointMoverUI_ring04_" + flag, "jointMoverUI_pinky01_" + flag, "jointMoverUI_pinky02_" + flag, "jointMoverUI_pinky03_" + flag, "jointMoverUI_pinky04_" + flag, "jointMoverUI_middle01_" + flag, "jointMoverUI_middle02_" + flag, "jointMoverUI_middle03_" + flag, "jointMoverUI_middle04_" + flag]:
+            cmds.button(self.widgets[button], edit = True, vis = True)
+            
+        cmds.button(self.widgets["jointMoverUI_bodyPageThumb01_" + flag], edit = True, vis = False)
+        cmds.button(self.widgets["jointMoverUI_bodyPageThumb02_" + flag], edit = True, vis = False)
+        cmds.button(self.widgets["jointMoverUI_bodyPageThumb03_" + flag], edit = True, vis = False)
+        
+        cmds.button(self.widgets["jointMoverUI_thumb01_" + flag], edit = True, vis = True)
+        cmds.button(self.widgets["jointMoverUI_thumb02_" + flag], edit = True, vis = True)
+        cmds.button(self.widgets["jointMoverUI_thumb03_" + flag], edit = True, vis = True)
+
+
+        
+        #set visibility based on number of fingers
+        if numFingers == "None":
+            cmds.button(self.widgets["jointMoverUI_" + side + "HandPage"], edit = True, vis = False )
+
+            cmds.button(self.widgets["jointMoverUI_bodyPageThumb01_" + flag], edit = True, vis = True)
+            cmds.button(self.widgets["jointMoverUI_bodyPageThumb02_" + flag], edit = True, vis = True)
+            cmds.button(self.widgets["jointMoverUI_bodyPageThumb03_" + flag], edit = True, vis = True)
+            
+            
+        if numFingers == "One":
+            for button in ["jointMoverUI_index01_" + flag, "jointMoverUI_index02_" + flag, "jointMoverUI_index03_" + flag, "jointMoverUI_index04_" + flag, "jointMoverUI_ring01_" + flag, "jointMoverUI_ring02_" + flag, "jointMoverUI_ring03_" + flag, "jointMoverUI_ring04_" + flag, "jointMoverUI_pinky01_" + flag, "jointMoverUI_pinky02_" + flag, "jointMoverUI_pinky03_" + flag, "jointMoverUI_pinky04_" + flag]:
+                cmds.button(self.widgets[button], edit = True, vis = False)
+                
+            if cmds.getAttr("Skeleton_Settings." + side + "middlemeta") == False:
+                cmds.button(self.widgets["jointMoverUI_middle01_" + flag], edit = True, vis = False)
+                
+            if cmds.getAttr("Skeleton_Settings." + side + "middle1") == False:
+                cmds.button(self.widgets["jointMoverUI_middle02_" + flag], edit = True, vis = False)
+                
+            if cmds.getAttr("Skeleton_Settings." + side + "middle2") == False:
+                cmds.button(self.widgets["jointMoverUI_middle03_" + flag], edit = True, vis = False)
+                
+            if cmds.getAttr("Skeleton_Settings." + side + "middle3") == False:
+                cmds.button(self.widgets["jointMoverUI_middle04_" + flag], edit = True, vis = False)
+                
+                
+                
+
+        if numFingers == "Two":
+            for button in ["jointMoverUI_middle01_" + flag, "jointMoverUI_middle02_" + flag, "jointMoverUI_middle03_" + flag, "jointMoverUI_middle04_" + flag, "jointMoverUI_ring01_" + flag, "jointMoverUI_ring02_" + flag, "jointMoverUI_ring03_" + flag, "jointMoverUI_ring04_" + flag]:
+                cmds.button(self.widgets[button], edit = True, vis = False)
+                
+                
+            if cmds.getAttr("Skeleton_Settings." + side + "indexmeta") == False:
+                cmds.button(self.widgets["jointMoverUI_index01_" + flag], edit = True, vis = False)
+                
+            if cmds.getAttr("Skeleton_Settings." + side + "index1") == False:
+                cmds.button(self.widgets["jointMoverUI_index02_" + flag], edit = True, vis = False)
+                
+            if cmds.getAttr("Skeleton_Settings." + side + "index2") == False:
+                cmds.button(self.widgets["jointMoverUI_index03_" + flag], edit = True, vis = False)
+                
+            if cmds.getAttr("Skeleton_Settings." + side + "index3") == False:
+                cmds.button(self.widgets["jointMoverUI_index04_" + flag], edit = True, vis = False)
+                
+                
+            if cmds.getAttr("Skeleton_Settings." + side + "pinkymeta") == False:
+                cmds.button(self.widgets["jointMoverUI_pinky01_" + flag], edit = True, vis = False)
+                
+            if cmds.getAttr("Skeleton_Settings." + side + "pinky1") == False:
+                cmds.button(self.widgets["jointMoverUI_pinky02_" + flag], edit = True, vis = False)
+                
+            if cmds.getAttr("Skeleton_Settings." + side + "pinky2") == False:
+                cmds.button(self.widgets["jointMoverUI_pinky03_" + flag], edit = True, vis = False)
+                
+            if cmds.getAttr("Skeleton_Settings." + side + "pinky3") == False:
+                cmds.button(self.widgets["jointMoverUI_pinky04_" + flag], edit = True, vis = False)
+                
+                
+                
+                
+        if numFingers == "Three":
+            for button in ["jointMoverUI_pinky01_" + flag, "jointMoverUI_pinky02_" + flag, "jointMoverUI_pinky03_" + flag, "jointMoverUI_pinky04_" + flag]:
+                cmds.button(self.widgets[button], edit = True, vis = False)
+                
+                
+            if cmds.getAttr("Skeleton_Settings." + side + "indexmeta") == False:
+                cmds.button(self.widgets["jointMoverUI_index01_" + flag], edit = True, vis = False)
+                
+            if cmds.getAttr("Skeleton_Settings." + side + "index1") == False:
+                cmds.button(self.widgets["jointMoverUI_index02_" + flag], edit = True, vis = False)
+                
+            if cmds.getAttr("Skeleton_Settings." + side + "index2") == False:
+                cmds.button(self.widgets["jointMoverUI_index03_" + flag], edit = True, vis = False)
+                
+            if cmds.getAttr("Skeleton_Settings." + side + "index3") == False:
+                cmds.button(self.widgets["jointMoverUI_index04_" + flag], edit = True, vis = False)
+                
+                
+            if cmds.getAttr("Skeleton_Settings." + side + "middlemeta") == False:
+                cmds.button(self.widgets["jointMoverUI_middle01_" + flag], edit = True, vis = False)
+                
+            if cmds.getAttr("Skeleton_Settings." + side + "middle1") == False:
+                cmds.button(self.widgets["jointMoverUI_middle02_" + flag], edit = True, vis = False)
+                
+            if cmds.getAttr("Skeleton_Settings." + side + "middle2") == False:
+                cmds.button(self.widgets["jointMoverUI_middle03_" + flag], edit = True, vis = False)
+                
+            if cmds.getAttr("Skeleton_Settings." + side + "middle3") == False:
+                cmds.button(self.widgets["jointMoverUI_middle04_" + flag], edit = True, vis = False)
+                
+                
+            if cmds.getAttr("Skeleton_Settings." + side + "ringmeta") == False:
+                cmds.button(self.widgets["jointMoverUI_ring01_" + flag], edit = True, vis = False)
+                
+            if cmds.getAttr("Skeleton_Settings." + side + "ring1") == False:
+                cmds.button(self.widgets["jointMoverUI_ring02_" + flag], edit = True, vis = False)
+                
+            if cmds.getAttr("Skeleton_Settings." + side + "ring2") == False:
+                cmds.button(self.widgets["jointMoverUI_ring03_" + flag], edit = True, vis = False)
+                
+            if cmds.getAttr("Skeleton_Settings." + side + "ring3") == False:
+                cmds.button(self.widgets["jointMoverUI_ring04_" + flag], edit = True, vis = False)
+                
+                
+                
+                
+        if numFingers == "Four":
+                
+            if cmds.getAttr("Skeleton_Settings." + side + "indexmeta") == False:
+                cmds.button(self.widgets["jointMoverUI_index01_" + flag], edit = True, vis = False)
+                
+            if cmds.getAttr("Skeleton_Settings." + side + "index1") == False:
+                cmds.button(self.widgets["jointMoverUI_index02_" + flag], edit = True, vis = False)
+                
+            if cmds.getAttr("Skeleton_Settings." + side + "index2") == False:
+                cmds.button(self.widgets["jointMoverUI_index03_" + flag], edit = True, vis = False)
+                
+            if cmds.getAttr("Skeleton_Settings." + side + "index3") == False:
+                cmds.button(self.widgets["jointMoverUI_index04_" + flag], edit = True, vis = False)
+                
+                
+            if cmds.getAttr("Skeleton_Settings." + side + "middlemeta") == False:
+                cmds.button(self.widgets["jointMoverUI_middle01_" + flag], edit = True, vis = False)
+                
+            if cmds.getAttr("Skeleton_Settings." + side + "middle1") == False:
+                cmds.button(self.widgets["jointMoverUI_middle02_" + flag], edit = True, vis = False)
+                
+            if cmds.getAttr("Skeleton_Settings." + side + "middle2") == False:
+                cmds.button(self.widgets["jointMoverUI_middle03_" + flag], edit = True, vis = False)
+                
+            if cmds.getAttr("Skeleton_Settings." + side + "middle3") == False:
+                cmds.button(self.widgets["jointMoverUI_middle04_" + flag], edit = True, vis = False)
+                
+                
+            if cmds.getAttr("Skeleton_Settings." + side + "ringmeta") == False:
+                cmds.button(self.widgets["jointMoverUI_ring01_" + flag], edit = True, vis = False)
+                
+            if cmds.getAttr("Skeleton_Settings." + side + "ring1") == False:
+                cmds.button(self.widgets["jointMoverUI_ring02_" + flag], edit = True, vis = False)
+                
+            if cmds.getAttr("Skeleton_Settings." + side + "ring2") == False:
+                cmds.button(self.widgets["jointMoverUI_ring03_" + flag], edit = True, vis = False)
+                
+            if cmds.getAttr("Skeleton_Settings." + side + "ring3") == False:
+                cmds.button(self.widgets["jointMoverUI_ring04_" + flag], edit = True, vis = False)
+                
+
+            if cmds.getAttr("Skeleton_Settings." + side + "pinkymeta") == False:
+                cmds.button(self.widgets["jointMoverUI_pinky01_" + flag], edit = True, vis = False)
+                
+            if cmds.getAttr("Skeleton_Settings." + side + "pinky1") == False:
+                cmds.button(self.widgets["jointMoverUI_pinky02_" + flag], edit = True, vis = False)
+                
+            if cmds.getAttr("Skeleton_Settings." + side + "pinky2") == False:
+                cmds.button(self.widgets["jointMoverUI_pinky03_" + flag], edit = True, vis = False)
+                
+            if cmds.getAttr("Skeleton_Settings." + side + "pinky3") == False:
+                cmds.button(self.widgets["jointMoverUI_pinky04_" + flag], edit = True, vis = False)
+                
+                
+        #Thumb
+        if cmds.getAttr("Skeleton_Settings." + side + "thumb1") == False:
+            cmds.button(self.widgets["jointMoverUI_thumb01_" + flag], edit = True, vis = False)
+            cmds.button(self.widgets["jointMoverUI_bodyPageThumb01_" + flag], edit = True, vis = False)
+
+            
+        if cmds.getAttr("Skeleton_Settings." + side + "thumb2") == False:
+            cmds.button(self.widgets["jointMoverUI_thumb02_" + flag], edit = True, vis = False)
+            cmds.button(self.widgets["jointMoverUI_bodyPageThumb02_" + flag], edit = True, vis = False)
+            
+        if cmds.getAttr("Skeleton_Settings." + side + "thumb3") == False:
+            cmds.button(self.widgets["jointMoverUI_thumb03_" + flag], edit = True, vis = False)
+            cmds.button(self.widgets["jointMoverUI_bodyPageThumb03_" + flag], edit = True, vis = False)
+            
+            
+            
+            
+    # # # # # # # # # # # # # # # # # # # # # # # # # # # # # # # # # # # # # # # # # # # # # # # # # # # # # # # # # # # # # # # 
+    # # # # # # # # # # # # # # # # # # # # # # # # # # # # # # # # # # # # # # # # # # # # # # # # # # # # # # # # # # # # # # # 
+    # # # # # # # # # # # # # # # # # # # # # # # # # # # # # # # # # # # # # # # # # # # # # # # # # # # # # # # # # # # # # # # 
+    # # # # # # # # # # # # # # # # # # # # # # # # # # # # # # # # # # # # # # # # # # # # # # # # # # # # # # # # # # # # # # # 
+    # # # # # # # # # # # # # # # # # # # # # # # # # # # # # # # # # # # # # # # # # # # # # # # # # # # # # # # # # # # # # # # 
+    def jointMoverUI_spineBones(self):
+
+        #get the number of spine bones and create buttons based on #
+        spineBones = cmds.getAttr("Skeleton_Settings.numSpineBones")
+        
+        #setup some variables for our equation
+        length = 175
+        padding = (int(spineBones) - 1) * 5
+        buttonHeight = (length - padding)/ spineBones
+        
+        #this will be the starting position from the bottom of the UI
+        verticalPosition = 290
+        
+        #for each of the spine bones, create the button
+        for i in range(int(spineBones)):
+            if spineBones < 10:
+                button = "jointMoverUI_spine_0" + str(i + 1)
+                niceName = "spine_0"
+            else:
+                button = "jointMoverUI_spine_" + str(i + 1)
+                niceName = "spine_"
+                
+            if i%2==0:
+                buttonColor = [.96, .96, .65]
+                
+            else:
+                buttonColor = [.77, .78, .52]
+            self.widgets[niceName + str(i + 1) + "_button"] = cmds.button(button, w = 100, height = buttonHeight, label = "", bgc = buttonColor, c = partial(self.jointMoverSelect, niceName + str(i + 1), None, "global"), parent = self.widgets["jointMover_torsoModeFormLayout"])            
+            cmds.formLayout(self.widgets["jointMover_torsoModeFormLayout"], edit = True, af = [(self.widgets[niceName + str(i + 1) + "_button"], 'bottom', verticalPosition), (self.widgets[niceName + str(i + 1) + "_button"], 'left', 174)])
+            verticalPosition = verticalPosition + buttonHeight + 5
+            
+            
+            
+    # # # # # # # # # # # # # # # # # # # # # # # # # # # # # # # # # # # # # # # # # # # # # # # # # # # # # # # # # # # # # # # 
+    # # # # # # # # # # # # # # # # # # # # # # # # # # # # # # # # # # # # # # # # # # # # # # # # # # # # # # # # # # # # # # # 
+    # # # # # # # # # # # # # # # # # # # # # # # # # # # # # # # # # # # # # # # # # # # # # # # # # # # # # # # # # # # # # # # 
+    # # # # # # # # # # # # # # # # # # # # # # # # # # # # # # # # # # # # # # # # # # # # # # # # # # # # # # # # # # # # # # # 
+    # # # # # # # # # # # # # # # # # # # # # # # # # # # # # # # # # # # # # # # # # # # # # # # # # # # # # # # # # # # # # # # 
+    def jointMoverUI_neckBones(self):
+
+        #get the number of spine bones and create buttons based on #
+        spineBones = cmds.getAttr("Skeleton_Settings.numNeckBones")
+        
+        #setup some variables for our equation
+        length = 40
+        padding = (int(spineBones) - 1) * 2
+        buttonHeight = (length - padding)/ spineBones
+        
+        #this will be the starting position from the bottom of the UI
+        verticalPosition = 470
+        
+        #for each of the spine bones, create the button
+        for i in range(int(spineBones)):
+            button = "jointMoverUI_neck_0" + str(i + 1)
+            niceName = "neck_0"
+
+            if i%2==0:
+                buttonColor = [.96, .96, .65]
+                
+            else:
+                buttonColor = [.77, .78, .52]
+                
+            self.widgets[niceName + str(i + 1) + "_button"] = cmds.button(button, w = 40, height = buttonHeight, label = "", bgc = buttonColor, c = partial(self.jointMoverSelect, niceName + str(i + 1), None, "global"), parent = self.widgets["jointMover_torsoModeFormLayout"])            
+            cmds.formLayout(self.widgets["jointMover_torsoModeFormLayout"], edit = True, af = [(self.widgets[niceName + str(i + 1) + "_button"], 'bottom', verticalPosition), (self.widgets[niceName + str(i + 1) + "_button"], 'left', 205)])
+            verticalPosition = verticalPosition + buttonHeight + 5
+            
+            
+    # # # # # # # # # # # # # # # # # # # # # # # # # # # # # # # # # # # # # # # # # # # # # # # # # # # # # # # # # # # # # # # 
+    # # # # # # # # # # # # # # # # # # # # # # # # # # # # # # # # # # # # # # # # # # # # # # # # # # # # # # # # # # # # # # # 
+    # # # # # # # # # # # # # # # # # # # # # # # # # # # # # # # # # # # # # # # # # # # # # # # # # # # # # # # # # # # # # # # 
+    # # # # # # # # # # # # # # # # # # # # # # # # # # # # # # # # # # # # # # # # # # # # # # # # # # # # # # # # # # # # # # # 
+    # # # # # # # # # # # # # # # # # # # # # # # # # # # # # # # # # # # # # # # # # # # # # # # # # # # # # # # # # # # # # # # 
+    def physiqueSettings(self,  *args):
+	
+	vis = cmds.formLayout(self.widgets["jointMover_physiqueFormLayout"], q = True, vis = True)
+	imagePath = self.mayaToolsDir + "/General/Icons/ART/"
+	
+	if vis == True:
+	    cmds.formLayout(self.widgets["jointMover_physiqueFormLayout"], edit = True, vis = False)
+	    cmds.symbolButton(self.widgets["physiqueMode"], edit = True, image = imagePath + "physique_off.bmp")
+	    
+	    #get current modes
+	    globalMode = cmds.iconTextRadioButton(self.widgets["moverMode"], q = True, sl = True)
+	    offsetMode = cmds.iconTextRadioButton(self.widgets["offsetMode"], q = True, sl = True)
+	    geoMode = cmds.iconTextRadioButton(self.widgets["geoMoverMode"], q = True, sl = True)
+	    
+	    
+	    self.createJointMoverToolbar("jointMover_bodyFormLayout", globalMode, offsetMode, geoMode)
+	    
+	    
+	if vis == False:
+	    #get current modes
+	    globalMode = cmds.iconTextRadioButton(self.widgets["moverMode"], q = True, sl = True)
+	    offsetMode = cmds.iconTextRadioButton(self.widgets["offsetMode"], q = True, sl = True)
+	    geoMode = cmds.iconTextRadioButton(self.widgets["geoMoverMode"], q = True, sl = True)
+	    
+
+	    
+	    cmds.formLayout(self.widgets["jointMover_physiqueFormLayout"], edit = True, vis = True)
+	    self.createJointMoverToolbar("jointMover_physiqueFormLayout", globalMode, offsetMode, geoMode)
+	    cmds.symbolButton(self.widgets["physiqueMode"], edit = True, image = imagePath + "physique.bmp")
+	    
+	try:
+	    cmds.connectAttr("spine1_shapes.spine_1_bust", "spine2_shapes.spine_2_bust")
+	    cmds.connectAttr("spine1_shapes.spine_1_bust", "spine3_shapes.spine_3_bust")
+	    cmds.connectAttr("spine1_shapes.spine_1_bust", "spine4_shapes.spine_4_bust")
+	    cmds.connectAttr("spine1_shapes.spine_1_bust", "spine5_shapes.spine_5_bust")
+	    
+	    cmds.connectAttr("spine1_shapes.spine_1_fat", "spine2_shapes.spine_2_fat")
+	    cmds.connectAttr("spine1_shapes.spine_1_fat", "spine3_shapes.spine_3_fat")
+	    cmds.connectAttr("spine1_shapes.spine_1_fat", "spine4_shapes.spine_4_fat")
+	    cmds.connectAttr("spine1_shapes.spine_1_fat", "spine5_shapes.spine_5_fat")
+	    
+	except:
+	    pass
+
+	    
+	
+	
+    # # # # # # # # # # # # # # # # # # # # # # # # # # # # # # # # # # # # # # # # # # # # # # # # # # # # # # # # # # # # # # # 
+    # # # # # # # # # # # # # # # # # # # # # # # # # # # # # # # # # # # # # # # # # # # # # # # # # # # # # # # # # # # # # # # 
+    # # # # # # # # # # # # # # # # # # # # # # # # # # # # # # # # # # # # # # # # # # # # # # # # # # # # # # # # # # # # # # # 
+    # # # # # # # # # # # # # # # # # # # # # # # # # # # # # # # # # # # # # # # # # # # # # # # # # # # # # # # # # # # # # # # 
+    # # # # # # # # # # # # # # # # # # # # # # # # # # # # # # # # # # # # # # # # # # # # # # # # # # # # # # # # # # # # # # # 
+    def switchMode(self, layout, toolBarCreationMethod, *args):
+        
+        
+        #sync buttons	
+	if cmds.objExists("Symmetry_Nodes"):
+	    cmds.iconTextCheckBox(self.widgets["symmetry"] , edit = True, v = True)
+	    
+	#find which mode we are in, store that temporarily
+	globalMode = cmds.iconTextRadioButton(self.widgets["moverMode"], q = True, sl = True)
+	offsetMode = cmds.iconTextRadioButton(self.widgets["offsetMode"], q = True, sl = True)
+	geoMode = cmds.iconTextRadioButton(self.widgets["geoMoverMode"], q = True, sl = True)
+	
+
+        layouts = ["formLayout",  "jointMover_bodyFormLayout", "jointMover_torsoModeFormLayout", "jointMover_leftHandModeFormLayout", "jointMover_rightHandModeFormLayout", "jointMover_leftFootModeFormLayout", "jointMover_rightFootModeFormLayout", "jointMover_lockedFormLayout"]
+        
+        for form in layouts:
+            if form != layout:
+                cmds.formLayout(self.widgets[form], edit = True, vis = False)
+                
+        #show active page
+        cmds.formLayout(self.widgets[layout], edit = True, vis = True)
+        
+        #create toolbar
+        if toolBarCreationMethod == "skeleton":
+            self.createSkeletonToolbar(layout)
+            
+        else:
+            self.createJointMoverToolbar(layout, True, False, False)
+            
+
+	#change mode button to correct one
+	cmds.iconTextRadioButton(self.widgets["moverMode"], edit = True, sl = globalMode)
+	cmds.iconTextRadioButton(self.widgets["offsetMode"], edit = True, sl = offsetMode)
+	cmds.iconTextRadioButton(self.widgets["geoMoverMode"], edit = True, sl = geoMode)
+	
+
+	
+
+        
+    # # # # # # # # # # # # # # # # # # # # # # # # # # # # # # # # # # # # # # # # # # # # # # # # # # # # # # # # # # # # # # # 
+    # # # # # # # # # # # # # # # # # # # # # # # # # # # # # # # # # # # # # # # # # # # # # # # # # # # # # # # # # # # # # # # 
+    # # # # # # # # # # # # # # # # # # # # # # # # # # # # # # # # # # # # # # # # # # # # # # # # # # # # # # # # # # # # # # # 
+    # # # # # # # # # # # # # # # # # # # # # # # # # # # # # # # # # # # # # # # # # # # # # # # # # # # # # # # # # # # # # # # 
+    # # # # # # # # # # # # # # # # # # # # # # # # # # # # # # # # # # # # # # # # # # # # # # # # # # # # # # # # # # # # # # # 
+    def edit(self, *args):
+        
+	#check for physique slider values
+	sliders = ["upperTopTaperR","upperBotTaperR","upperMidTaperR","upperMidExpandR","lowerTopTaperR","lowerBotTaperR","lowerMidTaperR","lowerMidExpandR", "upperThighTopTaperR","upperThighBotTaperR","upperThighInnerR","upperThighOuterR",
+	           "calfTopTaperR","calfBotTaperR","calfMidTaperR","calfMidExpandR", "head_squareShape","head_topTaper","head_botTaper", "neck1_botTaper","neck1_topTaper","neck2_botTaper","neck2_topTaper","neck3_botTaper","neck3_topTaper",
+	           "body_weight","body_bust","spine1_botTaper","spine1_topTaper","spine2_botTaper","spine2_topTaper","spine3_botTaper","spine3_topTaper", "spine4_botTaper","spine4_topTaper", "spine5_botTaper", "spine5_topTaper",
+	           "hip_topTaper","hip_bottomTaper","hip_backside","hip_curvature", "upperTopTaperL","upperBotTaperL","upperMidTaperL","upperMidExpandL","lowerTopTaperL","lowerBotTaperL","lowerMidTaperL","lowerMidExpandL",
+	           "upperThighTopTaperL","upperThighBotTaperL","upperThighInnerL","upperThighOuterL","calfTopTaperL","calfBotTaperL","calfMidTaperL","calfMidExpandL", "shoulderSizeL", "shoulderSizeR", "elbowSizeL", "elbowSizeR"]
+
+	altered = []
+	for slider in sliders:
+	    val= cmds.floatSlider(self.widgets[slider], q = True, v = True)
+	    if val != 0:
+		altered.append(slider)
+		
+	if len(altered) > 0:
+	    result = cmds.confirmDialog(title = "Edit Skeleton", icon = "warning", message = "Current Physique settings will be lost if you continue.", button = ["Continue", "Save Physique Template", "Cancel"])
+	    
+	    if result == "Save Physique Template":
+		self.physiqueSave()
+		#hide the physique UI if it is visible
+		cmds.formLayout(self.widgets["jointMover_physiqueFormLayout"], edit = True, vis = False)		
+		
+	    if result == "Cancel":
+		return
+	    
+	
+
+    
+	#hide the physique UI if it is visible
+	cmds.formLayout(self.widgets["jointMover_physiqueFormLayout"], edit = True, vis = False)
+	
+	#if symmetry is on, turn off
+	if cmds.iconTextCheckBox(self.widgets["symmetry"], q = True, v = True):
+	    cmds.iconTextCheckBox(self.widgets["symmetry"], e = True, v = False)
+	    self.jointMoverUI_SymmetryModeOff()
+	    
+	#delete preview skeleton
+	if cmds.objExists("preview_root"):
+	    cmds.delete("preview_root")
+	    
+	    
+        #check to see if coming from lock state(skin weight)
+	if cmds.objExists("SceneLocked"):
+	    self.unlock()
+	    
+	    
+        #find all nodes in the joint mover and unlock those nodes
+        for parent in ["root", "lra", "middle_01_l_lra1", "geo_mover", "root_mover_grp", "Proxy_Geo_Grp", "extra_joints_dynamic", "extra_joints_template", "extra_joints_dynamic_geo", "extra_joints_template_geo", "character_facing_direction", "JointMover"]:
+	    try:
+		cmds.select(parent, hi = True, add = True)
+	    except:
+		pass
+	    
+        cmds.select("Skeleton_Settings", add = True)
+        nodes = cmds.ls(sl = True)
+        cmds.select(clear = True)
+	
+        for node in nodes:
+	    if cmds.objExists(node):
+		cmds.lockNode(node, lock = False)
+	    
+
+	#switch modes
+	self.switchMode("jointMover_bodyFormLayout",  None)
+	
+
+        #first go through global movers and offset movers and grab values(translation, rotation, and scale) on each one, writing it to file in format: [control, tx, ty, tz, rx, ry, rz, sx, sy, sz]\n
+        #create the temporary file to store data
+        path = self.mayaToolsDir + "/General/ART/editCache.txt"
+        f = open(path, 'w')
+        
+        #global movers first
+        for mover in self.globalMovers:
+            mover = mover.partition("|")[0]
+            tx = cmds.getAttr(mover + ".tx")
+            ty = cmds.getAttr(mover + ".ty")
+            tz = cmds.getAttr(mover + ".tz")
+            rx = cmds.getAttr(mover + ".rx")
+            ry = cmds.getAttr(mover + ".ry")
+            rz = cmds.getAttr(mover + ".rz")
+            sx = cmds.getAttr(mover + ".sx")
+            sy = cmds.getAttr(mover + ".sy")
+            sz = cmds.getAttr(mover + ".sz")
+            
+            info = [str(mover), tx, ty, tz, rx, ry, rz, sx, sy, sz]
+            f.write(str(info) + "\n")
+            
+            
+        #offset movers
+        for mover in self.offsetMovers:
+            mover = mover.partition("|")[0]
+            tx = cmds.getAttr(mover + ".tx")
+            ty = cmds.getAttr(mover + ".ty")
+            tz = cmds.getAttr(mover + ".tz")
+            rx = cmds.getAttr(mover + ".rx")
+            ry = cmds.getAttr(mover + ".ry")
+            rz = cmds.getAttr(mover + ".rz")
+            sx = cmds.getAttr(mover + ".sx")
+            sy = cmds.getAttr(mover + ".sy")
+            sz = cmds.getAttr(mover + ".sz")
+            
+            info = [str(mover), tx, ty, tz, rx, ry, rz, sx, sy, sz]
+            f.write(str(info) + "\n")
+	    
+	#geo movers
+        for mover in self.geoMovers:
+            mover = mover.partition("|")[0]
+            tx = cmds.getAttr(mover + ".tx")
+            ty = cmds.getAttr(mover + ".ty")
+            tz = cmds.getAttr(mover + ".tz")
+            rx = cmds.getAttr(mover + ".rx")
+            ry = cmds.getAttr(mover + ".ry")
+            rz = cmds.getAttr(mover + ".rz")
+            sx = cmds.getAttr(mover + ".sx")
+            sy = cmds.getAttr(mover + ".sy")
+            sz = cmds.getAttr(mover + ".sz")
+            
+            info = [str(mover), tx, ty, tz, rx, ry, rz, sx, sy, sz]
+            f.write(str(info) + "\n")
+            
+        f.close()
+        
+        #Now display the original UI
+        self.switchMode("formLayout",  "skeleton")
+	
+	
+	
+	
+	
+	
+	#read the settings from the skeleton settings cache node to fill in our skeleton settings ui
+	attrs = cmds.listAttr("SkeletonSettings_Cache")
+	for attr in attrs:
+        #set values
+	    if attr.find("extraJoint") != 0:
+		try:
+		    if cmds.intField(self.widgets[attr], q = True,exists = True):
+			data = cmds.getAttr("SkeletonSettings_Cache." + attr)
+			data = ast.literal_eval(data)
+			cmds.intField(self.widgets[attr], edit = True, value = data[0])
+			cmds.intField(self.widgets[attr], edit = True, bgc = data[1])
+			
+		except KeyError:
+		    pass
+        
+	for attr in attrs:
+	    if attr.find("extraJoint") != 0:
+		try:
+		    if cmds.checkBox(self.widgets[attr], q = True,exists = True):
+			data = cmds.getAttr("SkeletonSettings_Cache." + attr)
+			data = ast.literal_eval(data)
+			cmds.checkBox(self.widgets[attr], edit = True, value = data)
+			
+		except KeyError:
+		    pass
+		
+		
+	for attr in attrs:
+	    if attr.find("legStyle") != -1:
+		try:
+		    if cmds.optionMenu(self.widgets["legStyleOptionMenu"], q = True, exists = True):
+			data = cmds.getAttr("SkeletonSettings_Cache." + attr)
+			cmds.optionMenu(self.widgets["legStyleOptionMenu"], edit = True, value = data)
+			
+		except:
+		    pass
+
+        
+        #first, remove any extra rig bones in the UI
+	cmds.textScrollList(self.widgets["addedRigModulesList"], e = True, removeAll = True)
+
+	cmds.formLayout(self.widgets["rigModuleDefaultSettingsLayout"], edit = True, visible = True)
+	for widget in self.widgets:
+	    if widget.find("rigMod_") == 0:
+		if widget.find("_settingsForm") != -1:
+		    cmds.deleteUI(self.widgets[widget])
+	
+
+        
+        #add any extra child bones
+	extraJoints = []
+        for attr in attrs:
+            if attr.find("extraJoint") == 0:
+                extraJoints.append(attr)
+		
+	for joint in extraJoints:
+	    data =  cmds.getAttr("SkeletonSettings_Cache." + joint, asString = True)
+	    type = data.partition("/")[2].partition("/")[0]
+	    parent = data.partition("/")[0]
+	    jointName = data.rpartition("/")[2]
+	    
+	    
+	    
+	    if type == "leaf":
+		name = jointName.partition(" (")[0]
+		data = jointName.partition(" (")[2].partition(")")[0]
+		
+		translation = False
+		rotation = False
+		
+		
+		#get checkbox values
+		if data.find("T") != -1:
+		    translation = True
+		if data.find("R") != -1:
+		    rotation = True
+
+		
+		
+		#add to the list of added rig modules
+		cmds.textScrollList(self.widgets["addedRigModulesList"], edit = True, append = "[leaf] " + name, sc = self.rigMod_changeSettingsForms)
+		#select the newly added item
+		cmds.textScrollList(self.widgets["addedRigModulesList"], edit = True, si = "[leaf] " + name)
+		
+		#create the entry info and parent to the self.widgets["rigModuleSettingsFrame"]
+		#for leaf joints, we need name, parent, translate/rotate
+		
+		#hide existing formLayouts
+		self.rigMod_hideSettingsForms()
+		
+		self.widgets["rigMod_" + name + "_settingsForm"] = cmds.formLayout(w = 150, h = 220, parent = self.widgets["rigModuleSettingsFrame"])
+		label_name = cmds.text(label = "Name: ")
+		self.widgets["rigMod_" + name + "_settings_name"] = cmds.textField(w = 100, parent = self.widgets["rigMod_" + name + "_settingsForm"], text = name)
+		label_parent = cmds.text(label = "Parent: ")
+		self.widgets["rigMod_" + name + "_settings_parent"] = cmds.textField(w = 80, parent = self.widgets["rigMod_" + name + "_settingsForm"], text = parent, editable = False)
+		self.widgets["rigMod_" + name + "_settings_changeParent"] = cmds.button(label = "...", w = 20, parent = self.widgets["rigMod_" + name + "_settingsForm"], ann = "Change Parent", c = partial(self.extraJointInfo_UI_ParentList, name))
+		self.widgets["rigMod_" + name + "_settings_translate"] = cmds.checkBox(label = "Translate", v = translation, ann = "Joint can Translate")
+		self.widgets["rigMod_" + name + "_settings_rotate"] = cmds.checkBox(label = "Rotate", v = rotation, ann = "Joint can Rotate")
+		
+		cmds.formLayout(self.widgets["rigMod_" + name + "_settingsForm"], edit = True, af = [(label_name, "top", 5),(label_name, "left", 5)])
+		cmds.formLayout(self.widgets["rigMod_" + name + "_settingsForm"], edit = True, af = [(self.widgets["rigMod_" + name + "_settings_name"], "top", 5),(self.widgets["rigMod_" + name + "_settings_name"], "right", 5)])
+		cmds.formLayout(self.widgets["rigMod_" + name + "_settingsForm"], edit = True, af = [(label_parent, "top", 35),(label_parent, "left", 5)])
+		cmds.formLayout(self.widgets["rigMod_" + name + "_settingsForm"], edit = True, af = [(self.widgets["rigMod_" + name + "_settings_parent"], "top", 35),(self.widgets["rigMod_" + name + "_settings_parent"], "right", 25)])	    
+		cmds.formLayout(self.widgets["rigMod_" + name + "_settingsForm"], edit = True, af = [(self.widgets["rigMod_" + name + "_settings_changeParent"], "top", 35),(self.widgets["rigMod_" + name + "_settings_changeParent"], "right", 5)])	    
+		cmds.formLayout(self.widgets["rigMod_" + name + "_settingsForm"], edit = True, af = [(self.widgets["rigMod_" + name + "_settings_translate"], "top", 65),(self.widgets["rigMod_" + name + "_settings_translate"], "left", 5)])	    
+		cmds.formLayout(self.widgets["rigMod_" + name + "_settingsForm"], edit = True, af = [(self.widgets["rigMod_" + name + "_settings_rotate"], "top", 85),(self.widgets["rigMod_" + name + "_settings_rotate"], "left", 5)])	    
+		
+
+
+	    if type == "jiggle":
+		name = jointName.partition(" (")[0]
+		
+		#add to the list of added rig modules
+		cmds.textScrollList(self.widgets["addedRigModulesList"], edit = True, append = "[jiggle] " + name, sc = self.rigMod_changeSettingsForms)
+		
+		#select the newly added item
+		cmds.textScrollList(self.widgets["addedRigModulesList"], edit = True, si = "[jiggle] " + name)
+		
+		#create the entry info and parent to the self.widgets["rigModuleSettingsFrame"]
+		
+		#hide existing formLayouts
+		self.rigMod_hideSettingsForms()
+		
+		self.widgets["rigMod_" + name + "_settingsForm"] = cmds.formLayout(w = 150, h = 220, parent = self.widgets["rigModuleSettingsFrame"])
+		label_name = cmds.text(label = "Name: ")
+		self.widgets["rigMod_" + name + "_settings_name"] = cmds.textField(w = 100, parent = self.widgets["rigMod_" + name + "_settingsForm"], text = name)
+		label_parent = cmds.text(label = "Parent: ")
+		self.widgets["rigMod_" + name + "_settings_parent"] = cmds.textField(w = 80, parent = self.widgets["rigMod_" + name + "_settingsForm"], text = parent, editable = False)
+		self.widgets["rigMod_" + name + "_settings_changeParent"] = cmds.button(label = "...", w = 20, parent = self.widgets["rigMod_" + name + "_settingsForm"], ann = "Change Parent", c = partial(self.extraJointInfo_UI_ParentList, name))
+		
+		cmds.formLayout(self.widgets["rigMod_" + name + "_settingsForm"], edit = True, af = [(label_name, "top", 5),(label_name, "left", 5)])
+		cmds.formLayout(self.widgets["rigMod_" + name + "_settingsForm"], edit = True, af = [(self.widgets["rigMod_" + name + "_settings_name"], "top", 5),(self.widgets["rigMod_" + name + "_settings_name"], "right", 5)])
+		cmds.formLayout(self.widgets["rigMod_" + name + "_settingsForm"], edit = True, af = [(label_parent, "top", 35),(label_parent, "left", 5)])
+		cmds.formLayout(self.widgets["rigMod_" + name + "_settingsForm"], edit = True, af = [(self.widgets["rigMod_" + name + "_settings_parent"], "top", 35),(self.widgets["rigMod_" + name + "_settings_parent"], "right", 25)])	    
+		cmds.formLayout(self.widgets["rigMod_" + name + "_settingsForm"], edit = True, af = [(self.widgets["rigMod_" + name + "_settings_changeParent"], "top", 35),(self.widgets["rigMod_" + name + "_settings_changeParent"], "right", 5)])	    
+		
+		
+	    if type == "chain":
+		name = jointName.partition(" (")[0]
+		num = jointName.partition(" (")[2].partition(")")[0]
+		
+		
+		#add to the list of added rig modules
+		cmds.textScrollList(self.widgets["addedRigModulesList"], edit = True, append = "[chain] " + name, sc = self.rigMod_changeSettingsForms)
+		#select the newly added item
+		cmds.textScrollList(self.widgets["addedRigModulesList"], edit = True, si = "[chain] " + name)
+		
+		#create the entry info and parent to the self.widgets["rigModuleSettingsFrame"]
+		#for leaf joints, we need name, parent, translate/rotate
+		
+		#hide existing formLayouts
+		self.rigMod_hideSettingsForms()
+		
+		self.widgets["rigMod_" + name + "_settingsForm"] = cmds.formLayout(w = 150, h = 220, parent = self.widgets["rigModuleSettingsFrame"])
+		label_name = cmds.text(label = "Name: ")
+		self.widgets["rigMod_" + name + "_settings_name"] = cmds.textField(w = 100, parent = self.widgets["rigMod_" + name + "_settingsForm"], text = name)
+		label_parent = cmds.text(label = "Parent: ")
+		self.widgets["rigMod_" + name + "_settings_parent"] = cmds.textField(w = 80, parent = self.widgets["rigMod_" + name + "_settingsForm"], text = parent, editable = False)
+		self.widgets["rigMod_" + name + "_settings_changeParent"] = cmds.button(label = "...", w = 20, parent = self.widgets["rigMod_" + name + "_settingsForm"], ann = "Change Parent", c = partial(self.extraJointInfo_UI_ParentList, name))
+		label_number = cmds.text(label = "Count: ")
+		self.widgets["rigMod_" + name + "_settings_numberInChain"] = cmds.textField(w = 100, parent = self.widgets["rigMod_" + name + "_settingsForm"], text = str(num), ann = "Number of joints in the chain.")
+		
+		cmds.formLayout(self.widgets["rigMod_" + name + "_settingsForm"], edit = True, af = [(label_name, "top", 5),(label_name, "left", 5)])
+		cmds.formLayout(self.widgets["rigMod_" + name + "_settingsForm"], edit = True, af = [(self.widgets["rigMod_" + name + "_settings_name"], "top", 5),(self.widgets["rigMod_" + name + "_settings_name"], "right", 5)])
+		cmds.formLayout(self.widgets["rigMod_" + name + "_settingsForm"], edit = True, af = [(label_parent, "top", 35),(label_parent, "left", 5)])
+		cmds.formLayout(self.widgets["rigMod_" + name + "_settingsForm"], edit = True, af = [(self.widgets["rigMod_" + name + "_settings_parent"], "top", 35),(self.widgets["rigMod_" + name + "_settings_parent"], "right", 25)])	    
+		cmds.formLayout(self.widgets["rigMod_" + name + "_settingsForm"], edit = True, af = [(self.widgets["rigMod_" + name + "_settings_changeParent"], "top", 35),(self.widgets["rigMod_" + name + "_settings_changeParent"], "right", 5)])	    
+		cmds.formLayout(self.widgets["rigMod_" + name + "_settingsForm"], edit = True, af = [(label_number, "top", 65),(label_number, "left", 5)])
+		cmds.formLayout(self.widgets["rigMod_" + name + "_settingsForm"], edit = True, af = [(self.widgets["rigMod_" + name + "_settings_numberInChain"], "top", 65),(self.widgets["rigMod_" + name + "_settings_numberInChain"], "right", 5)])
+
+		
+
+	    
+	    
+
+        
+        #delete spine and neck bone buttons in the torsoModeFormLayout
+        children = cmds.formLayout(self.widgets["jointMover_torsoModeFormLayout"], q = True, childArray = True)
+        
+        if children != None:
+            for child in children:
+                if child.find("jointMoverUI_spine") != -1:
+                    cmds.deleteUI(child)
+                    
+                if child.find("jointMoverUI_neck") != -1:
+                    cmds.deleteUI(child) 
+		    
+		    
+	
+	#change flow chart buttons:
+	imagePath = self.mayaToolsDir + "/General/Icons/ART/"
+	
+	cmds.symbolButton(self.widgets["previousStep"], edit = True, visible = False, c = self.null)
+	cmds.symbolButton(self.widgets["nextStep"], edit = True, visible = True, image = imagePath + "skelPlacement.bmp", c = self.build)
+	cmds.tabLayout(self.widgets["tabLayout"], edit = True, tv = False, sti = 1)
+
+
+
+        #wipe joint mover from scene and create skeleton settings node
+	cmds.lockNode("JointMover", lock = False)
+	nodes.reverse()
+	
+        for node in nodes:
+	    if cmds.objExists(node):
+		try:
+		    cmds.lockNode(node, lock = False)
+		    cmds.delete(node)
+		    
+		except:
+		    pass
+	    
+	if cmds.objExists("JointMover"):
+	    try:
+		cmds.delete("JointMover")
+	    except:
+		pass
+	    
+	
+	#lock menu
+	cmds.menu(self.widgets["jointMoverTools"], edit = True, enable = False)
+	
+	    
+
+
+
+    # # # # # # # # # # # # # # # # # # # # # # # # # # # # # # # # # # # # # # # # # # # # # # # # # # # # # # # # # # # # # # # 
+    # # # # # # # # # # # # # # # # # # # # # # # # # # # # # # # # # # # # # # # # # # # # # # # # # # # # # # # # # # # # # # # 
+    # # # # # # # # # # # # # # # # # # # # # # # # # # # # # # # # # # # # # # # # # # # # # # # # # # # # # # # # # # # # # # # 
+    # # # # # # # # # # # # # # # # # # # # # # # # # # # # # # # # # # # # # # # # # # # # # # # # # # # # # # # # # # # # # # # 
+    # # # # # # # # # # # # # # # # # # # # # # # # # # # # # # # # # # # # # # # # # # # # # # # # # # # # # # # # # # # # # # # 
+    def jointMover_ChangeMoveMode(self, mode, commandType, *args):
+
+
+	cmds.lockNode("JointMover", lock = False)
+	
+	#GLOBAL MOVERS
+	if mode == "Global":
+	    if commandType == "on":
+		
+		for each in self.globalMovers:
+		    offsetMover = cmds.listRelatives(each.partition("|")[0], children = True, type = "transform")
+		    if offsetMover != None:
+			geoMover = cmds.listRelatives(offsetMover[0], children = True, type = "transform")
+			if geoMover != None:
+			    proxyGeo = cmds.listRelatives(geoMover[0], children = True, type = "transform")
+			    if proxyGeo != None:
+				cmds.lockNode(each, lock = False)
+				cmds.setAttr(each + ".visibility", lock = False, keyable = True)
+				cmds.setAttr(each + ".visibility", 1)
+				cmds.lockNode(each, lock = True)
+				
+	    if commandType == "off":
+                for each in self.globalMovers:
+		    offsetMover = cmds.listRelatives(each.partition("|")[0], children = True, type = "transform")
+		    if offsetMover != None:
+			geoMover = cmds.listRelatives(offsetMover[0], children = True, type = "transform")
+			if geoMover != None:
+			    proxyGeo = cmds.listRelatives(geoMover[0], children = True, type = "transform")
+			    if proxyGeo != None:
+				cmds.lockNode(each, lock = False)
+				cmds.setAttr(each + ".visibility", lock = False, keyable = True)
+				cmds.setAttr(each + ".visibility", 0)
+				cmds.lockNode(each, lock = True)
+				
+				
+	#OFFSET MOVERS
+	if mode == "Offset":
+	    if commandType == "on":
+                for each in self.offsetMovers:
+		    geoMover = cmds.listRelatives(each.partition("|")[0], children = True, type = "transform")
+		    if geoMover != None:
+			proxyGeo = cmds.listRelatives(geoMover[0], children = True, type = "transform")
+			if proxyGeo != None:
+			    cmds.lockNode(each, lock = False)
+			    cmds.setAttr(each + ".visibility", lock = False, keyable = True)
+			    cmds.setAttr(each + ".visibility", 1)
+			    cmds.lockNode(each, lock = True)
+			    
+	    if commandType == "off":
+                for each in self.offsetMovers:
+		    geoMover = cmds.listRelatives(each.partition("|")[0], children = True, type = "transform")
+		    if geoMover != None:
+			proxyGeo = cmds.listRelatives(geoMover[0], children = True, type = "transform")
+			if proxyGeo != None:
+			    cmds.lockNode(each, lock = False)
+			    cmds.setAttr(each + ".visibility", lock = False, keyable = True)
+			    cmds.setAttr(each + ".visibility", 0)
+			    cmds.lockNode(each, lock = True)
+			    
+			    
+	#GEOMETRY MOVERS
+	if mode == "Geo":
+	    if commandType == "on":
+                for each in self.geoMovers:
+		    proxyGeo = cmds.listRelatives(each.partition("|")[0], children = True, type = "transform")
+		    if proxyGeo != None:
+			cmds.setAttr(each + ".visibility", 1)
+			
+	    if commandType == "off":
+                for each in self.geoMovers:
+		    proxyGeo = cmds.listRelatives(each.partition("|")[0], children = True, type = "transform")
+		    if proxyGeo != None:
+			cmds.setAttr(each + ".visibility", 0)
+			
+			
+	#clear selection
+	cmds.select(clear = True)
+	
+	cmds.lockNode("JointMover", lock = True)
+			
+
+		
+	
+    # # # # # # # # # # # # # # # # # # # # # # # # # # # # # # # # # # # # # # # # # # # # # # # # # # # # # # # # # # # # # # # 
+    # # # # # # # # # # # # # # # # # # # # # # # # # # # # # # # # # # # # # # # # # # # # # # # # # # # # # # # # # # # # # # # 
+    # # # # # # # # # # # # # # # # # # # # # # # # # # # # # # # # # # # # # # # # # # # # # # # # # # # # # # # # # # # # # # # 
+    # # # # # # # # # # # # # # # # # # # # # # # # # # # # # # # # # # # # # # # # # # # # # # # # # # # # # # # # # # # # # # # 
+    # # # # # # # # # # # # # # # # # # # # # # # # # # # # # # # # # # # # # # # # # # # # # # # # # # # # # # # # # # # # # # # 
+    def toggleJointMoverControls_Vis(self, keyword, button, *args):
+
+        imageName = cmds.symbolButton(self.widgets[button], q = True, image = True)
+        mode = imageName.partition("ART/")[2].partition(".")[0]
+	if mode == "proxyGeo":
+	    mode = "visible"
+	
+
+
+        if mode == "visible":
+            
+            #change button display
+            cmds.symbolButton(self.widgets[button], edit = True, image = self.mayaToolsDir + "/General/Icons/ART/proxyGeoOff.bmp")
+            
+            if keyword != None:
+                for widget in self.widgets:
+                    if widget.find(keyword) != -1:
+			try:
+			    cmds.button(self.widgets[widget], edit = True, vis = False)
+			except:
+			    pass
+			
+
+            if keyword == "SKEL":
+                for each in self.skelMeshRefGeo:
+                    cmds.setAttr(each + ".visibility", 0)
+                    
+            if keyword == "Mesh":
+                for each in self.proxyGeo:
+                    cmds.setAttr(each + ".visibility", 0)
+                    
+                
+
+                
+                
+
+        else:
+            
+            cmds.symbolButton(self.widgets[button], edit = True, image = self.mayaToolsDir + "/General/Icons/ART/proxyGeo.bmp") 
+            
+            if keyword != None:
+                for widget in self.widgets:
+                    if widget.find(keyword) != -1:
+			try:
+			    cmds.button(self.widgets[widget], edit = True, vis = True)
+			    
+			except:
+			    pass
+
+            if keyword == "SKEL":
+                for each in self.skelMeshRefGeo:
+                    cmds.setAttr(each + ".visibility", 1)
+                    
+            if keyword == "Mesh":
+                for each in self.proxyGeo:
+                    cmds.setAttr(each + ".visibility", 1)
+                    
+    # # # # # # # # # # # # # # # # # # # # # # # # # # # # # # # # # # # # # # # # # # # # # # # # # # # # # # # # # # # # # # # 
+    # # # # # # # # # # # # # # # # # # # # # # # # # # # # # # # # # # # # # # # # # # # # # # # # # # # # # # # # # # # # # # # 
+    # # # # # # # # # # # # # # # # # # # # # # # # # # # # # # # # # # # # # # # # # # # # # # # # # # # # # # # # # # # # # # # 
+    # # # # # # # # # # # # # # # # # # # # # # # # # # # # # # # # # # # # # # # # # # # # # # # # # # # # # # # # # # # # # # # 
+    # # # # # # # # # # # # # # # # # # # # # # # # # # # # # # # # # # # # # # # # # # # # # # # # # # # # # # # # # # # # # # # 
+    def rigPose_worldSpaceHelper(self, *args):
+	
+	
+	#if symmetry is on, turn off
+	if cmds.iconTextCheckBox(self.widgets["symmetry"], q = True, v = True):
+	    cmds.iconTextCheckBox(self.widgets["symmetry"], e = True, v = False)
+	    self.jointMoverUI_SymmetryModeOff()
+	    
+	    
+	#run bake 2 global
+	self.transferToGlobal(True)
+	
+	legStyle = cmds.getAttr("SkeletonSettings_Cache.legStyle")
+	if legStyle == "Standard Biped":
+	    #create locators for the arms and legs and orient constrain the arms and legs global movers to the locators
+	    for mover in ["upperarm_mover_l", "upperarm_mover_r", "lowerarm_mover_l", "lowerarm_mover_r", "hand_mover_l", "hand_mover_r", "foot_mover_l", "foot_mover_r"]:
+		
+		#create a locator
+		locator = cmds.spaceLocator(name = mover + "_loc")[0]
+		
+		#constrain mover to locator
+		constraint = cmds.orientConstraint(locator, mover)[0]
+	    
+		#set offsets (-90x on l hand, 90x on r hand
+		
+		if mover == "upperarm_mover_r":
+		    cmds.setAttr(constraint + ".offsetX", 180)
+		    
+		if mover == "foot_mover_l":
+		    cmds.setAttr(constraint + ".offsetY", -90)
+		    
+		if mover == "foot_mover_r":
+		    cmds.setAttr(constraint + ".offsetY", 90)
+		    cmds.setAttr(constraint + ".offsetZ", 180)
+		    
+		if mover == "lowerarm_mover_r":
+		    cmds.setAttr(constraint + ".offsetX", 180)
+		    
+		if mover == "hand_mover_l":
+		    cmds.setAttr(constraint + ".offsetX", -90)
+		    
+		if mover == "hand_mover_r":
+		    cmds.setAttr(constraint + ".offsetX", 90)
+		    
+		#delete constraint and locator
+		cmds.delete(constraint)
+		cmds.delete(locator)
+		
+	
+	if legStyle == "Hind Leg":
+	    #create locators for the arms and legs and orient constrain the arms and legs global movers to the locators
+	    for mover in ["upperarm_mover_l", "upperarm_mover_r", "lowerarm_mover_l", "lowerarm_mover_r", "hand_mover_l", "hand_mover_r"]:
+		
+		#create a locator
+		locator = cmds.spaceLocator(name = mover + "_loc")[0]
+		
+		#constrain mover to locator
+		constraint = cmds.orientConstraint(locator, mover)[0]
+	    
+		#set offsets (-90x on l hand, 90x on r hand
+		
+		if mover == "upperarm_mover_r":
+		    cmds.setAttr(constraint + ".offsetX", 180)
+		    
+		if mover == "lowerarm_mover_r":
+		    cmds.setAttr(constraint + ".offsetX", 180)
+		    
+		if mover == "hand_mover_l":
+		    cmds.setAttr(constraint + ".offsetX", -90)
+		    
+		if mover == "hand_mover_r":
+		    cmds.setAttr(constraint + ".offsetX", 90)
+		    
+		#delete constraint and locator
+		cmds.delete(constraint)
+		cmds.delete(locator)
+	    
+		
+		
+    # # # # # # # # # # # # # # # # # # # # # # # # # # # # # # # # # # # # # # # # # # # # # # # # # # # # # # # # # # # # # # # 
+    # # # # # # # # # # # # # # # # # # # # # # # # # # # # # # # # # # # # # # # # # # # # # # # # # # # # # # # # # # # # # # # 
+    # # # # # # # # # # # # # # # # # # # # # # # # # # # # # # # # # # # # # # # # # # # # # # # # # # # # # # # # # # # # # # # 
+    # # # # # # # # # # # # # # # # # # # # # # # # # # # # # # # # # # # # # # # # # # # # # # # # # # # # # # # # # # # # # # # 
+    # # # # # # # # # # # # # # # # # # # # # # # # # # # # # # # # # # # # # # # # # # # # # # # # # # # # # # # # # # # # # # # 
+    def transferToGlobal(self, rigPose, *args):
+	
+	
+	#if symmetry is on, turn off
+	if cmds.iconTextCheckBox(self.widgets["symmetry"], q = True, v = True):
+	    cmds.iconTextCheckBox(self.widgets["symmetry"], e = True, v = False)
+	    self.jointMoverUI_SymmetryModeOff()
+	    
+	    
+	locators = []
+	constraints = []
+	
+	for mover in self.offsetMovers:
+	    
+	    mover = mover.partition("|")[0]
+	    
+	    #create a locator and position and orient at offset mover
+	    locator = cmds.spaceLocator(name = mover + "_loc")[0]
+	    constraint = cmds.parentConstraint(mover, locator)[0]
+	    cmds.delete(constraint)	    
+	    locators.append(locator)
+	    
+	
+	#zero out all offset movers
+	avoidAttrs = ["visibility", "scaleX", "scaleY", "scaleZ"]
+	
+	for mover in self.offsetMovers:
+	    mover = mover.partition("|")[0]
+	    
+	    attrs = cmds.listAttr(mover, keyable = True)
+	    for attr in attrs:
+		if attr not in avoidAttrs:
+		    cmds.setAttr(mover + "." + attr, 0)
+		
+	
+	
+	#snap global movers to locators
+	for mover in self.offsetMovers:
+	    mover = mover.partition("|")[0]
+	    globalMover = cmds.listRelatives(mover, parent = True)[0]
+	        
+	    for locator in locators:
+		if locator == mover + "_loc":
+
+		    grp = cmds.listRelatives(globalMover, parent = True)[0]
+		    parent = cmds.listRelatives(grp, parent = True)[0]
+		    offsets = cmds.listRelatives(parent, children = True)
+		    for offset in offsets:
+			if offset.find("_mover_offset") != -1:
+			    cmds.parent(locator, offset + "_loc")
+
+		    constraint = cmds.parentConstraint(locator, globalMover)[0]
+		    constraints.append(constraint)
+
+		    
+		    
+	if len(constraints) > 0:
+	    for constraint in constraints:
+		cmds.delete(constraint)
+		
+	for locator in locators:
+	    if cmds.objExists(locator):
+		cmds.delete(locator)
+	    
+	
+	#tell user they may want to reset their pose
+	if rigPose == False:
+	    cmds.confirmDialog(title = "Operation Complete", icon = "warning", message= "To ensure these changes stay, make sure to reset your Model Pose (and/or Rig Pose).")
+	
+	    
+		    
+	    
+    # # # # # # # # # # # # # # # # # # # # # # # # # # # # # # # # # # # # # # # # # # # # # # # # # # # # # # # # # # # # # # # 
+    # # # # # # # # # # # # # # # # # # # # # # # # # # # # # # # # # # # # # # # # # # # # # # # # # # # # # # # # # # # # # # # 
+    # # # # # # # # # # # # # # # # # # # # # # # # # # # # # # # # # # # # # # # # # # # # # # # # # # # # # # # # # # # # # # # 
+    # # # # # # # # # # # # # # # # # # # # # # # # # # # # # # # # # # # # # # # # # # # # # # # # # # # # # # # # # # # # # # # 
+    # # # # # # # # # # # # # # # # # # # # # # # # # # # # # # # # # # # # # # # # # # # # # # # # # # # # # # # # # # # # # # # 
+    def jointMover_compareModelPose(self, *args):
+	
+	#look at the current positions of everything compared to the recorded model pose. 
+	if cmds.iconTextCheckBox(self.widgets["symmetry"], q = True, v = True):
+	    cmds.iconTextCheckBox(self.widgets["symmetry"], e = True, v = False)
+	    self.jointMoverUI_SymmetryModeOff()
+	    
+	    
+	info = cmds.getAttr("Model_Pose.notes")
+	dirtyNodes = []
+	
+	poseInfo = info.splitlines()
+	for pose in poseInfo:
+	    pose = pose.lstrip("[")
+	    pose = pose.rstrip("]")
+	    splitString = pose.split(",")
+	    mover = splitString[0].strip("'")
+	    tx = splitString[1]
+	    ty = splitString[2]
+	    tz = splitString[3]
+	    rx = splitString[4]
+	    ry = splitString[5]
+	    rz = splitString[6]
+	    sx = splitString[7]
+	    sy = splitString[8]
+	    sz = splitString[9]
+	    
+	    
+	    if cmds.getAttr(mover + ".tx") != float(tx):
+		dirtyNodes.append(mover)
+	    if cmds.getAttr(mover + ".ty") != float(ty):
+		dirtyNodes.append(mover)
+	    if cmds.getAttr(mover + ".tz") != float(tz):
+		dirtyNodes.append(mover)
+		
+	    if cmds.getAttr(mover + ".rx") != float(rx):
+		dirtyNodes.append(mover)
+	    if cmds.getAttr(mover + ".ry") != float(ry):
+		dirtyNodes.append(mover)
+	    if cmds.getAttr(mover + ".rz") != float(rz):
+		dirtyNodes.append(mover)    
+		
+	    if cmds.getAttr(mover + ".sx") != float(sx):
+		dirtyNodes.append(mover)
+	    if cmds.getAttr(mover + ".sy") != float(sy):
+		dirtyNodes.append(mover)
+	    if cmds.getAttr(mover + ".sz") != float(sz):
+		dirtyNodes.append(mover)
+		
+	
+	if len(dirtyNodes) > 0:
+	    if cmds.objExists("Rig_Pose"):
+		if cmds.getAttr("Rig_Pose.v") == 0:
+    
+		    #launch a UI that will tell the user their changes will not get saved unless they reset the model pose
+		    result = cmds.confirmDialog(title = "Warning", icon = "warning", message = "Your character's model pose is different from what has been recorded. Would you like to update the model pose now?", button = ["Update[recommended]", "Ignore"])
+		    
+		    if result == "Ignore":
+			self.lock_phase1b(None)
+			
+		    if result == "Update[recommended]":
+			#reset the model pose
+			self.resetModelPose()
+			confirm = cmds.confirmDialog(title = "Rig Pose", icon = "question", message = "Depending on the type of changes made to your model pose, you may want to consider updating your rig pose. Would you like to do that now?", button = ["Update", "Ignore"])
+			
+			if confirm == "Ignore":
+			    self.lock_phase1b(None)
+			    
+			if confirm == "Update":
+			    #reset the rig pose
+			    self.resetRigPose(True)
+			
+	    
+		else:
+		    self.lock_phase1b(None)
+			
+	else:
+	    self.lock_phase1b(None)
+		
+
+
+
+	
+	
+	
+    # # # # # # # # # # # # # # # # # # # # # # # # # # # # # # # # # # # # # # # # # # # # # # # # # # # # # # # # # # # # # # # 
+    # # # # # # # # # # # # # # # # # # # # # # # # # # # # # # # # # # # # # # # # # # # # # # # # # # # # # # # # # # # # # # # 
+    # # # # # # # # # # # # # # # # # # # # # # # # # # # # # # # # # # # # # # # # # # # # # # # # # # # # # # # # # # # # # # # 
+    # # # # # # # # # # # # # # # # # # # # # # # # # # # # # # # # # # # # # # # # # # # # # # # # # # # # # # # # # # # # # # # 
+    # # # # # # # # # # # # # # # # # # # # # # # # # # # # # # # # # # # # # # # # # # # # # # # # # # # # # # # # # # # # # # # 
+    def jointMoverSelect(self, joint, side,  mode, *args):
+
+	selected = cmds.iconTextRadioCollection(self.widgets["jointMover_Modes_Collection"], q = True, sl = True)
+	
+	if selected == self.widgets["moverMode"].rpartition("|")[2]:
+	    mode ="global"
+	    
+	if selected == self.widgets["offsetMode"].rpartition("|")[2]:
+	    mode = "offset"
+	    
+	if selected == self.widgets["geoMoverMode"].rpartition("|")[2]:
+	    mode = "geo"
+	    
+
+        if mode == "global":
+            if side != None:
+                cmds.select(joint + "_mover_" + side)
+            
+            else:
+                cmds.select(joint + "_mover")
+            
+        if mode == "offset":
+            if side != None:
+                cmds.select(joint + "_mover_offset_" + side)
+                
+            else:
+                cmds.select(joint + "_mover_offset")
+            
+        if mode == "geo": 
+            if side != None:
+                cmds.select(joint + "_geo_mover_" + side)
+            
+            else:
+                cmds.select(joint + "_geo_mover")
+		
+		
+          
+    # # # # # # # # # # # # # # # # # # # # # # # # # # # # # # # # # # # # # # # # # # # # # # # # # # # # # # # # # # # # # # # 
+    # # # # # # # # # # # # # # # # # # # # # # # # # # # # # # # # # # # # # # # # # # # # # # # # # # # # # # # # # # # # # # # 
+    # # # # # # # # # # # # # # # # # # # # # # # # # # # # # # # # # # # # # # # # # # # # # # # # # # # # # # # # # # # # # # # 
+    # # # # # # # # # # # # # # # # # # # # # # # # # # # # # # # # # # # # # # # # # # # # # # # # # # # # # # # # # # # # # # # 
+    # # # # # # # # # # # # # # # # # # # # # # # # # # # # # # # # # # # # # # # # # # # # # # # # # # # # # # # # # # # # # # # 
+    def jointMoverUI_SymmetryMode(self, *args):
+	
+	if cmds.objExists("Symmetry_Nodes"):
+	    self.jointMoverUI_SymmetryModeOff()
+	    if cmds.objExists("Symmetry_Nodes"):
+		cmds.lockNode("Symmetry_Nodes", lock = False)
+		cmds.delete("Symmetry_Nodes")
+		
+		
+	previousSelection = cmds.ls(sl = True)
+        #while symmetry is activated, disable buttons we don't want to user to be able to use
+	cmds.symbolButton(self.widgets["reset"], edit = True, enable = False)
+	cmds.symbolButton(self.widgets["templatesButton"], edit = True, enable = False)
+	
+
+        #create a container to store the newly created nodes for symmetry
+        symmetryContainer = cmds.container(name = "Symmetry_Nodes")
+	
+        #create progress window
+	cmds.progressWindow(title='Symmetry Mode', progress = 0, isInterruptable=True, status = "Setting up symmetry..." )
+        
+        
+        #connect rotates and scales directly. Translate will need to go through a multiply node. Do this for Global, Offset, Geo, and LRA movers
+	moverLists = [self.globalMovers, self.offsetMovers, self.geoMovers]
+	for l in moverLists:
+	    cmds.progressWindow(edit = True, progress = 0)
+	    i = 0
+	    for mover in l:
+		mover = mover.partition("|")[0]
+		if cmds.objExists(mover + ".pinned"):
+		    if cmds.getAttr(mover + ".pinned") == 1.0:
+			continue
+		    
+		    
+		cmds.progressWindow(edit = True, progress = i)
+		
+	
+		#every standard joint mover node that is not custom created will fall under this category
+		if mover.rpartition("_")[2] == "l":
+		    #find the mirror mover
+		    mirrorMover = mover.rpartition("_")[0] + "_r"
+		    if cmds.objExists(mirrorMover + ".pinned"):
+			if cmds.getAttr(mirrorMover + ".pinned") == 1.0:
+			    continue
+			
+			
+		    multiplyNode = cmds.shadingNode("multiplyDivide", asUtility = True, name = mirrorMover + "_MultiplyNode")
+		    cmds.container(symmetryContainer, edit = True, addNode= multiplyNode,includeNetwork=True,includeHierarchyBelow=True)
+		    
+		    
+		    for keyword in ["mid_pivot", "tip_pivot", "pinky_pivot", "thumb_pivot"]:
+			if mover.find(keyword) == -1:
+			    #hook up connections between mover and its mirror
+			    try:
+				cmds.setAttr(multiplyNode + ".input2X", -1)
+				cmds.setAttr(multiplyNode + ".input2Y", -1)
+				cmds.setAttr(multiplyNode + ".input2Z", -1)
+			    except:
+				pass
+			    
+			    
+			else:
+			    #create the custom mirror attributes and hook them up to input2X
+			    cmds.lockNode(mover, lock = False)
+			    
+			    if cmds.objExists(mover + ".mirrorTX") == False:
+				cmds.addAttr(mover, ln = "mirrorTX", keyable = True, dv = 1)
+				
+			    if cmds.objExists(mover + ".mirrorTY") == False:
+				cmds.addAttr(mover, ln = "mirrorTY", keyable = True,  dv = 1)
+	    
+			    if cmds.objExists(mover + ".mirrorTZ") == False:
+				cmds.addAttr(mover, ln = "mirrorTZ", keyable = True,  dv = 1)
+				
+			    if cmds.objExists(mover + ".mirrorRX") == False:
+				cmds.addAttr(mover, ln = "mirrorRX", keyable = True, dv = 1)
+				
+			    if cmds.objExists(mover + ".mirrorRY") == False:
+				cmds.addAttr(mover, ln = "mirrorRY", keyable = True,  dv = 1)
+	    
+			    if cmds.objExists(mover + ".mirrorRZ") == False:
+				cmds.addAttr(mover, ln = "mirrorRZ", keyable = True,  dv = 1)
+				
+			    cmds.lockNode(mover, lock = True)			
+			    
+			    cmds.connectAttr( mover + '.mirrorTX', multiplyNode + '.input2X' )
+			    cmds.connectAttr( mover + '.mirrorTY', multiplyNode + '.input2Y' )
+			    cmds.connectAttr( mover + '.mirrorTZ', multiplyNode + '.input2Z' )
+			    
+			    #create the multiply node for rotations
+			    multiplyRNode = cmds.shadingNode("multiplyDivide", asUtility = True, name = mirrorMover + "_MultiplyNode")
+			    cmds.container(symmetryContainer, edit = True, addNode= multiplyRNode,includeNetwork=True,includeHierarchyBelow=True)
+			    
+			    cmds.connectAttr(mover + ".mirrorRX", multiplyRNode + ".input2X")
+			    cmds.connectAttr(mover + ".mirrorRY", multiplyRNode + ".input2Y")
+			    cmds.connectAttr(mover + ".mirrorRZ", multiplyRNode + ".input2Z")
+			    
+			    cmds.connectAttr( mover + '.rotateX', multiplyRNode + '.input1X' )
+			    cmds.connectAttr( mover + '.rotateY', multiplyRNode + '.input1Y' )
+			    cmds.connectAttr( mover + '.rotateZ', multiplyRNode + '.input1Z' )	
+			    
+			    
+			    cmds.connectAttr( multiplyRNode + '.outputX', mirrorMover + '.rotateX' )
+			    cmds.connectAttr( multiplyRNode + '.outputY', mirrorMover + '.rotateY' )
+			    cmds.connectAttr( multiplyRNode + '.outputZ', mirrorMover + '.rotateZ' )			    
+			    
+
+			
+				
+			    if mover.find("thumb_pivot") != -1:
+				cmds.setAttr(mover + ".mirrorTY", -1)
+				cmds.setAttr(mover + ".mirrorTZ", -1)
+				cmds.setAttr(mover + ".mirrorRY", -1)
+				cmds.setAttr(mover + ".mirrorRZ", -1)				
+				
+			    
+
+			    
+
+		    #continue hooking up default behavior
+		    try:
+			cmds.connectAttr( mover + '.rotate', mirrorMover + '.rotate' )
+			cmds.connectAttr( mover + '.scale', mirrorMover + '.scale' )
+			cmds.connectAttr( mover + '.translateX', multiplyNode + '.input1X' )
+			cmds.connectAttr( mover + '.translateY', multiplyNode + '.input1Y' )
+			cmds.connectAttr( mover + '.translateZ', multiplyNode + '.input1Z' )
+			
+			cmds.connectAttr( multiplyNode + '.outputX', mirrorMover + '.translateX' )
+			cmds.connectAttr( multiplyNode + '.outputY', mirrorMover + '.translateY' )
+			cmds.connectAttr( multiplyNode + '.outputZ', mirrorMover + '.translateZ' )
+			
+		    except:
+			pass
+
+		    #set the color of the mirror mover to a light grey to show connections
+		    cmds.setAttr(mirrorMover + ".overrideColor", 3)
+
+
+
+    
+		
+		#this will most likely only contain custom joints
+		mirrorSettings = "Mirror Suffix:"
+		mirrorLeft = "l_"
+		mirrorRight = "r_"
+		
+		
+		# # # # # # # # # # # # # # # # # # # # # # # # # # # # # # # # # # # # # # # # # # # # # #
+		# # # # # # # # # # # # # # # # # # # # # # # # # # # # # # # # # # # # # # # # # # # # # #
+		# # # # # # # # # # # # # # # # # # # # # # # # # # # # # # # # # # # # # # # # # # # # # #
+		#custom modules
+		try:
+		    mirrorSettings = cmds.getAttr("SkeletonSettings_Cache.mirrorSetting")
+		    mirrorLeft = cmds.getAttr("SkeletonSettings_Cache.mirrorLeft")
+		    mirrorRight = cmds.getAttr("SkeletonSettings_Cache.mirrorRight")
+		    
+		    cmds.optionMenu(self.widgets["rigModuleMirrorMode"], edit = True, v = mirrorSettings)
+		    cmds.textField(self.widgets["rigModuleMirrorTextL"], edit = True, text = mirrorLeft)
+		    cmds.textField(self.widgets["rigModuleMirrorTextR"], edit = True, text = mirrorRight)
+		    
+		    
+		    if mover.find(mirrorLeft) == 0:
+			mirrorMover = mirrorRight + mover.partition(mirrorLeft)[2]
+			if cmds.objExists(mirrorMover + ".pinned"):
+			    if cmds.getAttr(mirrorMover + ".pinned") == 1.0:
+				continue			
+			
+			#create multiply nodes and add to container
+			multiplyNode = cmds.shadingNode("multiplyDivide", asUtility = True, name = mirrorMover + "_MultiplyNode")
+			multiplyRNode = cmds.shadingNode("multiplyDivide", asUtility = True, name = mirrorMover + "_MultiplyNode")
+			cmds.container(symmetryContainer, edit = True, addNode= multiplyNode,includeNetwork=True,includeHierarchyBelow=True)
+			cmds.container(symmetryContainer, edit = True, addNode= multiplyRNode,includeNetwork=True,includeHierarchyBelow=True)
+	
+			
+			#add attrs to the mover so user can change mult node behavior
+			cmds.lockNode(mover, lock = False)
+			
+			if cmds.objExists(mover + ".mirrorTX") == False:
+			    cmds.addAttr(mover, ln = "mirrorTX", keyable = True, dv = -1)
+			    
+			if cmds.objExists(mover + ".mirrorTY") == False:
+			    cmds.addAttr(mover, ln = "mirrorTY", keyable = True,  dv = -1)
+	
+			if cmds.objExists(mover + ".mirrorTZ") == False:
+			    cmds.addAttr(mover, ln = "mirrorTZ", keyable = True,  dv = -1)
+			    
+			if cmds.objExists(mover + ".mirrorRX") == False:
+			    cmds.addAttr(mover, ln = "mirrorRX", keyable = True, dv = 1)
+			    
+			if cmds.objExists(mover + ".mirrorRY") == False:
+			    cmds.addAttr(mover, ln = "mirrorRY", keyable = True,  dv = 1)
+	
+			if cmds.objExists(mover + ".mirrorRZ") == False:
+			    cmds.addAttr(mover, ln = "mirrorRZ", keyable = True,  dv = 1)
+			    
+			cmds.lockNode(mover, lock = True)
+			
+			#hook up connections so mover will control its mirror
+			cmds.connectAttr(mover + ".mirrorTX", multiplyNode + ".input2X")
+			cmds.connectAttr(mover + ".mirrorTY", multiplyNode + ".input2Y")
+			cmds.connectAttr(mover + ".mirrorTZ", multiplyNode + ".input2Z")
+			
+			cmds.connectAttr(mover + ".mirrorRX", multiplyRNode + ".input2X")
+			cmds.connectAttr(mover + ".mirrorRY", multiplyRNode + ".input2Y")
+			cmds.connectAttr(mover + ".mirrorRZ", multiplyRNode + ".input2Z")
+	
+			
+			cmds.connectAttr( mover + '.scale', mirrorMover + '.scale' )
+			cmds.connectAttr( mover + '.translateX', multiplyNode + '.input1X' )
+			cmds.connectAttr( mover + '.translateY', multiplyNode + '.input1Y' )
+			cmds.connectAttr( mover + '.translateZ', multiplyNode + '.input1Z' )
+			
+			cmds.connectAttr( mover + '.rotateX', multiplyRNode + '.input1X' )
+			cmds.connectAttr( mover + '.rotateY', multiplyRNode + '.input1Y' )
+			cmds.connectAttr( mover + '.rotateZ', multiplyRNode + '.input1Z' )		
+	
+			cmds.connectAttr( multiplyNode + '.outputX', mirrorMover + '.translateX' )
+			cmds.connectAttr( multiplyNode + '.outputY', mirrorMover + '.translateY' )
+			cmds.connectAttr( multiplyNode + '.outputZ', mirrorMover + '.translateZ' )
+			
+			cmds.connectAttr( multiplyRNode + '.outputX', mirrorMover + '.rotateX' )
+			cmds.connectAttr( multiplyRNode + '.outputY', mirrorMover + '.rotateY' )
+			cmds.connectAttr( multiplyRNode + '.outputZ', mirrorMover + '.rotateZ' )
+
+
+			#set the color of the mirror mover to a light grey to show connections
+			cmds.setAttr(mirrorMover + ".overrideColor", 3)
+		    
+		    
+			
+			
+			
+		    if mover.find(mirrorLeft) != -1:
+			try:
+			    mirrorMover = mover.replace(mirrorLeft, mirrorRight)
+			    if cmds.objExists(mirrorMover + ".pinned"):
+				if cmds.getAttr(mirrorMover + ".pinned") == 1.0:
+				    continue			    
+
+			    #create multiply nodes and add to container
+			    multiplyNode = cmds.shadingNode("multiplyDivide", asUtility = True, name = mirrorMover + "_MultiplyNode")
+			    multiplyRNode = cmds.shadingNode("multiplyDivide", asUtility = True, name = mirrorMover + "_MultiplyNode")
+			    cmds.container(symmetryContainer, edit = True, addNode= multiplyNode,includeNetwork=True,includeHierarchyBelow=True)
+			    cmds.container(symmetryContainer, edit = True, addNode= multiplyRNode,includeNetwork=True,includeHierarchyBelow=True)
+	    
+			    
+			    #add attrs to the mover so user can change mult node behavior
+			    cmds.lockNode(mover, lock = False)
+			    
+			    if cmds.objExists(mover + ".mirrorTX") == False:
+				cmds.addAttr(mover, ln = "mirrorTX", keyable = True, dv = -1)
+				
+			    if cmds.objExists(mover + ".mirrorTY") == False:
+				cmds.addAttr(mover, ln = "mirrorTY", keyable = True,  dv = -1)
+	    
+			    if cmds.objExists(mover + ".mirrorTZ") == False:
+				cmds.addAttr(mover, ln = "mirrorTZ", keyable = True,  dv = -1)
+				
+			    if cmds.objExists(mover + ".mirrorRX") == False:
+				cmds.addAttr(mover, ln = "mirrorRX", keyable = True, dv = 1)
+				
+			    if cmds.objExists(mover + ".mirrorRY") == False:
+				cmds.addAttr(mover, ln = "mirrorRY", keyable = True,  dv = 1)
+	    
+			    if cmds.objExists(mover + ".mirrorRZ") == False:
+				cmds.addAttr(mover, ln = "mirrorRZ", keyable = True,  dv = 1)
+				
+			    cmds.lockNode(mover, lock = True)
+			    
+			    #hook up connections so mover will control its mirror
+			    cmds.connectAttr(mover + ".mirrorTX", multiplyNode + ".input2X")
+			    cmds.connectAttr(mover + ".mirrorTY", multiplyNode + ".input2Y")
+			    cmds.connectAttr(mover + ".mirrorTZ", multiplyNode + ".input2Z")
+			    
+			    cmds.connectAttr(mover + ".mirrorRX", multiplyRNode + ".input2X")
+			    cmds.connectAttr(mover + ".mirrorRY", multiplyRNode + ".input2Y")
+			    cmds.connectAttr(mover + ".mirrorRZ", multiplyRNode + ".input2Z")
+	    
+			    
+			    cmds.connectAttr( mover + '.scale', mirrorMover + '.scale' )
+			    cmds.connectAttr( mover + '.translateX', multiplyNode + '.input1X' )
+			    cmds.connectAttr( mover + '.translateY', multiplyNode + '.input1Y' )
+			    cmds.connectAttr( mover + '.translateZ', multiplyNode + '.input1Z' )
+			    
+			    cmds.connectAttr( mover + '.rotateX', multiplyRNode + '.input1X' )
+			    cmds.connectAttr( mover + '.rotateY', multiplyRNode + '.input1Y' )
+			    cmds.connectAttr( mover + '.rotateZ', multiplyRNode + '.input1Z' )		
+	    
+			    cmds.connectAttr( multiplyNode + '.outputX', mirrorMover + '.translateX' )
+			    cmds.connectAttr( multiplyNode + '.outputY', mirrorMover + '.translateY' )
+			    cmds.connectAttr( multiplyNode + '.outputZ', mirrorMover + '.translateZ' )
+			    
+			    cmds.connectAttr( multiplyRNode + '.outputX', mirrorMover + '.rotateX' )
+			    cmds.connectAttr( multiplyRNode + '.outputY', mirrorMover + '.rotateY' )
+			    cmds.connectAttr( multiplyRNode + '.outputZ', mirrorMover + '.rotateZ' )
+	    
+			    
+			    #set the color of the mirror mover to a light grey to show connections
+			    cmds.setAttr(mirrorMover + ".overrideColor", 3)
+			    
+			except:
+			    pass
+			
+		    
+		except:
+		    pass
+		# # # # # # # # # # # # # # # # # # # # # # # # # # # # # # # # # # # # # # # # # # # # # #
+		# # # # # # # # # # # # # # # # # # # # # # # # # # # # # # # # # # # # # # # # # # # # # #
+		# # # # # # # # # # # # # # # # # # # # # # # # # # # # # # # # # # # # # # # # # # # # # #
+		
+		
+		
+		#more checks for custom joints
+		if mover.find("_l_") != -1:
+		    prefix = mover.partition("_l_")[0]
+		    suffix = mover.partition("_l_")[2]
+		    mirrorMover = prefix + "_r_" + suffix
+		    if cmds.objExists(mirrorMover + ".pinned"):
+			if cmds.getAttr(mirrorMover + ".pinned") == 1.0:
+			    continue
+			
+		    
+		    try:
+			#create multiply nodes and add to container
+			multiplyNode = cmds.shadingNode("multiplyDivide", asUtility = True, name = mirrorMover + "_MultiplyNode")
+			multiplyRNode = cmds.shadingNode("multiplyDivide", asUtility = True, name = mirrorMover + "_MultiplyNode")
+			cmds.container(symmetryContainer, edit = True, addNode= multiplyNode,includeNetwork=True,includeHierarchyBelow=True)
+			cmds.container(symmetryContainer, edit = True, addNode= multiplyRNode,includeNetwork=True,includeHierarchyBelow=True)
+	
+			
+			#add attrs to the mover so user can change mult node behavior
+			cmds.lockNode(mover, lock = False)
+			
+			if cmds.objExists(mover + ".mirrorTX") == False:
+			    cmds.addAttr(mover, ln = "mirrorTX", keyable = True, dv = -1)
+			    
+			if cmds.objExists(mover + ".mirrorTY") == False:
+			    cmds.addAttr(mover, ln = "mirrorTY", keyable = True,  dv = -1)
+	
+			if cmds.objExists(mover + ".mirrorTZ") == False:
+			    cmds.addAttr(mover, ln = "mirrorTZ", keyable = True,  dv = -1)
+			    
+			if cmds.objExists(mover + ".mirrorRX") == False:
+			    cmds.addAttr(mover, ln = "mirrorRX", keyable = True, dv = 1)
+			    
+			if cmds.objExists(mover + ".mirrorRY") == False:
+			    cmds.addAttr(mover, ln = "mirrorRY", keyable = True,  dv = 1)
+	
+			if cmds.objExists(mover + ".mirrorRZ") == False:
+			    cmds.addAttr(mover, ln = "mirrorRZ", keyable = True,  dv = 1)
+			    
+			cmds.lockNode(mover, lock = True)
+			
+			#hook up connections so mover will control its mirror
+			cmds.connectAttr(mover + ".mirrorTX", multiplyNode + ".input2X")
+			cmds.connectAttr(mover + ".mirrorTY", multiplyNode + ".input2Y")
+			cmds.connectAttr(mover + ".mirrorTZ", multiplyNode + ".input2Z")
+			
+			cmds.connectAttr(mover + ".mirrorRX", multiplyRNode + ".input2X")
+			cmds.connectAttr(mover + ".mirrorRY", multiplyRNode + ".input2Y")
+			cmds.connectAttr(mover + ".mirrorRZ", multiplyRNode + ".input2Z")
+	
+			
+			cmds.connectAttr( mover + '.scale', mirrorMover + '.scale' )
+			cmds.connectAttr( mover + '.translateX', multiplyNode + '.input1X' )
+			cmds.connectAttr( mover + '.translateY', multiplyNode + '.input1Y' )
+			cmds.connectAttr( mover + '.translateZ', multiplyNode + '.input1Z' )
+			
+			cmds.connectAttr( mover + '.rotateX', multiplyRNode + '.input1X' )
+			cmds.connectAttr( mover + '.rotateY', multiplyRNode + '.input1Y' )
+			cmds.connectAttr( mover + '.rotateZ', multiplyRNode + '.input1Z' )		
+	
+			cmds.connectAttr( multiplyNode + '.outputX', mirrorMover + '.translateX' )
+			cmds.connectAttr( multiplyNode + '.outputY', mirrorMover + '.translateY' )
+			cmds.connectAttr( multiplyNode + '.outputZ', mirrorMover + '.translateZ' )
+			
+			cmds.connectAttr( multiplyRNode + '.outputX', mirrorMover + '.rotateX' )
+			cmds.connectAttr( multiplyRNode + '.outputY', mirrorMover + '.rotateY' )
+			cmds.connectAttr( multiplyRNode + '.outputZ', mirrorMover + '.rotateZ' )
+	
+			
+			#set the color of the mirror mover to a light grey to show connections
+			cmds.setAttr(mirrorMover + ".overrideColor", 3)
+			
+		    except:
+			pass
+		    
+		
+		#set defaults
+		for keyword in ["heel", "outside", "inside", "toe_pivot"]:
+		    if mover.find(keyword) == 0:
+			try:
+			    #set translate and rotate mirror defaults
+			    cmds.setAttr(mover + ".mirrorTX", 1)
+			    cmds.setAttr(mover + ".mirrorTY", 1)
+			    cmds.setAttr(mover + ".mirrorTZ", -1)
+	
+			    cmds.setAttr(mover + ".mirrorRX", 1)
+			    cmds.setAttr(mover + ".mirrorRY", -1)
+			    cmds.setAttr(mover + ".mirrorRZ", -1)
+			    
+			except:
+			    pass
+			
+			
+		#iterate counter
+		i = i + 1
+		
+
+        cmds.lockNode("Symmetry_Nodes", lock = True)
+                
+	self.symModeScriptJob = cmds.scriptJob(event = ["SceneSaved", self.symmetryModeScriptJob], runOnce = True)
+	if len(previousSelection) > 0:
+	    cmds.select(previousSelection)
+	    
+	    
+	#end progress
+	cmds.progressWindow(edit = True, endProgress=1)
+	
+	
+    # # # # # # # # # # # # # # # # # # # # # # # # # # # # # # # # # # # # # # # # # # # # # # # # # # # # # # # # # # # # # # # 
+    # # # # # # # # # # # # # # # # # # # # # # # # # # # # # # # # # # # # # # # # # # # # # # # # # # # # # # # # # # # # # # # 
+    # # # # # # # # # # # # # # # # # # # # # # # # # # # # # # # # # # # # # # # # # # # # # # # # # # # # # # # # # # # # # # # 
+    # # # # # # # # # # # # # # # # # # # # # # # # # # # # # # # # # # # # # # # # # # # # # # # # # # # # # # # # # # # # # # # 
+    # # # # # # # # # # # # # # # # # # # # # # # # # # # # # # # # # # # # # # # # # # # # # # # # # # # # # # # # # # # # # # # 
+    def symmetryModeScriptJob(self, *args):
+	self.jointMoverUI_SymmetryModeOff()
+	cmds.iconTextCheckBox(self.widgets["symmetry"] , edit = True, v = 0)
+	
+	cmds.file(save = True, force = True)
+
+    # # # # # # # # # # # # # # # # # # # # # # # # # # # # # # # # # # # # # # # # # # # # # # # # # # # # # # # # # # # # # # # 
+    # # # # # # # # # # # # # # # # # # # # # # # # # # # # # # # # # # # # # # # # # # # # # # # # # # # # # # # # # # # # # # # 
+    # # # # # # # # # # # # # # # # # # # # # # # # # # # # # # # # # # # # # # # # # # # # # # # # # # # # # # # # # # # # # # # 
+    # # # # # # # # # # # # # # # # # # # # # # # # # # # # # # # # # # # # # # # # # # # # # # # # # # # # # # # # # # # # # # # 
+    # # # # # # # # # # # # # # # # # # # # # # # # # # # # # # # # # # # # # # # # # # # # # # # # # # # # # # # # # # # # # # # 
+    def jointMoverUI_SymmetryModeOff(self, *args):
+	
+	#cmds.scriptJob(kill = self.symModeScriptJob)
+	#reenable deactivated toolbar buttons
+	cmds.symbolButton(self.widgets["reset"], edit = True, enable = True)
+	cmds.symbolButton(self.widgets["templatesButton"], edit = True, enable = True)
+	
+	
+        #delete the container which contains all of the connections
+	if cmds.objExists("Symmetry_Nodes"):
+	    cmds.lockNode("Symmetry_Nodes", lock = False)
+	    cmds.delete("Symmetry_Nodes")
+
+	else:
+	    return
+
+        
+        #disconnect rotates and scales from each mover type
+        for mover in self.globalMovers:
+
+	    try:
+		mover = mover.partition("|")[0]
+		if mover.rpartition("_")[2] == "l":
+		    mirrorMover = mover.rpartition("_")[0] + "_r"
+    
+		    
+		    cmds.disconnectAttr( mover + '.rotate', mirrorMover + '.rotate' )
+		    cmds.disconnectAttr( mover + '.scale', mirrorMover + '.scale' )
+		    
+		    cmds.setAttr(mirrorMover + ".overrideColor", 17)
+		    
+		#custom joints check
+		
+		if mover.find("l_") == 0:
+		    mirrorMover = "r_" + mover.partition("l_")[2]
+		    
+		    cmds.disconnectAttr( mover + '.scale', mirrorMover + '.scale' )
+		    cmds.setAttr(mirrorMover + ".overrideColor", 17)
+		    
+		    
+		if mover.find("_l_") != -1:
+		    prefix = mover.partition("_l_")[0]
+		    suffix = mover.partition("_l_")[2]
+		    mirrorMover = prefix + "_r_" + suffix
+		    
+		    cmds.disconnectAttr( mover + '.scale', mirrorMover + '.scale' )
+		    cmds.setAttr(mirrorMover + ".overrideColor", 17)
+
+	    except:
+		pass
+
+		
+		
+		
+        for mover in self.offsetMovers:
+	    try:
+		mover = mover.partition("|")[0]
+		if mover.rpartition("_")[2] == "l":
+		    mirrorMover = mover.rpartition("_")[0] + "_r"
+    
+		    
+		    cmds.disconnectAttr( mover + '.rotate', mirrorMover + '.rotate' )
+		    cmds.disconnectAttr( mover + '.scale', mirrorMover + '.scale' )
+		    
+		    cmds.setAttr(mirrorMover + ".overrideColor", 18)
+		    
+		    
+		if mover.find("l_") == 0:
+		    mirrorMover = "r_" + mover.partition("l_")[2]
+		    
+		    cmds.disconnectAttr( mover + '.scale', mirrorMover + '.scale' )
+		    cmds.setAttr(mirrorMover + ".overrideColor", 18)
+		    
+		    
+		if mover.find("_l_") != -1:
+		    prefix = mover.partition("_l_")[0]
+		    suffix = mover.partition("_l_")[2]
+		    mirrorMover = prefix + "_r_" + suffix
+		    
+		    cmds.disconnectAttr( mover + '.scale', mirrorMover + '.scale' )
+		    cmds.setAttr(mirrorMover + ".overrideColor", 18)
+		
+	    except:
+		pass
+
+		
+                
+        for mover in self.geoMovers:
+	    try:
+		mover = mover.partition("|")[0]
+		if mover.rpartition("_")[2] == "l":
+		    mirrorMover = mover.rpartition("_")[0] + "_r"
+    
+		    
+		    cmds.disconnectAttr( mover + '.rotate', mirrorMover + '.rotate' )
+		    cmds.disconnectAttr( mover + '.scale', mirrorMover + '.scale' )
+		    
+		    cmds.setAttr(mirrorMover + ".overrideColor", 20)
+		    
+    
+		if mover.find("l_") == 0:
+		    mirrorMover = "r_" + mover.partition("l_")[2]
+		    
+		    cmds.disconnectAttr( mover + '.scale', mirrorMover + '.scale' )
+		    cmds.setAttr(mirrorMover + ".overrideColor", 20)
+		    
+		    
+		if mover.find("_l_") != -1:
+		    prefix = mover.partition("_l_")[0]
+		    suffix = mover.partition("_l_")[2]
+		    mirrorMover = prefix + "_r_" + suffix
+		    
+		    cmds.disconnectAttr( mover + '.scale', mirrorMover + '.scale' )
+		    cmds.setAttr(mirrorMover + ".overrideColor", 20)
+
+	    except:
+		pass
+
+                
+
+    # # # # # # # # # # # # # # # # # # # # # # # # # # # # # # # # # # # # # # # # # # # # # # # # # # # # # # # # # # # # # # # 
+    # # # # # # # # # # # # # # # # # # # # # # # # # # # # # # # # # # # # # # # # # # # # # # # # # # # # # # # # # # # # # # # 
+    # # # # # # # # # # # # # # # # # # # # # # # # # # # # # # # # # # # # # # # # # # # # # # # # # # # # # # # # # # # # # # # 
+    # # # # # # # # # # # # # # # # # # # # # # # # # # # # # # # # # # # # # # # # # # # # # # # # # # # # # # # # # # # # # # # 
+    # # # # # # # # # # # # # # # # # # # # # # # # # # # # # # # # # # # # # # # # # # # # # # # # # # # # # # # # # # # # # # # 
+    def jointMoverUI_resetAll(self, *args):
+
+	
+	if cmds.iconTextCheckBox(self.widgets["symmetry"], q = True, v = True):
+	    self.jointMoverUI_SymmetryModeOff()
+	    
+	    
+        mods = cmds.getModifiers()
+
+        if mods == 0:
+	    
+	    for list in [self.globalMovers, self.offsetMovers, self.geoMovers]:
+		
+		for control in list:
+		    control= control.partition("|")[0]
+		    for attr in [[".tx", 0], [".ty", 0], [".tz", 0], [".rx", 0], [".ry", 0], [".rz", 0], [".sx", 1], [".sy", 1], [".sz", 1]]:
+			try:
+			    cmds.setAttr(control + attr[0], attr[1])
+			except:
+			    pass
+
+
+            
+        if (mods & 4) > 0:
+	    obj = cmds.ls(sl = True)[0]
+	    for attr in [[".tx", 0], [".ty", 0], [".tz", 0], [".rx", 0], [".ry", 0], [".rz", 0], [".sx", 1], [".sy", 1], [".sz", 1]]:
+		try:
+		    cmds.setAttr(obj + attr[0], attr[1])
+		    
+		except:
+		    pass
+
+            
+        if (mods & 8) > 0: 
+	    
+	    obj = cmds.ls(sl = True)[0]
+	    cmds.select(obj, r = True, hi = True)
+	    hierarchy = cmds.ls(sl = True)
+	    
+	    for control in hierarchy:
+		if control.find("|") == -1 and control.find("grp") == -1 and control.find("proxy") == -1 and control.find("Shape") == -1:
+		    for attr in [[".tx", 0], [".ty", 0], [".tz", 0], [".rx", 0], [".ry", 0], [".rz", 0], [".sx", 1], [".sy", 1], [".sz", 1]]:
+			try:
+			    cmds.setAttr(control + attr[0], attr[1])
+			except:
+			    pass
+		    
+    # # # # # # # # # # # # # # # # # # # # # # # # # # # # # # # # # # # # # # # # # # # # # # # # # # # # # # # # # # # # # # # 
+    # # # # # # # # # # # # # # # # # # # # # # # # # # # # # # # # # # # # # # # # # # # # # # # # # # # # # # # # # # # # # # # 
+    # # # # # # # # # # # # # # # # # # # # # # # # # # # # # # # # # # # # # # # # # # # # # # # # # # # # # # # # # # # # # # # 
+    # # # # # # # # # # # # # # # # # # # # # # # # # # # # # # # # # # # # # # # # # # # # # # # # # # # # # # # # # # # # # # # 
+    # # # # # # # # # # # # # # # # # # # # # # # # # # # # # # # # # # # # # # # # # # # # # # # # # # # # # # # # # # # # # # # 
+    def assumeModelPose(self, *args):
+	
+	#set the joint mover and the skeleton into model pose
+	if cmds.objExists("Model_Pose"):
+	    self.setModelPose_JM()
+	    if cmds.objExists("Skeleton_Model_Pose"):
+		self.setModelPose_Skel()
+		
+	else:
+	    cmds.confirmDialog(icon = "warning", title = "Error", message = "No pose found. It is possible you haven't created one yet during the place skeleton phase.")
+	    
+	
+    # # # # # # # # # # # # # # # # # # # # # # # # # # # # # # # # # # # # # # # # # # # # # # # # # # # # # # # # # # # # # # # 
+    # # # # # # # # # # # # # # # # # # # # # # # # # # # # # # # # # # # # # # # # # # # # # # # # # # # # # # # # # # # # # # # 
+    # # # # # # # # # # # # # # # # # # # # # # # # # # # # # # # # # # # # # # # # # # # # # # # # # # # # # # # # # # # # # # # 
+    # # # # # # # # # # # # # # # # # # # # # # # # # # # # # # # # # # # # # # # # # # # # # # # # # # # # # # # # # # # # # # # 
+    # # # # # # # # # # # # # # # # # # # # # # # # # # # # # # # # # # # # # # # # # # # # # # # # # # # # # # # # # # # # # # # 
+    def assumeRigPose(self, *args):
+	
+	#set the joint mover and the skeleton into model pose
+	if cmds.objExists("Rig_Pose"):
+	    self.setRigPose_JM()
+	    
+	    
+	if cmds.objExists("Skeleton_Rig_Pose"):
+	    if cmds.objExists("joint_mover_root") and cmds.objExists("root"):
+		self.setRigPose_Skel()
+		
+		
+	else:
+	    cmds.confirmDialog(icon = "warning", title = "Error", message = "No pose found. It is possible you haven't created one yet during the place skeleton phase.")
+	    
+	    
+    # # # # # # # # # # # # # # # # # # # # # # # # # # # # # # # # # # # # # # # # # # # # # # # # # # # # # # # # # # # # # # # 
+    # # # # # # # # # # # # # # # # # # # # # # # # # # # # # # # # # # # # # # # # # # # # # # # # # # # # # # # # # # # # # # # 
+    # # # # # # # # # # # # # # # # # # # # # # # # # # # # # # # # # # # # # # # # # # # # # # # # # # # # # # # # # # # # # # # 
+    # # # # # # # # # # # # # # # # # # # # # # # # # # # # # # # # # # # # # # # # # # # # # # # # # # # # # # # # # # # # # # # 
+    # # # # # # # # # # # # # # # # # # # # # # # # # # # # # # # # # # # # # # # # # # # # # # # # # # # # # # # # # # # # # # # 
+    def getModelPose(self):
+	
+	#create a model Pose node(empty group, hide it, parent under skeleton settings)
+	if cmds.objExists("Model_Pose"):
+	    cmds.delete("Model_Pose")
+	modelPoseNode = cmds.group(empty = True, name = "Model_Pose")
+	cmds.parent(modelPoseNode, "Skeleton_Settings")
+	
+	#get all of the position and rotation information for all mover controls add construct a list
+	poseInfo = ""
+
+        #global movers first
+        for mover in self.globalMovers:
+            mover = mover.partition("|")[0]
+            tx = cmds.getAttr(mover + ".tx")
+            ty = cmds.getAttr(mover + ".ty")
+            tz = cmds.getAttr(mover + ".tz")
+            rx = cmds.getAttr(mover + ".rx")
+            ry = cmds.getAttr(mover + ".ry")
+            rz = cmds.getAttr(mover + ".rz")
+            sx = cmds.getAttr(mover + ".sx")
+            sy = cmds.getAttr(mover + ".sy")
+            sz = cmds.getAttr(mover + ".sz")
+            
+            poseInfo += str([str(mover), tx, ty, tz, rx, ry, rz, sx, sy, sz]) + "\n"
+            
+            
+        #offset movers
+        for mover in self.offsetMovers:
+            mover = mover.partition("|")[0]
+            tx = cmds.getAttr(mover + ".tx")
+            ty = cmds.getAttr(mover + ".ty")
+            tz = cmds.getAttr(mover + ".tz")
+            rx = cmds.getAttr(mover + ".rx")
+            ry = cmds.getAttr(mover + ".ry")
+            rz = cmds.getAttr(mover + ".rz")
+            sx = cmds.getAttr(mover + ".sx")
+            sy = cmds.getAttr(mover + ".sy")
+            sz = cmds.getAttr(mover + ".sz")
+            
+	    poseInfo += str([str(mover), tx, ty, tz, rx, ry, rz, sx, sy, sz]) + "\n"
+	    
+        #geo movers
+        for mover in self.geoMovers:
+            mover = mover.partition("|")[0]
+            tx = cmds.getAttr(mover + ".tx")
+            ty = cmds.getAttr(mover + ".ty")
+            tz = cmds.getAttr(mover + ".tz")
+            rx = cmds.getAttr(mover + ".rx")
+            ry = cmds.getAttr(mover + ".ry")
+            rz = cmds.getAttr(mover + ".rz")
+            sx = cmds.getAttr(mover + ".sx")
+            sy = cmds.getAttr(mover + ".sy")
+            sz = cmds.getAttr(mover + ".sz")
+            
+	    poseInfo += str([str(mover), tx, ty, tz, rx, ry, rz, sx, sy, sz]) + "\n"
+	
+	
+	#write the poseInfo to the .notes attr of the modelPoseNode
+	if cmds.objExists(modelPoseNode + ".notes") == False:
+	    cmds.addAttr(modelPoseNode, sn = "nts", ln = "notes", dt = "string")
+	    
+	cmds.setAttr(modelPoseNode + ".notes", poseInfo, type = "string")
+    
+
+    # # # # # # # # # # # # # # # # # # # # # # # # # # # # # # # # # # # # # # # # # # # # # # # # # # # # # # # # # # # # # # # 
+    # # # # # # # # # # # # # # # # # # # # # # # # # # # # # # # # # # # # # # # # # # # # # # # # # # # # # # # # # # # # # # # 
+    # # # # # # # # # # # # # # # # # # # # # # # # # # # # # # # # # # # # # # # # # # # # # # # # # # # # # # # # # # # # # # # 
+    # # # # # # # # # # # # # # # # # # # # # # # # # # # # # # # # # # # # # # # # # # # # # # # # # # # # # # # # # # # # # # # 
+    # # # # # # # # # # # # # # # # # # # # # # # # # # # # # # # # # # # # # # # # # # # # # # # # # # # # # # # # # # # # # # # 
+    def createRigPose(self, skinWeight):
+	
+	#create a UI that has a 'done' button
+	if cmds.window("rigPoseCreationUI", exists = True):
+	    cmds.deleteUI("rigPoseCreationUI", window = True)
+	    
+	window = cmds.window("rigPoseCreationUI", w = 200, h = 130, sizeable = False, titleBar = True, titleBarMenu = True, mnb = False, mxb = False, title = "Create Rig Pose")
+	layout = cmds.columnLayout(w = 200, h = 130, bgc = [0,0,0])
+	
+	#banner
+	cmds.image(w = 200, h = 50, image = self.mayaToolsDir + "/General/Icons/ART/artBanner200px.bmp", parent = layout)
+	
+	buttonLayout = cmds.columnLayout(parent = layout, cat = ["both", 10], w = 200, rs = 5, bgc = [0,0,0])
+	cmds.button(w = 180, h = 50, label = "Save Rig Pose", c = partial(self.createRigPose_execute, skinWeight), parent = buttonLayout, bgc = [.3,.3,.3])
+	cmds.button(w = 180, h = 50, label = "Auto World Space(Rig Pose)", c = self.rigPose_worldSpaceHelper, parent = buttonLayout, bgc = [.3,.3,.3])
+	cmds.separator(h = 10, parent = layout)
+	
+	cmds.showWindow(window)
+	
+	
+    # # # # # # # # # # # # # # # # # # # # # # # # # # # # # # # # # # # # # # # # # # # # # # # # # # # # # # # # # # # # # # # 
+    # # # # # # # # # # # # # # # # # # # # # # # # # # # # # # # # # # # # # # # # # # # # # # # # # # # # # # # # # # # # # # # 
+    # # # # # # # # # # # # # # # # # # # # # # # # # # # # # # # # # # # # # # # # # # # # # # # # # # # # # # # # # # # # # # # 
+    # # # # # # # # # # # # # # # # # # # # # # # # # # # # # # # # # # # # # # # # # # # # # # # # # # # # # # # # # # # # # # # 
+    # # # # # # # # # # # # # # # # # # # # # # # # # # # # # # # # # # # # # # # # # # # # # # # # # # # # # # # # # # # # # # # 
+    def createRigPose_execute(self, skinWeight, *args):
+	
+	
+	#create a model Pose node(empty group, hide it, parent under skeleton settings)
+	rigPoseNode = cmds.group(empty = True, name = "Rig_Pose")
+	cmds.parent(rigPoseNode, "Skeleton_Settings")
+	
+	#get all of the position and rotation information for all mover controls add construct a list
+	poseInfo = ""
+
+        #global movers first
+        for mover in self.globalMovers:
+            mover = mover.partition("|")[0]
+            tx = cmds.getAttr(mover + ".tx")
+            ty = cmds.getAttr(mover + ".ty")
+            tz = cmds.getAttr(mover + ".tz")
+            rx = cmds.getAttr(mover + ".rx")
+            ry = cmds.getAttr(mover + ".ry")
+            rz = cmds.getAttr(mover + ".rz")
+            sx = cmds.getAttr(mover + ".sx")
+            sy = cmds.getAttr(mover + ".sy")
+            sz = cmds.getAttr(mover + ".sz")
+            
+            poseInfo += str([str(mover), tx, ty, tz, rx, ry, rz, sx, sy, sz]) + "\n"
+            
+            
+        #offset movers
+        for mover in self.offsetMovers:
+            mover = mover.partition("|")[0]
+            tx = cmds.getAttr(mover + ".tx")
+            ty = cmds.getAttr(mover + ".ty")
+            tz = cmds.getAttr(mover + ".tz")
+            rx = cmds.getAttr(mover + ".rx")
+            ry = cmds.getAttr(mover + ".ry")
+            rz = cmds.getAttr(mover + ".rz")
+            sx = cmds.getAttr(mover + ".sx")
+            sy = cmds.getAttr(mover + ".sy")
+            sz = cmds.getAttr(mover + ".sz")
+            
+	    poseInfo += str([str(mover), tx, ty, tz, rx, ry, rz, sx, sy, sz]) + "\n"
+	    
+        #geo movers
+        for mover in self.geoMovers:
+            mover = mover.partition("|")[0]
+            tx = cmds.getAttr(mover + ".tx")
+            ty = cmds.getAttr(mover + ".ty")
+            tz = cmds.getAttr(mover + ".tz")
+            rx = cmds.getAttr(mover + ".rx")
+            ry = cmds.getAttr(mover + ".ry")
+            rz = cmds.getAttr(mover + ".rz")
+            sx = cmds.getAttr(mover + ".sx")
+            sy = cmds.getAttr(mover + ".sy")
+            sz = cmds.getAttr(mover + ".sz")
+            
+	    poseInfo += str([str(mover), tx, ty, tz, rx, ry, rz, sx, sy, sz]) + "\n"
+	
+	
+	#write the poseInfo to the .notes attr of the rigPoseNode
+	if cmds.objExists(rigPoseNode + ".notes") == False:
+	    cmds.addAttr(rigPoseNode, sn = "nts", ln = "notes", dt = "string")
+	    
+	cmds.setAttr(rigPoseNode + ".notes", poseInfo, type = "string")
+	
+	#delete the UI
+	if cmds.window("rigPoseCreationUI", exists = True):
+	    cmds.deleteUI("rigPoseCreationUI")
+	    
+
+	if skinWeight:
+	    #auto-launch lock again
+	    if cmds.iconTextCheckBox(self.widgets["symmetry"], q = True, v = True):
+		cmds.iconTextCheckBox(self.widgets["symmetry"], e = True, v = False)
+		self.jointMoverUI_SymmetryModeOff()
+		
+	    self.lock_phase1b(None)
+	
+	
+	
+    # # # # # # # # # # # # # # # # # # # # # # # # # # # # # # # # # # # # # # # # # # # # # # # # # # # # # # # # # # # # # # # 
+    # # # # # # # # # # # # # # # # # # # # # # # # # # # # # # # # # # # # # # # # # # # # # # # # # # # # # # # # # # # # # # # 
+    # # # # # # # # # # # # # # # # # # # # # # # # # # # # # # # # # # # # # # # # # # # # # # # # # # # # # # # # # # # # # # # 
+    # # # # # # # # # # # # # # # # # # # # # # # # # # # # # # # # # # # # # # # # # # # # # # # # # # # # # # # # # # # # # # # 
+    # # # # # # # # # # # # # # # # # # # # # # # # # # # # # # # # # # # # # # # # # # # # # # # # # # # # # # # # # # # # # # # 
+    def getSkeletonModelPose(self):
+	if cmds.objExists("Skeleton_Model_Pose"):
+	    cmds.delete("Skeleton_Model_Pose")
+	    
+	    
+	cmds.select("root", hi = True)
+	joints = cmds.ls(sl = True, type = 'joint')
+	
+	
+	#create a model Pose node(empty group, hide it, parent under skeleton settings)
+	skeletonModelPoseNode = cmds.group(empty = True, name = "Skeleton_Model_Pose")
+	cmds.parent(skeletonModelPoseNode, "Skeleton_Settings")
+	
+	#get all of the position and rotation information for all joints and construct a list
+	poseInfo = ""
+
+        #global movers first
+        for joint in joints:
+	    if joint.find("twist") == -1:
+		tx = cmds.getAttr(joint + ".tx")
+		ty = cmds.getAttr(joint + ".ty")
+		tz = cmds.getAttr(joint + ".tz")
+		rx = cmds.getAttr(joint + ".rx")
+		ry = cmds.getAttr(joint + ".ry")
+		rz = cmds.getAttr(joint + ".rz")
+		
+		poseInfo += str([str(joint), tx, ty, tz, rx, ry, rz]) + "\n"
+
+	#write the poseInfo to the .notes attr of the skeletonModelPoseNode
+	if cmds.objExists(skeletonModelPoseNode + ".notes") == False:
+	    cmds.addAttr(skeletonModelPoseNode, sn = "nts", ln = "notes", dt = "string")
+	    
+	cmds.setAttr(skeletonModelPoseNode + ".notes", poseInfo, type = "string")
+	
+	
+    
+    # # # # # # # # # # # # # # # # # # # # # # # # # # # # # # # # # # # # # # # # # # # # # # # # # # # # # # # # # # # # # # # 
+    # # # # # # # # # # # # # # # # # # # # # # # # # # # # # # # # # # # # # # # # # # # # # # # # # # # # # # # # # # # # # # # 
+    # # # # # # # # # # # # # # # # # # # # # # # # # # # # # # # # # # # # # # # # # # # # # # # # # # # # # # # # # # # # # # # 
+    # # # # # # # # # # # # # # # # # # # # # # # # # # # # # # # # # # # # # # # # # # # # # # # # # # # # # # # # # # # # # # # 
+    # # # # # # # # # # # # # # # # # # # # # # # # # # # # # # # # # # # # # # # # # # # # # # # # # # # # # # # # # # # # # # # 
+    def getSkeletonRigPose(self):
+	if cmds.objExists("Skeleton_Rig_Pose"):
+	    cmds.delete("Skeleton_Rig_Pose")
+	    
+	    
+	cmds.select("root", hi = True)
+	joints = cmds.ls(sl = True, type = 'joint')
+	
+	
+	#create a model Pose node(empty group, hide it, parent under skeleton settings)
+	skeletonModelPoseNode = cmds.group(empty = True, name = "Skeleton_Rig_Pose")
+	cmds.parent(skeletonModelPoseNode, "Skeleton_Settings")
+	
+	#get all of the position and rotation information for all joints and construct a list
+	poseInfo = ""
+
+        #global movers first
+        for joint in joints:
+	    if joint.find("twist") == -1:
+		tx = cmds.getAttr(joint + ".tx")
+		ty = cmds.getAttr(joint + ".ty")
+		tz = cmds.getAttr(joint + ".tz")
+		rx = cmds.getAttr(joint + ".rx")
+		ry = cmds.getAttr(joint + ".ry")
+		rz = cmds.getAttr(joint + ".rz")
+		
+		poseInfo += str([str(joint), tx, ty, tz, rx, ry, rz]) + "\n"
+
+	#write the poseInfo to the .notes attr of the skeletonModelPoseNode
+	if cmds.objExists(skeletonModelPoseNode + ".notes") == False:
+	    cmds.addAttr(skeletonModelPoseNode, sn = "nts", ln = "notes", dt = "string")
+	    
+	cmds.setAttr(skeletonModelPoseNode + ".notes", poseInfo, type = "string")
+    
+    # # # # # # # # # # # # # # # # # # # # # # # # # # # # # # # # # # # # # # # # # # # # # # # # # # # # # # # # # # # # # # # 
+    # # # # # # # # # # # # # # # # # # # # # # # # # # # # # # # # # # # # # # # # # # # # # # # # # # # # # # # # # # # # # # # 
+    # # # # # # # # # # # # # # # # # # # # # # # # # # # # # # # # # # # # # # # # # # # # # # # # # # # # # # # # # # # # # # # 
+    # # # # # # # # # # # # # # # # # # # # # # # # # # # # # # # # # # # # # # # # # # # # # # # # # # # # # # # # # # # # # # # 
+    # # # # # # # # # # # # # # # # # # # # # # # # # # # # # # # # # # # # # # # # # # # # # # # # # # # # # # # # # # # # # # # 
+    def setModelPose_JM(self):
+	
+	if cmds.iconTextCheckBox(self.widgets["symmetry"], q = True, v = True):
+	    cmds.iconTextCheckBox(self.widgets["symmetry"], e = True, v = False)
+	    self.jointMoverUI_SymmetryModeOff()
+	    
+	    
+	info = cmds.getAttr("Model_Pose.notes")
+	
+	poseInfo = info.splitlines()
+	for pose in poseInfo:
+	    pose = pose.lstrip("[")
+	    pose = pose.rstrip("]")
+	    splitString = pose.split(",")
+	    mover = splitString[0].strip("'")
+	    tx = splitString[1]
+	    ty = splitString[2]
+	    tz = splitString[3]
+	    rx = splitString[4]
+	    ry = splitString[5]
+	    rz = splitString[6]
+	    sx = splitString[7]
+	    sy = splitString[8]
+	    sz = splitString[9]
+	    
+	    if cmds.getAttr(mover + ".tx", lock = True) == False:
+		cmds.setAttr(mover + ".tx", float(tx))
+	    if cmds.getAttr(mover + ".ty", lock = True) == False:
+		cmds.setAttr(mover + ".ty", float(ty))
+	    if cmds.getAttr(mover + ".tz", lock = True) == False:
+		cmds.setAttr(mover + ".tz", float(tz))
+	    
+	    if cmds.getAttr(mover + ".rx", lock = True) == False:
+		cmds.setAttr(mover + ".rx", float(rx))
+	    if cmds.getAttr(mover + ".ry", lock = True) == False:
+		cmds.setAttr(mover + ".ry", float(ry))
+	    if cmds.getAttr(mover + ".rz", lock = True) == False:
+		cmds.setAttr(mover + ".rz", float(rz))
+	    
+	    if cmds.getAttr(mover + ".sx", keyable = True) == True:
+		cmds.setAttr(mover + ".sx", float(sx))
+	    if cmds.getAttr(mover + ".sy", keyable = True) == True:
+		cmds.setAttr(mover + ".sy", float(sy))
+	    if cmds.getAttr(mover + ".sz", keyable = True) == True:
+		cmds.setAttr(mover + ".sz", float(sz))
+	
+	#turn the vis on for the rig pose node to show we are in that mode
+	if cmds.objExists("Rig_Pose"):
+	    cmds.setAttr("Rig_Pose.v", 0)
+	cmds.setAttr("Model_Pose.v", 1)
+	
+    # # # # # # # # # # # # # # # # # # # # # # # # # # # # # # # # # # # # # # # # # # # # # # # # # # # # # # # # # # # # # # # 
+    # # # # # # # # # # # # # # # # # # # # # # # # # # # # # # # # # # # # # # # # # # # # # # # # # # # # # # # # # # # # # # # 
+    # # # # # # # # # # # # # # # # # # # # # # # # # # # # # # # # # # # # # # # # # # # # # # # # # # # # # # # # # # # # # # # 
+    # # # # # # # # # # # # # # # # # # # # # # # # # # # # # # # # # # # # # # # # # # # # # # # # # # # # # # # # # # # # # # # 
+    # # # # # # # # # # # # # # # # # # # # # # # # # # # # # # # # # # # # # # # # # # # # # # # # # # # # # # # # # # # # # # # 
+    def setRigPose_JM(self):
+	
+	if cmds.iconTextCheckBox(self.widgets["symmetry"], q = True, v = True):
+	    cmds.iconTextCheckBox(self.widgets["symmetry"], e = True, v = False)
+	    
+	    self.jointMoverUI_SymmetryModeOff()
+	    
+	    
+	
+	info = cmds.getAttr("Rig_Pose.notes")
+	
+	poseInfo = info.splitlines()
+	for pose in poseInfo:
+	    pose = pose.lstrip("[")
+	    pose = pose.rstrip("]")
+	    splitString = pose.split(",")
+	    mover = splitString[0].strip("'")
+	    tx = splitString[1]
+	    ty = splitString[2]
+	    tz = splitString[3]
+	    rx = splitString[4]
+	    ry = splitString[5]
+	    rz = splitString[6]
+	    sx = splitString[7]
+	    sy = splitString[8]
+	    sz = splitString[9]
+	    
+	    if cmds.getAttr(mover + ".tx", lock = True) == False:
+		cmds.setAttr(mover + ".tx", float(tx))
+	    if cmds.getAttr(mover + ".ty", lock = True) == False:
+		cmds.setAttr(mover + ".ty", float(ty))
+	    if cmds.getAttr(mover + ".tz", lock = True) == False:
+		cmds.setAttr(mover + ".tz", float(tz))
+	    
+	    if cmds.getAttr(mover + ".rx", lock = True) == False:
+		cmds.setAttr(mover + ".rx", float(rx))
+	    if cmds.getAttr(mover + ".ry", lock = True) == False:
+		cmds.setAttr(mover + ".ry", float(ry))
+	    if cmds.getAttr(mover + ".rz", lock = True) == False:
+		cmds.setAttr(mover + ".rz", float(rz))
+	    
+	    if cmds.getAttr(mover + ".sx", keyable = True) == True:
+		cmds.setAttr(mover + ".sx", float(sx))
+	    if cmds.getAttr(mover + ".sy", keyable = True) == True:
+		cmds.setAttr(mover + ".sy", float(sy))
+	    if cmds.getAttr(mover + ".sz", keyable = True) == True:
+		cmds.setAttr(mover + ".sz", float(sz))
+		
+		
+	#turn the vis on for the rig pose node to show we are in that mode
+	cmds.setAttr("Rig_Pose.v", 1)
+	cmds.setAttr("Model_Pose.v", 0)
+	
+    # # # # # # # # # # # # # # # # # # # # # # # # # # # # # # # # # # # # # # # # # # # # # # # # # # # # # # # # # # # # # # # 
+    # # # # # # # # # # # # # # # # # # # # # # # # # # # # # # # # # # # # # # # # # # # # # # # # # # # # # # # # # # # # # # # 
+    # # # # # # # # # # # # # # # # # # # # # # # # # # # # # # # # # # # # # # # # # # # # # # # # # # # # # # # # # # # # # # # 
+    # # # # # # # # # # # # # # # # # # # # # # # # # # # # # # # # # # # # # # # # # # # # # # # # # # # # # # # # # # # # # # # 
+    # # # # # # # # # # # # # # # # # # # # # # # # # # # # # # # # # # # # # # # # # # # # # # # # # # # # # # # # # # # # # # # 
+    def resetModelPose(self, *args):
+	
+	
+	if cmds.objExists("Skeleton_Model_Pose"):
+	    cmds.delete("Skeleton_Model_Pose")
+	    cmds.delete("Model_Pose")
+	            
+	self.getModelPose()
+	self.getSkeletonModelPose()
+	
+	
+    # # # # # # # # # # # # # # # # # # # # # # # # # # # # # # # # # # # # # # # # # # # # # # # # # # # # # # # # # # # # # # # 
+    # # # # # # # # # # # # # # # # # # # # # # # # # # # # # # # # # # # # # # # # # # # # # # # # # # # # # # # # # # # # # # # 
+    # # # # # # # # # # # # # # # # # # # # # # # # # # # # # # # # # # # # # # # # # # # # # # # # # # # # # # # # # # # # # # # 
+    # # # # # # # # # # # # # # # # # # # # # # # # # # # # # # # # # # # # # # # # # # # # # # # # # # # # # # # # # # # # # # # 
+    # # # # # # # # # # # # # # # # # # # # # # # # # # # # # # # # # # # # # # # # # # # # # # # # # # # # # # # # # # # # # # # 
+    def resetRigPose(self, skinWeight, *args):
+	
+	if cmds.objExists("SceneLocked"):
+	    cmds.confirmDialog(icon = "warning", title = "Reset Pose", message  = "Cannot reset rig pose from this mode. You must be in Edit Placement mode.")
+	    return
+	
+	else:
+	    if cmds.objExists("Rig_Pose"):
+		cmds.delete("Rig_Pose")
+			
+	    self.createRigPose(skinWeight)
+
+    # # # # # # # # # # # # # # # # # # # # # # # # # # # # # # # # # # # # # # # # # # # # # # # # # # # # # # # # # # # # # # # 
+    # # # # # # # # # # # # # # # # # # # # # # # # # # # # # # # # # # # # # # # # # # # # # # # # # # # # # # # # # # # # # # # 
+    # # # # # # # # # # # # # # # # # # # # # # # # # # # # # # # # # # # # # # # # # # # # # # # # # # # # # # # # # # # # # # # 
+    # # # # # # # # # # # # # # # # # # # # # # # # # # # # # # # # # # # # # # # # # # # # # # # # # # # # # # # # # # # # # # # 
+    # # # # # # # # # # # # # # # # # # # # # # # # # # # # # # # # # # # # # # # # # # # # # # # # # # # # # # # # # # # # # # # 
+    def setModelPose_Skel(self):
+	
+	info = cmds.getAttr("Skeleton_Model_Pose.notes")
+	
+	poseInfo = info.splitlines()
+	for pose in poseInfo:
+	    pose = pose.lstrip("[")
+	    pose = pose.rstrip("]")
+	    splitString = pose.split(",")
+	    mover = splitString[0].strip("'")
+	    tx = splitString[1]
+	    ty = splitString[2]
+	    tz = splitString[3]
+	    rx = splitString[4]
+	    ry = splitString[5]
+	    rz = splitString[6]
+
+	    
+	    if cmds.getAttr(mover + ".tx", lock = True) == False:
+		cmds.setAttr(mover + ".tx", float(tx))
+	    if cmds.getAttr(mover + ".ty", lock = True) == False:
+		cmds.setAttr(mover + ".ty", float(ty))
+	    if cmds.getAttr(mover + ".tz", lock = True) == False:
+		cmds.setAttr(mover + ".tz", float(tz))
+	    
+	    if cmds.getAttr(mover + ".rx", lock = True) == False:
+		cmds.setAttr(mover + ".rx", float(rx))
+	    if cmds.getAttr(mover + ".ry", lock = True) == False:
+		cmds.setAttr(mover + ".ry", float(ry))
+	    if cmds.getAttr(mover + ".rz", lock = True) == False:
+		cmds.setAttr(mover + ".rz", float(rz))
+		
+		
+    # # # # # # # # # # # # # # # # # # # # # # # # # # # # # # # # # # # # # # # # # # # # # # # # # # # # # # # # # # # # # # # 
+    # # # # # # # # # # # # # # # # # # # # # # # # # # # # # # # # # # # # # # # # # # # # # # # # # # # # # # # # # # # # # # # 
+    # # # # # # # # # # # # # # # # # # # # # # # # # # # # # # # # # # # # # # # # # # # # # # # # # # # # # # # # # # # # # # # 
+    # # # # # # # # # # # # # # # # # # # # # # # # # # # # # # # # # # # # # # # # # # # # # # # # # # # # # # # # # # # # # # # 
+    # # # # # # # # # # # # # # # # # # # # # # # # # # # # # # # # # # # # # # # # # # # # # # # # # # # # # # # # # # # # # # # 
+    def setRigPose_Skel(self):
+	
+	info = cmds.getAttr("Skeleton_Rig_Pose.notes")
+	
+	poseInfo = info.splitlines()
+	for pose in poseInfo:
+	    pose = pose.lstrip("[")
+	    pose = pose.rstrip("]")
+	    splitString = pose.split(",")
+	    mover = splitString[0].strip("'")
+	    tx = splitString[1]
+	    ty = splitString[2]
+	    tz = splitString[3]
+	    rx = splitString[4]
+	    ry = splitString[5]
+	    rz = splitString[6]
+
+	    
+	    if cmds.getAttr(mover + ".tx", lock = True) == False:
+		cmds.setAttr(mover + ".tx", float(tx))
+	    if cmds.getAttr(mover + ".ty", lock = True) == False:
+		cmds.setAttr(mover + ".ty", float(ty))
+	    if cmds.getAttr(mover + ".tz", lock = True) == False:
+		cmds.setAttr(mover + ".tz", float(tz))
+	    
+	    if cmds.getAttr(mover + ".rx", lock = True) == False:
+		cmds.setAttr(mover + ".rx", float(rx))
+	    if cmds.getAttr(mover + ".ry", lock = True) == False:
+		cmds.setAttr(mover + ".ry", float(ry))
+	    if cmds.getAttr(mover + ".rz", lock = True) == False:
+		cmds.setAttr(mover + ".rz", float(rz))
+		
+		
+    # # # # # # # # # # # # # # # # # # # # # # # # # # # # # # # # # # # # # # # # # # # # # # # # # # # # # # # # # # # # # # # 
+    # # # # # # # # # # # # # # # # # # # # # # # # # # # # # # # # # # # # # # # # # # # # # # # # # # # # # # # # # # # # # # # 
+    # # # # # # # # # # # # # # # # # # # # # # # # # # # # # # # # # # # # # # # # # # # # # # # # # # # # # # # # # # # # # # # 
+    # # # # # # # # # # # # # # # # # # # # # # # # # # # # # # # # # # # # # # # # # # # # # # # # # # # # # # # # # # # # # # # 
+    # # # # # # # # # # # # # # # # # # # # # # # # # # # # # # # # # # # # # # # # # # # # # # # # # # # # # # # # # # # # # # # 
+    def skinProxyGeo_UI(self, deleteProxy):
+	
+	if cmds.window("skinProxyGeo_UI_question", exists = True):
+	    cmds.deleteUI("skinProxyGeo_UI_question")
+	    
+	window = cmds.window("skinProxyGeo_UI_question", title = "Skin Proxy Mesh", w = 500, h = 225, sizeable = False, titleBarMenu = False)
+	mainLayout = cmds.formLayout()
+	
+	imagePath = self.mayaToolsDir + "/General/Icons/ART/"
+	image = cmds.image(w = 200, h = 209, image= imagePath + "skinProxyGeo.jpg")
+	cmds.formLayout(mainLayout, edit = True, af = [(image, "top", 7), (image, "left", 5)])
+	
+	text = cmds.scrollField( w = 275, h = 145, editable=False, wordWrap=True, text= 'Would you like the tool to automatically skin weight the proxy mesh?\n\n(Useful if the final mesh is not ready for skin weights or if you just want to prototype character designs)' )
+	cmds.formLayout(mainLayout, edit = True, af = [(text, "top", 7), (text, "left", 215)])
+	
+	acceptButton = cmds.button(w = 135, h = 50, label = "Yes", c = partial(self.skinProxyGeo_UI_Accept, deleteProxy))
+	cancelButton = cmds.button(w = 135, h = 50, label = "No", c = partial(self.skinProxyGeo_UI_Cancel, deleteProxy))
+	
+	cmds.formLayout(mainLayout, edit = True, af = [(acceptButton, "bottom", 7), (acceptButton, "left", 215)])
+	cmds.formLayout(mainLayout, edit = True, af = [(cancelButton, "bottom", 7), (cancelButton, "right", 10)])
+	
+	cmds.showWindow(window)
+	
+    # # # # # # # # # # # # # # # # # # # # # # # # # # # # # # # # # # # # # # # # # # # # # # # # # # # # # # # # # # # # # # # 
+    # # # # # # # # # # # # # # # # # # # # # # # # # # # # # # # # # # # # # # # # # # # # # # # # # # # # # # # # # # # # # # # 
+    # # # # # # # # # # # # # # # # # # # # # # # # # # # # # # # # # # # # # # # # # # # # # # # # # # # # # # # # # # # # # # # 
+    # # # # # # # # # # # # # # # # # # # # # # # # # # # # # # # # # # # # # # # # # # # # # # # # # # # # # # # # # # # # # # # 
+    # # # # # # # # # # # # # # # # # # # # # # # # # # # # # # # # # # # # # # # # # # # # # # # # # # # # # # # # # # # # # # # 
+    def skinProxyGeo_UI_Accept(self, deleteProxy, *args):
+	cmds.deleteUI("skinProxyGeo_UI_question")
+
+	try:
+	    self.lock_phase2(deleteProxy, True)
+	
+	except Exception, e:
+
+	    try:
+		self.unlock()
+		cmds.select("skin_geo_*")
+		cmds.delete(cmds.ls(sl = True))
+
+	    except:
+		pass
+
+	    cmds.confirmDialog(title = "Naming Conflict", icon = 'critical', message = str(e) + ".\nAborting Operation. Please correct naming conflict and try again.")
+
+	    return
+
+
+
+
+    
+    # # # # # # # # # # # # # # # # # # # # # # # # # # # # # # # # # # # # # # # # # # # # # # # # # # # # # # # # # # # # # # # 
+    # # # # # # # # # # # # # # # # # # # # # # # # # # # # # # # # # # # # # # # # # # # # # # # # # # # # # # # # # # # # # # # 
+    # # # # # # # # # # # # # # # # # # # # # # # # # # # # # # # # # # # # # # # # # # # # # # # # # # # # # # # # # # # # # # # 
+    # # # # # # # # # # # # # # # # # # # # # # # # # # # # # # # # # # # # # # # # # # # # # # # # # # # # # # # # # # # # # # # 
+    # # # # # # # # # # # # # # # # # # # # # # # # # # # # # # # # # # # # # # # # # # # # # # # # # # # # # # # # # # # # # # # 
+    def skinProxyGeo_UI_Cancel(self, deleteProxy, *args):
+	cmds.deleteUI("skinProxyGeo_UI_question")
+
+
+	try:
+	    self.lock_phase2(deleteProxy, False)
+	
+	except Exception, e:
+
+	    try:
+		self.unlock()
+		cmds.select("skin_geo_*")
+		cmds.delete(cmds.ls(sl = True))
+
+	    except:
+		pass
+
+	    cmds.confirmDialog(title = "Naming Conflict", icon = 'critical', message = str(e) + ".\nAborting Operation. Please correct naming conflict and try again.")
+
+	    return
+	
+	
+    # # # # # # # # # # # # # # # # # # # # # # # # # # # # # # # # # # # # # # # # # # # # # # # # # # # # # # # # # # # # # # # 
+    # # # # # # # # # # # # # # # # # # # # # # # # # # # # # # # # # # # # # # # # # # # # # # # # # # # # # # # # # # # # # # # 
+    # # # # # # # # # # # # # # # # # # # # # # # # # # # # # # # # # # # # # # # # # # # # # # # # # # # # # # # # # # # # # # # 
+    # # # # # # # # # # # # # # # # # # # # # # # # # # # # # # # # # # # # # # # # # # # # # # # # # # # # # # # # # # # # # # # 
+    # # # # # # # # # # # # # # # # # # # # # # # # # # # # # # # # # # # # # # # # # # # # # # # # # # # # # # # # # # # # # # # 
+    def createRigPose_ConfirmationUI(self):
+	
+	if cmds.window("createRigPoseConfirmationUI", exists = True):
+	    cmds.deleteUI("createRigPoseConfirmationUI")
+	    
+	window = cmds.window("createRigPoseConfirmationUI", title = "Create Rig Pose", w = 500, h = 305, sizeable = False, titleBarMenu = False)
+	mainLayout = cmds.formLayout()
+	
+	#image
+	imagePath = self.mayaToolsDir + "/General/Icons/ART/"
+	image = cmds.image(w = 219, h = 300, image = imagePath + "tPoseExample.jpg")
+	text = cmds.scrollField(w = 254, h = 195, editable = False, wordWrap = True, text = "The ideal 'Rig Pose' is a T-Pose, like the one pictured to the left.\n\nThe key points are world axis aligned hands and feet so the control creation is nice and clean.\n\nNote: If you choose to create, pose your character into the desired rig pose, and then hit the [Save Rig Pose] button")
+	createNewButton = cmds.button(w = 126, h = 44, label = "Create Rig Pose", c = self.autoRigPose)
+	cancelButton = cmds.button(w = 126, h = 44, label = "Cancel", c = partial(self.lock_phase1b,"Cancel"))
+	
+	#place widgets
+	cmds.formLayout(mainLayout, edit = True, af = [(image, "top", 7),(image, "left", 6)])
+	cmds.formLayout(mainLayout, edit = True, af = [(text, "top", 7),(text, "left", 234)])
+	cmds.formLayout(mainLayout, edit = True, af = [(createNewButton, "bottom", 7),(createNewButton, "left", 234)])
+	cmds.formLayout(mainLayout, edit = True, af = [(cancelButton, "bottom", 7),(cancelButton, "left", 365)])
+	
+	cmds.showWindow(window)
+	
+	
+    # # # # # # # # # # # # # # # # # # # # # # # # # # # # # # # # # # # # # # # # # # # # # # # # # # # # # # # # # # # # # # # 
+    # # # # # # # # # # # # # # # # # # # # # # # # # # # # # # # # # # # # # # # # # # # # # # # # # # # # # # # # # # # # # # # 
+    # # # # # # # # # # # # # # # # # # # # # # # # # # # # # # # # # # # # # # # # # # # # # # # # # # # # # # # # # # # # # # # 
+    # # # # # # # # # # # # # # # # # # # # # # # # # # # # # # # # # # # # # # # # # # # # # # # # # # # # # # # # # # # # # # # 
+    # # # # # # # # # # # # # # # # # # # # # # # # # # # # # # # # # # # # # # # # # # # # # # # # # # # # # # # # # # # # # # # 
+    def autoRigPose(self, *args):
+	
+	result = cmds.confirmDialog(title = "Create Rig Pose", icon = "question", message = "Would you like the tool to put your character in a T Pose? You will still be able to edit the rig pose before finalizing.", button = ["Yes", "No"], defaultButton = "Yes", dismissString = "No", cancelButton = "No")
+	
+	if result == "Yes":
+	    self.rigPose_worldSpaceHelper()
+	    self.lock_phase1b("Create Rig Pose")
+	
+	if result == "No":
+	    self.lock_phase1b("Create Rig Pose")
+	    
+	
+	
+	
+    # # # # # # # # # # # # # # # # # # # # # # # # # # # # # # # # # # # # # # # # # # # # # # # # # # # # # # # # # # # # # # # 
+    # # # # # # # # # # # # # # # # # # # # # # # # # # # # # # # # # # # # # # # # # # # # # # # # # # # # # # # # # # # # # # # 
+    # # # # # # # # # # # # # # # # # # # # # # # # # # # # # # # # # # # # # # # # # # # # # # # # # # # # # # # # # # # # # # # 
+    # # # # # # # # # # # # # # # # # # # # # # # # # # # # # # # # # # # # # # # # # # # # # # # # # # # # # # # # # # # # # # # 
+    # # # # # # # # # # # # # # # # # # # # # # # # # # # # # # # # # # # # # # # # # # # # # # # # # # # # # # # # # # # # # # # 
+    def lock_phase1(self, *args):
+	
+	
+	#disable top buttons
+	cmds.symbolButton(self.widgets["previousStep"], edit = True, enable = False)
+	cmds.symbolButton(self.widgets["nextStep"], edit = True, enable = False)
+	cmds.menu(self.widgets["jointMoverTools"], edit = True, enable = False)
+	
+	
+	#hide the physique UI if it is visible
+	cmds.formLayout(self.widgets["jointMover_physiqueFormLayout"], edit = True, vis = False)
+	
+	#remove the preview skeleton
+	if cmds.objExists("preview_root"):
+	    cmds.delete("preview_root")
+	    
+	#first, capture the model pose(current pose of joint mover) and store
+	#create model pose and rig pose
+	if cmds.objExists("Model_Pose") == False:
+	    self.getModelPose()
+	    
+
+
+	    
+	#ask the user if they want to use the model pose as the rig pose, or create a separate rig pose
+	if cmds.objExists("Rig_Pose") == False:
+	    self.createRigPose_ConfirmationUI()
+	    
+	    
+
+	#check to see if the model pose is dirty
+	if cmds.objExists("Model_Pose"):
+	    if cmds.objExists("Rig_Pose"):
+		self.jointMover_compareModelPose()
+
+	else:
+	    self.lock_phase1b(None)
+	    
+	
+	    
+	    
+    # # # # # # # # # # # # # # # # # # # # # # # # # # # # # # # # # # # # # # # # # # # # # # # # # # # # # # # # # # # # # # # 
+    # # # # # # # # # # # # # # # # # # # # # # # # # # # # # # # # # # # # # # # # # # # # # # # # # # # # # # # # # # # # # # # 
+    # # # # # # # # # # # # # # # # # # # # # # # # # # # # # # # # # # # # # # # # # # # # # # # # # # # # # # # # # # # # # # # 
+    # # # # # # # # # # # # # # # # # # # # # # # # # # # # # # # # # # # # # # # # # # # # # # # # # # # # # # # # # # # # # # # 
+    # # # # # # # # # # # # # # # # # # # # # # # # # # # # # # # # # # # # # # # # # # # # # # # # # # # # # # # # # # # # # # # 
+    def lock_phase1b(self, result, *args):
+	
+	if cmds.window("createRigPoseConfirmationUI", exists = True):
+	    cmds.deleteUI("createRigPoseConfirmationUI")
+	
+	if result == "Use Existing":
+	    self.createRigPose_execute(True)
+	    return
+	    
+	if result == "Create Rig Pose":
+	    #tell the user to put the joint mover in the desired rig pose, and then have them hit the confirm button on the launched UI
+	    self.createRigPose(True)
+	    return
+	    
+	if result == "Cancel":
+	    cmds.symbolButton(self.widgets["previousStep"], edit = True, enable = True)
+	    cmds.symbolButton(self.widgets["nextStep"], edit = True, enable = True)		    
+	    return
+	    
+
+	#show geometry in case user has it hidden via vis menu
+	cmds.symbolButton(self.widgets["meshMode"], edit = True, image = self.mayaToolsDir + "/General/Icons/ART/proxyGeo.bmp")
+	
+	
+	for each in self.skelMeshRefGeo:
+	    cmds.setAttr(each + ".visibility", 1)
+	    
+	for each in self.proxyGeo:
+	    cmds.setAttr(each + ".visibility", 1)
+
+	for each in self.geoMovers:
+	    cmds.setAttr(each + ".visibility", 1)
+	
+	#create a node that let's the UI know the scene is locked in case the user reinits the UI
+	lockNode = cmds.group(empty = True, name = "SceneLocked")
+	cmds.lockNode(lockNode, lock = True)
+	
+
+	#get the joint mover back into rig pose (same pose as default skeleton build)
+	self.setRigPose_JM()
+	cmds.select("root")
+	cmds.setToolTo( 'moveSuperContext' )
+	cmds.refresh(force = True)
+	cmds.select(clear = True)
+	
+	
+	#Create Skeleton
+	self.buildFinalSkeleton()
+	self.getSkeletonRigPose()
+	
+
+	#Skin Proxy Geo if user wants
+	deleteProxy = False
+	self.skinProxyResult = ""
+	self.skinProxyGeo_UI(deleteProxy)
+	
+	#re-enable top buttons
+	cmds.symbolButton(self.widgets["previousStep"], edit = True, enable = True)
+	cmds.symbolButton(self.widgets["nextStep"], edit = True, enable = True)
+	cmds.tabLayout(self.widgets["tabLayout"], edit = True, tv = False, sti = 1)
+	
+	
+    # # # # # # # # # # # # # # # # # # # # # # # # # # # # # # # # # # # # # # # # # # # # # # # # # # # # # # # # # # # # # # # 
+    # # # # # # # # # # # # # # # # # # # # # # # # # # # # # # # # # # # # # # # # # # # # # # # # # # # # # # # # # # # # # # # 
+    # # # # # # # # # # # # # # # # # # # # # # # # # # # # # # # # # # # # # # # # # # # # # # # # # # # # # # # # # # # # # # # 
+    # # # # # # # # # # # # # # # # # # # # # # # # # # # # # # # # # # # # # # # # # # # # # # # # # # # # # # # # # # # # # # # 
+    # # # # # # # # # # # # # # # # # # # # # # # # # # # # # # # # # # # # # # # # # # # # # # # # # # # # # # # # # # # # # # # 
+    def lock_phase2(self, deleteProxy, skinWeight, *args):
+	
+	if skinWeight == False:
+	    deleteProxy = True
+	    
+	if skinWeight == True:
+	    cmds.menu(self.widgets["proxyMeshToolsMenu"], edit = True, enable = True)
+	    cmds.menu(self.widgets["jointMoverTools"], edit = True, enable = False)
+	
+	
+	self.skinProxyGeo()
+	
+        #find all nodes in the joint mover and unlock those nodes
+        cmds.select("root_mover_grp", r = True, hi = True)
+        cmds.select("Skeleton_Settings", add = True)
+        nodes = cmds.ls(sl = True, transforms = True)
+        
+        cmds.select(clear = True)
+	
+        for node in nodes:
+            cmds.lockNode(node, lock = False)
+	    
+	#hide the joint mover
+	if cmds.objExists("character_facing_direction"):
+	    cmds.setAttr("character_facing_direction.v", 0)
+	    cmds.setAttr("character_facing_direction.v", lock = True, keyable = False)
+	    
+        cmds.setAttr("root_mover_grp.v", 0)
+        cmds.setAttr("root_mover_grp.v", lock = True, keyable = False)
+	
+
+	#find all nodes in the joint mover and lock those nodes
+        cmds.select("root_mover_grp", r = True, hi = True)
+        cmds.select("Skeleton_Settings", add = True)
+        nodes = cmds.ls(sl = True, transforms = True)
+        
+        for node in nodes:
+            cmds.lockNode(node, lock = True)
+	    
+	
+	#lock down the remaining assets
+	    
+	cmds.lockNode("JointMover", lock = True)
+	    
+	     
+	#change front page to lock page
+	self.switchMode("jointMover_lockedFormLayout", None)
+	
+	
+	#disable the other toolbar buttons
+	cmds.iconTextCheckBox(self.widgets["aimMode"], edit = True, enable = False, vis = False)
+	cmds.iconTextCheckBox(self.widgets["symmetry"], edit = True, enable = False, vis = False)
+	cmds.symbolButton(self.widgets["reset"], edit = True, enable = False, vis = False)
+	cmds.symbolButton(self.widgets["templatesButton"], edit = True, enable = False, vis = False)
+	cmds.symbolButton(self.widgets["romButton"], edit = True, visible = True, vis = True)
+	cmds.iconTextCheckBox(self.widgets["tabletPressure"], edit = True, visible = True, vis = True)
+	cmds.symbolButton(self.widgets["paintWeightsButton"], edit = True, visible = True, vis = True)
+	cmds.symbolButton(self.widgets["mirrorWeightsButton"], edit = True, visible = True, vis = True)
+	cmds.symbolButton(self.widgets["pruneWeightsButton"], edit = True, visible = True, vis = True)
+	cmds.symbolButton(self.widgets["addInfluenceButton"], edit = True, visible = True, vis = True)
+	cmds.symbolButton(self.widgets["removeInfluenceButton"], edit = True, visible = True, vis = True)
+	cmds.symbolButton(self.widgets["exportSkinWeightsButton"], edit = True, visible = True, vis = True)
+	cmds.symbolButton(self.widgets["importSkinWeightsButton"], edit = True, visible = True, vis = True)
+	
+	cmds.iconTextRadioButton(self.widgets["moverMode"], edit = True, visible = True, vis = False)
+	cmds.iconTextRadioButton(self.widgets["offsetMode"], edit = True, visible = True, vis = False)
+	cmds.iconTextRadioButton(self.widgets["geoMoverMode"], edit = True, visible = True, vis = False)
+	cmds.symbolButton(self.widgets["physiqueMode"], edit = True, visible = True, vis = False)
+	cmds.symbolButton(self.widgets["meshMode"], edit = True, visible = True, vis = False)
+	
+	#change out help button command
+	cmds.symbolButton(self.widgets["helpButton"], edit = True, c = self.skeletonBuilderHelp)
+	
+	
+	
+
+	
+	#clear selection
+	cmds.select(clear = True)
+	
+	
+	#change flow chart buttons:
+	imagePath = self.mayaToolsDir + "/General/Icons/ART/"
+
+	cmds.symbolButton(self.widgets["previousStep"], edit = True, visible = True, image = imagePath + "skelPlacementBack.bmp", c = self.unlock)
+	cmds.symbolButton(self.widgets["nextStep"], edit = True, visible = True, image = imagePath + "buildRig.bmp", c = self.publish_UI)
+
+	
+	
+	#set joint mover and skeleton to model pose
+	constraints = []
+	for mover in self.geoMovers:
+	    mover = mover.partition("|")[0]
+	    children = cmds.listRelatives(mover, children = True, type = "transform")
+	    if children != None or mover.find("ball") == 0:
+
+	    
+		prefix = mover.partition("_geo_mover")[0]
+		suffix = mover.partition("_geo_mover")[2]
+		lra = prefix + suffix + "_lra"
+		jointName = prefix + suffix
+		constraint = cmds.parentConstraint(lra, jointName)[0]
+		constraints.append(constraint)
+		
+	self.setModelPose_JM()
+	
+	cmds.select("root")
+	cmds.setToolTo( 'moveSuperContext' )
+	cmds.refresh(force = True)
+	cmds.select(clear = True)
+	
+	if cmds.objExists("Skeleton_Model_Pose"):
+	    cmds.delete("Skeleton_Model_Pose")
+	    
+	self.getSkeletonModelPose()
+	cmds.delete(constraints)
+	self.setModelPose_Skel()
+	
+	cmds.select("root")
+	cmds.setToolTo( 'moveSuperContext' )
+	cmds.refresh(force = True)
+	cmds.select(clear = True)
+	
+	
+	#load cached weights if any
+	self.loadCachedWeights()
+	
+	if deleteProxy:
+	    cmds.delete(["Proxy_Geo_Skin_Grp", "proxy_geometry_layer"])
+	
+
+	#lastly, reset bindpose on root
+	cmds.select("*bindPose*")
+	bindposes = cmds.ls(sl = True)
+	cmds.delete(bindposes)
+	
+	cmds.select("root")
+	cmds.dagPose(s = True, bp = True)
+	
+	if deleteProxy == False:
+	    #auto weight meshes
+	    meshes = cmds.ls(type = "mesh")
+	    skinnableGeo = []
+	    for mesh in meshes:
+		parent = cmds.listRelatives(mesh, parent = True, type = "transform")[0]
+		if parent != None:
+		    if parent.find("proxy_geo") == -1:
+			if parent.find("lra") == -1:
+			    if parent.find("extra_joints") == -1:
+				if parent.find("skin_geo") == -1:
+				    skinnableGeo.append(parent)
+				
+	    if len(skinnableGeo) > 0:
+		self.autoWeightMeshes(skinnableGeo)
+		
+		
+	#unlock skeleton
+	cmds.select("root", hi = True)
+	cmds.lockNode(lock = False)
+	cmds.select(clear = True)
+	
+
+    # # # # # # # # # # # # # # # # # # # # # # # # # # # # # # # # # # # # # # # # # # # # # # # # # # # # # # # # # # # # # # # 
+    # # # # # # # # # # # # # # # # # # # # # # # # # # # # # # # # # # # # # # # # # # # # # # # # # # # # # # # # # # # # # # # 
+    # # # # # # # # # # # # # # # # # # # # # # # # # # # # # # # # # # # # # # # # # # # # # # # # # # # # # # # # # # # # # # # 
+    # # # # # # # # # # # # # # # # # # # # # # # # # # # # # # # # # # # # # # # # # # # # # # # # # # # # # # # # # # # # # # # 
+    # # # # # # # # # # # # # # # # # # # # # # # # # # # # # # # # # # # # # # # # # # # # # # # # # # # # # # # # # # # # # # # 
+    def unlock(self, *args):
+	
+	#cache any weights on the character meshes
+	continueProc = self.cacheWeights()
+	if continueProc:
+	    #disable proxy geo menu
+	    cmds.menu(self.widgets["proxyMeshToolsMenu"], edit = True, enable = False)
+	    cmds.menu(self.widgets["jointMoverTools"], edit = True, enable = True)
+	    cmds.tabLayout(self.widgets["tabLayout"], edit = True, tv = True)
+	    
+	    
+	    #put joint mover back in model pose
+	    if cmds.objExists("Model_Pose"):
+		self.setModelPose_JM()
+	    
+	    
+	    #delete the proxy geo, and the final skeleton
+	    cmds.select("root", hi = True)
+	    joints = cmds.ls(sl = True)
+	    
+	    for joint in joints:
+		cmds.lockNode(joint, lock = False)
+		
+	    for obj in ["Proxy_Geo_Skin_Grp", "root", "proxy_geometry_layer", "SceneLocked"]:
+		if cmds.objExists(obj):
+		    try:
+			cmds.lockNode(obj, lock = False)
+			cmds.delete(obj)
+		    except:
+			pass
+	    
+    
+	    
+	    #set default control visibility
+	    for mover in self.globalMovers:
+		try:
+		    cmds.setAttr(mover + ".v", 1)
+		except:
+		    pass
+		
+	    for mover in self.offsetMovers:
+		try:
+		    cmds.setAttr(mover + ".v", 0)
+		except:
+		    pass
+		
+	    for mover in self.geoMovers:
+		try:
+		    cmds.setAttr(mover + ".v", 0)
+		except:
+		    pass
+		
+	    for geo in self.proxyGeo:
+		try:
+		    cmds.setAttr(geo + ".v", 1)
+		except:
+		    pass
+		
+	    for geo in self.skelMeshRefGeo:
+		cmds.setAttr(geo + ".v", 1)
+		
+		
+	    #find all nodes in the joint mover and unlock those nodes
+	    cmds.lockNode("JointMover", lock = False)
+	    cmds.select("root_mover_grp", r = True, hi = True)
+	    cmds.select("Skeleton_Settings", add = True)
+	    nodes = cmds.ls(sl = True, transforms = True)
+	    
+	    cmds.select(clear = True)
+	    
+	    for node in nodes:
+		cmds.lockNode(node, lock = False)
+		
+	    cmds.rename("joint_mover_root", "root")
+		    
+	    #Unhide the joint mover
+	    cmds.setAttr("root_mover_grp.v", lock = False, keyable = True)
+	    cmds.setAttr("root_mover_grp.v", 1)
+	    
+	    if cmds.objExists("character_facing_direction"):
+		cmds.setAttr("character_facing_direction.v", lock = False, keyable = True)
+		cmds.setAttr("character_facing_direction.v", 1)
+	    
+	    
+	    
+	    #find all nodes in the join mover and lock those nodes
+	    cmds.select("root_mover_grp", r = True, hi = True)
+	    cmds.select("Skeleton_Settings", add = True)
+	    nodes = cmds.ls(sl = True, transforms = True)
+	    
+	    for node in nodes:
+		cmds.lockNode(node, lock = True)
+	    
+    
+	    #change front page back to torso page and disable it
+	    self.switchMode("jointMover_bodyFormLayout", None)
+	    
+	    #re enable/show buttons on toolbar
+	    imagePath = self.mayaToolsDir + "/General/Icons/ART/"
+	    cmds.iconTextCheckBox(self.widgets["aimMode"], edit = True, enable = True, vis = True)
+	    cmds.iconTextCheckBox(self.widgets["symmetry"], edit = True, enable = True, vis = True)
+	    cmds.symbolButton(self.widgets["reset"], edit = True, enable = True, vis = True)
+	    cmds.symbolButton(self.widgets["templatesButton"], edit = True, enable = True, vis = True)
+	    cmds.symbolButton(self.widgets["romButton"], edit = True, visible = False)
+	    cmds.symbolButton(self.widgets["paintWeightsButton"], edit = True, visible = False)
+	    cmds.iconTextCheckBox(self.widgets["tabletPressure"], edit = True, visible = False)
+	    cmds.symbolButton(self.widgets["mirrorWeightsButton"], edit = True, visible = False)
+	    cmds.symbolButton(self.widgets["pruneWeightsButton"], edit = True, visible = False)
+	    cmds.symbolButton(self.widgets["addInfluenceButton"], edit = True, visible = False)
+	    cmds.symbolButton(self.widgets["removeInfluenceButton"], edit = True, visible = False)
+	    cmds.symbolButton(self.widgets["exportSkinWeightsButton"], edit = True, visible = False)
+	    cmds.symbolButton(self.widgets["importSkinWeightsButton"], edit = True, visible = False)
+	    
+	    cmds.iconTextRadioButton(self.widgets["moverMode"], edit = True, visible = True, vis = True)
+	    cmds.iconTextRadioButton(self.widgets["offsetMode"], edit = True, visible = True, vis = True)
+	    cmds.iconTextRadioButton(self.widgets["geoMoverMode"], edit = True, visible = True, vis = True)
+	    cmds.symbolButton(self.widgets["physiqueMode"], edit = True, visible = True, vis = True)
+	    cmds.symbolButton(self.widgets["meshMode"], edit = True, visible = True, vis = True)
+	    
+	    #change out help button command
+	    cmds.symbolButton(self.widgets["helpButton"], edit = True, c = self.skeletonBuilderHelp)
+    
+	    
+	    #clear selection
+	    cmds.select(clear = True)
+	    
+	    #change flow chart buttons:
+	    imagePath = self.mayaToolsDir + "/General/Icons/ART/"
+	    
+	    cmds.symbolButton(self.widgets["previousStep"], edit = True, visible = True, image = imagePath + "skelCreation.bmp", c = self.edit)
+	    cmds.symbolButton(self.widgets["nextStep"], edit = True, visible = True, image = imagePath + "defSetup.bmp", c = self.lock_phase1)
+	    
+	    
+	    #now delete skeleton model pose
+	    if cmds.objExists("Skeleton_Model_Pose"):
+		cmds.delete("Skeleton_Model_Pose")
+	    
+	    
+	    #change mode buttons to default
+	    self.jointMover_ChangeMoveMode("Geo", "off")
+	    self.jointMover_ChangeMoveMode("Offset", "off")
+	    self.jointMover_ChangeMoveMode("Global", "off")
+	    self.jointMover_ChangeMoveMode("Global", "on")
+		
+	    cmds.iconTextRadioButton(self.widgets["moverMode"], edit = True, sl = True)
+	    cmds.iconTextRadioButton(self.widgets["offsetMode"], edit = True, sl = False)
+	    cmds.iconTextRadioButton(self.widgets["geoMoverMode"], edit = True, sl = False)
+	    
+	    
+	    
+	    #refresh UI
+	    import ART_skeletonBuilder_UI
+	    reload(ART_skeletonBuilder_UI)
+	    UI = ART_skeletonBuilder_UI.SkeletonBuilder_UI()
+	    print "UI Refreshed"
+	
+    # # # # # # # # # # # # # # # # # # # # # # # # # # # # # # # # # # # # # # # # # # # # # # # # # # # # # # # # # # # # # # # 
+    # # # # # # # # # # # # # # # # # # # # # # # # # # # # # # # # # # # # # # # # # # # # # # # # # # # # # # # # # # # # # # # 
+    # # # # # # # # # # # # # # # # # # # # # # # # # # # # # # # # # # # # # # # # # # # # # # # # # # # # # # # # # # # # # # # 
+    # # # # # # # # # # # # # # # # # # # # # # # # # # # # # # # # # # # # # # # # # # # # # # # # # # # # # # # # # # # # # # # 
+    # # # # # # # # # # # # # # # # # # # # # # # # # # # # # # # # # # # # # # # # # # # # # # # # # # # # # # # # # # # # # # # 
+    def publish(self, project, characterName, handCtrlSpace, *args):
+	
+	sourceControl = False
+	
+	#unlock joints
+	cmds.select("root", hi = True)
+	joints = cmds.ls(sl = True)
+	
+	for joint in joints:
+	    cmds.lockNode(joint, lock = False)
+	
+	#clear any keys on the joints
+	cmds.select("root", hi = True)
+	cmds.cutKey()
+	
+	#add the icon attr to the SceneLocked node
+	if cmds.objExists("SceneLocked.iconPath"):
+	    cmds.setAttr("SceneLocked.iconPath", self.mayaToolsDir + "/General/Icons/ART/Thumbnails/" + project + "/" + characterName + ".bmp", type = 'string')
+	    
+	else:
+	    cmds.select("SceneLocked")
+	    cmds.lockNode("SceneLocked", lock = False)
+	    cmds.addAttr(ln = "iconPath", dt = 'string')
+	    cmds.setAttr("SceneLocked.iconPath", self.mayaToolsDir + "/General/Icons/ART/Thumbnails/" + project + "/" + characterName + ".bmp", type = 'string')
+	    cmds.lockNode("SceneLocked", lock = True)
+	    
+	#set the model and joint mover back to rig pose
+	cmds.select("root", hi = True)
+	joints = cmds.ls(sl = True)
+	
+	for joint in joints:
+	    cmds.setAttr(joint + ".rx", 0)
+	    cmds.setAttr(joint + ".ry", 0)
+	    cmds.setAttr(joint + ".rz", 0)
+	    
+	self.setRigPose_JM()
+	cmds.select("root")
+	cmds.setToolTo( 'moveSuperContext' )
+	cmds.refresh(force = True)
+	cmds.select(clear = True)
+    
+    
+	#check if a pre script is loaded, and if so, save path in scene
+	if not cmds.objExists("SkeletonSettings_Cache.preScriptPath"):
+	    cmds.addAttr("SkeletonSettings_Cache", ln = "preScriptPath", dt = 'string')
+	    
+	scriptPath = cmds.textField(self.widgets["publishUIPreScriptField"], q = True, text = True)
+	
+	if scriptPath.find(".py") != -1 or scriptPath.find(".mel") != -1:
+	    cmds.setAttr("SkeletonSettings_Cache.preScriptPath", scriptPath, type = 'string')
+	    cmds.setAttr("SkeletonSettings_Cache.preScriptPath", keyable = False)
+	    
+	    
+	    
+	#check if a post script is loaded, and if so, save path in scene
+	if not cmds.objExists("SkeletonSettings_Cache.postScriptPath"):
+	    cmds.addAttr("SkeletonSettings_Cache", ln = "postScriptPath", dt = 'string')
+	    
+	scriptPath = cmds.textField(self.widgets["publishUIPostScriptField"], q = True, text = True)
+	
+	if scriptPath.find(".py") != -1 or scriptPath.find(".mel") != -1:
+	    cmds.setAttr("SkeletonSettings_Cache.postScriptPath", scriptPath, type = 'string')
+	    cmds.setAttr("SkeletonSettings_Cache.postScriptPath", keyable = False)
+	    
+	
+	             
+	#save out "export" file
+	exportPath = self.mayaToolsDir + "/General/ART/Projects/" + project + "/ExportFiles/"
+	if not os.path.exists(exportPath):
+	    os.makedirs(exportPath)
+	    
+	
+	#check if source control is on
+	settingsLocation = self.mayaToolsDir + "/General/Scripts/projectSettings.txt"
+	if os.path.exists(settingsLocation):
+	    f = open(settingsLocation, 'r')
+	    settings = cPickle.load(f)
+	    f.close()	    
+	    sourceControl = settings.get("UseSourceControl")	
+	    
+	
+	#save the export file out
+	cmds.file(rename = exportPath + characterName + "_Export.mb")
+	
+	#try to save the file
+	try:
+	    cmds.file(save = True, type = "mayaBinary", force = True, prompt = True)
+	    
+	except Exception, e:
+	    if sourceControl == False:
+		cmds.confirmDialog(title = "Publish", icon = "critical", message = str(e))
+		return
+	    
+	    else:
+		#if using source control, check to see if we can check out the file
+		result = cmds.confirmDialog(title = "Publish", icon = "critical", message = "Could not save Export file. File may exist already and be marked as read only.", button = ["Check Out File", "Cancel"])
+				
+		if result == "Check Out File":
+		    import perforceUtils
+		    reload(perforceUtils)
+		    writeable = perforceUtils.p4_checkOutCurrentFile(exportPath + characterName + "_Export.mb")
+		    
+		    if writeable:
+			#now that it is checked out, try saving again
+			try:
+			    cmds.file(save = True, type = "mayaBinary", force = True, prompt = True)
+			    
+			except:
+			    cmds.confirmDialog(title = "Publish", icon = "critical", message = "Perforce operation unsucessful. Could not save file. Aborting operation.")
+			    return
+			
+		    else:
+			cmds.warning("Operation Aborted")
+			return
+		    
+		else:
+		    cmds.warning("Operation Aborted.")
+		    return
+		
+		    
+		    
+	#Execute Pre Build Script if present
+	script = cmds.textField(self.widgets["publishUIPreScriptField"], q = True, text = True)
+	sourceType = ""
+	preScriptStatus = None
+	
+	if script.find(".py") != -1:
+	    sourceType = "python"
+	    
+	if script.find(".mel") != -1:
+	    sourceType = "mel"
+
+	
+	if sourceType == "mel":
+	    try:
+		command = ""
+		#open the file, and for each line in the file, add it to our command string.
+		f = open(script, 'r')
+		lines = f.readlines()
+		for line in lines:
+		    command += line
+		    
+		import maya.mel as mel
+		mel.eval(command)
+		
+		
+		#try to save out the export file
+		try:
+		    cmds.file(save = True, type = "mayaBinary", force = True, prompt = True)
+		    preScriptStatus = True
+		
+		except:
+		    preScriptStatus = False
+		    
+	    except:
+		preScriptStatus = False
+    
+	    
+	if sourceType == "python":
+	    try:
+		execfile("" + script + "")
+		
+		#try to save out the export file
+		try:
+		    cmds.file(save = True, type = "mayaBinary", force = True, prompt = True)
+		    preScriptStatus = True
+		    
+		except:
+		    preScriptStatus = False
+		    
+	    except:
+		preScriptStatus = False
+		
+		
+		
+		
+	#create new file
+	cmds.file( force=True, new=True )
+		
+	#reference in export file with no namespace
+	cmds.file(exportPath + characterName + "_Export.mb", r = True, type = "mayaBinary", loadReferenceDepth = "all", mergeNamespacesOnClash = True, namespace = ":", options = "v=0")
+	
+
+	#clear selection and fit view
+	cmds.select(clear = True)
+	cmds.viewFit()
+	panels = cmds.getPanel(type = 'modelPanel')
+
+	#turn on smooth shading
+	for panel in panels:
+	    editor = cmds.modelPanel(panel, q = True, modelEditor = True)
+	    cmds.modelEditor(editor, edit = True, displayAppearance = "smoothShaded", displayTextures = True, textures = True )
+
+	self.setRigPose_JM()
+	cmds.select("root")
+	cmds.setToolTo( 'moveSuperContext' )
+	cmds.refresh(force = True)
+	cmds.select(clear = True)
+	
+	#Import Auto Rig Class to build rig on skeleton
+	import ART_autoRigger
+	reload(ART_autoRigger)
+	ART_autoRigger.AutoRigger(handCtrlSpace, self.widgets["publishUI_ProgressBar"])
+
+	#find all skeleton mesh geo and add to a layer and hide the layer
+	if cmds.objExists("skeleton_skin_mesh*"):
+	    cmds.select("skeleton_skin_mesh_*")
+	    skelGeo = cmds.ls(sl = True, type = "transform")
+	    cmds.select(clear = True)
+	    
+	    #add skelGeo to a display layer
+	    cmds.select(skelGeo)
+	    cmds.createDisplayLayer(name = "skeleton_geometry_layer", nr = True)
+	    cmds.setAttr("skeleton_geometry_layer.enabled", 1)
+	    cmds.setAttr("skeleton_geometry_layer.displayType", 2)
+	    cmds.setAttr("skeleton_geometry_layer.visibility", 0)
+	    cmds.select(clear = True)
+	
+	cmds.select(clear = True)
+	
+
+	#Save out anim rig file
+	rigPath = self.mayaToolsDir + "/General/ART/Projects/" + project + "/AnimRigs/" 
+	if not os.path.exists(rigPath):
+	    os.makedirs(rigPath)
+	    
+	cmds.file(rename = rigPath + characterName + ".mb")
+	
+	#try to save out the anim rig file
+	try:
+	    cmds.file(save = True, type = "mayaBinary", force = True, prompt = True)
+	    
+	
+	except:
+	    if sourceControl == False:
+		cmds.confirmDialog(title = "Publish", icon = "critical", message = "Could not save file: " + str(rigPath + characterName) + ".mb not a valid file.\n File may exist already and be marked as read only. Aborting operation.")
+		return
+	    
+	    else:
+		#check to see if the file is currently checked out
+		result = cmds.confirmDialog(title = "Publish", icon = "critical", message = "Could not save Animation Rig file. File may exist already and be marked as read only.", button = ["Check Out File", "Cancel"])
+		
+		if result == "Check Out File":
+		    import perforceUtils
+		    reload(perforceUtils)
+		    writeable = perforceUtils.p4_checkOutCurrentFile(rigPath + characterName + ".mb")
+		    
+		    
+		    if writeable:
+			#try to save the file again now that it is checked out
+			try:
+			    cmds.file(save = True, type = "mayaBinary", force = True)
+			    
+			except:
+			    cmds.confirmDialog(title = "Publish", icon = "critical", message = "Perforce operation unsucessful. Could not save file. Aborting operation.")		
+	
+		    else:
+			cmds.confirmDialog(title = "Publish", icon = "critical", message = "Perforce operation unsucessful. Could not save file. Aborting operation.")
+			return
+	
+		else:
+		    cmds.confirmDialog(title = "Publish", icon = "critical", message = "Perforce operation unsucessful. Could not save file. Aborting operation.")
+		    return
+		    
+		    
+	
+	#check to see if there was a post script to execute
+	script = cmds.textField(self.widgets["publishUIPostScriptField"], q = True, text = True)
+	sourceType = ""
+	postScriptStatus = None
+	
+	if script.find(".py") != -1:
+	    sourceType = "python"
+	    
+	if script.find(".mel") != -1:
+	    sourceType = "mel"
+
+	
+	if sourceType == "mel":
+	    try:
+		command = ""
+		#open the file, and for each line in the file, add it to our command string.
+		f = open(script, 'r')
+		lines = f.readlines()
+		for line in lines:
+		    command += line
+		    
+		import maya.mel as mel
+		mel.eval(command)
+		
+		
+		#try to save out the anim rig file
+		try:
+		    cmds.file(save = True, type = "mayaBinary", force = True, prompt = True)
+		    postScriptStatus = True
+		
+		except:
+		    postScriptStatus = False
+		    
+	    except:
+		postScriptStatus = False
+    
+	    
+	if sourceType == "python":
+	    try:
+		execfile("" + script + "")
+		
+		#try to save out the anim rig file
+		try:
+		    cmds.file(save = True, type = "mayaBinary", force = True, prompt = True)
+		    postScriptStatus = True
+		    
+		except:
+		    postScriptStatus = False
+		    
+	    except:
+		postScriptStatus = False
+
+	    
+	#delete publish UI
+	cmds.deleteUI(self.widgets["publishUIWindow"])
+	
+
+	
+	#show results UI
+	if sourceControl == False:
+	    self.publishUI_Results(False, preScriptStatus, postScriptStatus, self.mayaToolsDir + "/General/ART/Projects/" + project + "/ExportFiles/" + characterName + "_Export.mb", self.mayaToolsDir + "/General/ART/Projects/" + project + "/AnimRigs/" + characterName + ".mb", self.mayaToolsDir + "/General/Icons/ART/Thumbnails/" + project + "/" + characterName + ".bmp", self.mayaToolsDir + "/General/Icons/ART/Thumbnails/" + project + "/" + characterName + "_small.bmp")
+	
+	if sourceControl == True:
+	    self.publishUI_Results(True, preScriptStatus, postScriptStatus, self.mayaToolsDir + "/General/ART/Projects/" + project + "/ExportFiles/" + characterName + "_Export.mb", self.mayaToolsDir + "/General/ART/Projects/" + project + "/AnimRigs/" + characterName + ".mb", self.mayaToolsDir + "/General/Icons/ART/Thumbnails/" + project + "/" + characterName + ".bmp", self.mayaToolsDir + "/General/Icons/ART/Thumbnails/" + project + "/" + characterName + "_small.bmp")
+	    
+	
+    # # # # # # # # # # # # # # # # # # # # # # # # # # # # # # # # # # # # # # # # # # # # # # # # # # # # # # # # # # # # # # # 
+    # # # # # # # # # # # # # # # # # # # # # # # # # # # # # # # # # # # # # # # # # # # # # # # # # # # # # # # # # # # # # # # 
+    # # # # # # # # # # # # # # # # # # # # # # # # # # # # # # # # # # # # # # # # # # # # # # # # # # # # # # # # # # # # # # # 
+    # # # # # # # # # # # # # # # # # # # # # # # # # # # # # # # # # # # # # # # # # # # # # # # # # # # # # # # # # # # # # # # 
+    # # # # # # # # # # # # # # # # # # # # # # # # # # # # # # # # # # # # # # # # # # # # # # # # # # # # # # # # # # # # # # # 
+    def publishUI_Results(self, sourceControl, preScriptStatus, postScriptStatus, exportFile, rigFile, lrgIcon, smIcon, *args):
+	
+	#if window exists, delete
+	if cmds.window("publishResultsInterface", exists = True):
+	    cmds.deleteUI("publishResultsInterface")
+
+	#create window
+	self.widgets["publishResultsWindow"] = cmds.window("publishResultsInterface", title = "Character Published!", w = 480, h = 200, mnb = False, mxb = False, sizeable = False)
+	
+	#widgets
+	self.widgets["publishResults_mainLayout"] = cmds.columnLayout(rs = 10)
+	cmds.image(image = self.mayaToolsDir + "/General/Icons/ART/publishResults.bmp", w = 480, h = 50, parent = self.widgets["publishResults_mainLayout"])
+	
+	self.widgets["publishResults_topLayout"] = cmds.rowColumnLayout(nc = 3, rs = [1,5], cw = [(1, 30),(2,400),(3,30)], cat = [(1, "both", 5),(2, "both", 5),(3,"both",5)], parent = self.widgets["publishResults_mainLayout"])
+	cmds.text(label = "", parent = self.widgets["publishResults_topLayout"])
+	cmds.text(label = "Files Created:", font = "boldLabelFont", parent = self.widgets["publishResults_topLayout"], align = "center")
+	cmds.text(label = "", parent = self.widgets["publishResults_topLayout"])
+	
+	#for each file produced, create a checkbox, textField, and browse button
+	i = 1
+	for each in [exportFile, rigFile, lrgIcon, smIcon]:
+	    
+	    #checkbox
+	    checkbox = cmds.checkBox("publishResultsCB_" + str(i), parent = self.widgets["publishResults_topLayout"], v = True, label = "")
+	    
+	    #textField
+	    field = cmds.textField(w = 320, parent = self.widgets["publishResults_topLayout"], text = each, ed = False)
+	    
+	    #browse button
+	    button = cmds.symbolButton(w = 30, h = 30, image = self.mayaToolsDir + "/General/Icons/ART/find.bmp", parent = self.widgets["publishResults_topLayout"], c = partial(self.findFileOnDisk, each))
+	    
+	    i = i + 1
+	    
+	    
+	#detail post script results
+	if preScriptStatus:
+	    layout = cmds.columnLayout(w = 480, parent = self.widgets["publishResults_mainLayout"], cat = ["both", 100])
+	    text = cmds.text(w = 280, label = "Pre Script executed successfully!", align = "center", bgc = [0, 1, 0], parent = layout)
+	    
+	if preScriptStatus == False:
+	    layout = cmds.columnLayout(w = 480, parent = self.widgets["publishResults_mainLayout"], cat = ["both", 100])
+	    text = cmds.text(w = 280, label = "Pre Script failed to execute properly.", align = "center", bgc = [1, 0, 0], parent = layout)
+	    
+	if postScriptStatus:
+	    layout = cmds.columnLayout(w = 480, parent = self.widgets["publishResults_mainLayout"], cat = ["both", 100])
+	    text = cmds.text(w = 280, label = "Post Script executed successfully!", align = "center", bgc = [0, 1, 0], parent = layout)
+	    
+	if postScriptStatus == False:
+	    layout = cmds.columnLayout(w = 480, parent = self.widgets["publishResults_mainLayout"], cat = ["both", 100])
+	    text = cmds.text(w = 280, label = "Post Script failed to execute properly.", align = "center", bgc = [1, 0, 0], parent = layout)
+	    
+	if sourceControl == False:
+	    #create close button
+	    self.widgets["publishResults_botLayout"] = cmds.rowColumnLayout(w = 480, nc = 2, cat = [(1, "left", 20),(2, "both", 20)], parent = self.widgets["publishResults_mainLayout"])
+	    cmds.button(w = 210, h = 50, label = "Edit Rig File", parent = self.widgets["publishResults_botLayout"], c = partial(self.publishUI_Results_EditRig, rigFile))
+	    cmds.button(w = 210, h = 50, label = "Close", parent = self.widgets["publishResults_botLayout"], c = self.publishUI_Results_Close)
+	    
+	    
+	if sourceControl == True:
+	    self.widgets["publishResults_botLayout"] = cmds.rowColumnLayout(w = 480, nc = 3, cat = [(1, "left", 20),(2, "both", 20), (3, "right", 20)], parent = self.widgets["publishResults_mainLayout"])
+	    cmds.button(w = 133, h = 50, label = "Edit Rig File", parent = self.widgets["publishResults_botLayout"], c = partial(self.publishUI_Results_EditRig, rigFile))
+	    cmds.button(w = 133, h = 50, label = "Submit Selected Files", parent = self.widgets["publishResults_botLayout"], c = partial(self.publishUI_Results_Submit, exportFile, rigFile, lrgIcon, smIcon))
+	    cmds.button(w = 133, h = 50, label = "Close", parent = self.widgets["publishResults_botLayout"], c = self.publishUI_Results_Close)
+	    
+	    
+	cmds.separator(h = 5, parent = self.widgets["publishResults_mainLayout"])
+	
+	#show window
+	cmds.showWindow(self.widgets["publishResultsWindow"])
+	
+	
+    # # # # # # # # # # # # # # # # # # # # # # # # # # # # # # # # # # # # # # # # # # # # # # # # # # # # # # # # # # # # # # # 
+    # # # # # # # # # # # # # # # # # # # # # # # # # # # # # # # # # # # # # # # # # # # # # # # # # # # # # # # # # # # # # # # 
+    # # # # # # # # # # # # # # # # # # # # # # # # # # # # # # # # # # # # # # # # # # # # # # # # # # # # # # # # # # # # # # # 
+    # # # # # # # # # # # # # # # # # # # # # # # # # # # # # # # # # # # # # # # # # # # # # # # # # # # # # # # # # # # # # # # 
+    # # # # # # # # # # # # # # # # # # # # # # # # # # # # # # # # # # # # # # # # # # # # # # # # # # # # # # # # # # # # # # # 
+    def publishUI_Results_Close(self, *args):
+	
+	cmds.deleteUI("publishResultsInterface")
+	
+	
+    # # # # # # # # # # # # # # # # # # # # # # # # # # # # # # # # # # # # # # # # # # # # # # # # # # # # # # # # # # # # # # # 
+    # # # # # # # # # # # # # # # # # # # # # # # # # # # # # # # # # # # # # # # # # # # # # # # # # # # # # # # # # # # # # # # 
+    # # # # # # # # # # # # # # # # # # # # # # # # # # # # # # # # # # # # # # # # # # # # # # # # # # # # # # # # # # # # # # # 
+    # # # # # # # # # # # # # # # # # # # # # # # # # # # # # # # # # # # # # # # # # # # # # # # # # # # # # # # # # # # # # # # 
+    # # # # # # # # # # # # # # # # # # # # # # # # # # # # # # # # # # # # # # # # # # # # # # # # # # # # # # # # # # # # # # # 
+    def publishUI_Results_EditRig(self, rigFile, *args):
+	
+	cmds.file(rigFile, open = True, force = True)
+	cmds.deleteUI("publishResultsInterface")
+	
+
+    # # # # # # # # # # # # # # # # # # # # # # # # # # # # # # # # # # # # # # # # # # # # # # # # # # # # # # # # # # # # # # # 
+    # # # # # # # # # # # # # # # # # # # # # # # # # # # # # # # # # # # # # # # # # # # # # # # # # # # # # # # # # # # # # # # 
+    # # # # # # # # # # # # # # # # # # # # # # # # # # # # # # # # # # # # # # # # # # # # # # # # # # # # # # # # # # # # # # # 
+    # # # # # # # # # # # # # # # # # # # # # # # # # # # # # # # # # # # # # # # # # # # # # # # # # # # # # # # # # # # # # # # 
+    # # # # # # # # # # # # # # # # # # # # # # # # # # # # # # # # # # # # # # # # # # # # # # # # # # # # # # # # # # # # # # # 
+    def publishUI_Results_Submit(self, exportFile, rigFile, lrgIcon, smIcon, *args):
+	
+	submittedFiles = []
+			
+	#import perforce utils
+	import perforceUtils
+	reload(perforceUtils)
+	
+	#get a description for the changelist
+	cmds.promptDialog(title = "Perforce", message = "Please Enter a Description..", button = ["Submit"])
+	desc = cmds.promptDialog(q = True, text = True)
+	
+	
+	
+	#submit the export file
+	if cmds.checkBox("publishResultsCB_1", q = True, v = True):
+	    submitted = perforceUtils.p4_submitCurrentFile(exportFile, desc)
+	    
+	    if submitted == False:		    
+		added = perforceUtils.p4_addAndSubmitCurrentFile(exportFile, desc)
+		if added == True:
+		    submittedFiles.append(exportFile)
+		
+	    if submitted == True:
+		submittedFiles.append(exportFile)
+		
+	    
+	    
+	#submit the anim rig
+	if cmds.checkBox("publishResultsCB_2", q = True, v = True):
+	    submitted = perforceUtils.p4_submitCurrentFile(rigFile, desc)
+		
+	    if submitted == False:		    
+		added = perforceUtils.p4_addAndSubmitCurrentFile(rigFile, desc)
+		if added == True:
+		    submittedFiles.append(rigFile)
+		
+	    if submitted == True:
+		submittedFiles.append(rigFile)		
+	
+	
+	try:
+	    #submit the icons
+	    if cmds.checkBox("publishResultsCB_3", q = True, v = True):
+		submitted = perforceUtils.p4_submitCurrentFile(lrgIcon, desc)
+		
+		if submitted == False:		    
+		    added = perforceUtils.p4_addAndSubmitCurrentFile(lrgIcon, desc)
+		    if added == True:
+			submittedFiles.append(lrgIcon)
+		    
+		if submitted == True:
+		    submittedFiles.append(lrgIcon)
+		
+		
+		
+	    #small icon		
+	    if cmds.checkBox("publishResultsCB_4", q = True, v = True):
+		submitted = perforceUtils.p4_submitCurrentFile(smIcon, desc)
+		
+		if submitted == False:		    
+		    added = perforceUtils.p4_addAndSubmitCurrentFile(smIcon, desc)
+		    if added == True:
+			submittedFiles.append(smIcon)
+		    
+		if submitted == True:
+		    submittedFiles.append(smIcon)
+		
+	except:
+	    pass	
+	
+	
+	if len(submittedFiles) > 0:
+	    message = ""
+	    for each in submittedFiles:
+		message += each + "\n"
+		
+	    cmds.confirmDialog(title = "Perforce", icon = "information", message = "Submit Operation was successful! Submitted Files:\n\n" + message, button = "Close")
+	    
+	    
+    # # # # # # # # # # # # # # # # # # # # # # # # # # # # # # # # # # # # # # # # # # # # # # # # # # # # # # # # # # # # # # # 
+    # # # # # # # # # # # # # # # # # # # # # # # # # # # # # # # # # # # # # # # # # # # # # # # # # # # # # # # # # # # # # # # 
+    # # # # # # # # # # # # # # # # # # # # # # # # # # # # # # # # # # # # # # # # # # # # # # # # # # # # # # # # # # # # # # # 
+    # # # # # # # # # # # # # # # # # # # # # # # # # # # # # # # # # # # # # # # # # # # # # # # # # # # # # # # # # # # # # # # 
+    # # # # # # # # # # # # # # # # # # # # # # # # # # # # # # # # # # # # # # # # # # # # # # # # # # # # # # # # # # # # # # # 
+    def publish_UI(self, *args):
+	
+	
+	
+	#check to make sure rig pose and model pose nodes exist
+	if cmds.objExists("Model_Pose") == False:
+	    cmds.confirmDialog(icon = "warning", title = "Error", message = "Model Pose not found for joint mover. Going back to Edit mode to create model pose.")
+	    self.unlock()
+	    return
+	    
+	if cmds.objExists("Rig_Pose") == False:
+	    cmds.confirmDialog(icon = "warning", title = "Error", message = "Rig Pose not found for joint mover. Going back to Edit mode to create rig pose.")
+	    self.unlock()
+	    return
+	
+	
+	#check worldspace position of feet
+	footPos = cmds.xform("foot_l", q = True, ws = True, t = True)
+	if footPos[2] < 0:
+	    result = cmds.confirmDialog(title = "Publish", icon = "warning",  message = "Character's feet are below ground plane. It is advised to have character feet on ground plane.", button = ["Ignore", "Adjust Character"], defaultButton = "Adjust Character", cancelButton = "Ignore", dismissString = "Ignore")
+	    if result == "Adjust Character":
+		return
+	    
+	#check root bone position
+	rootPos = cmds.xform("root", q = True, ws = True, t = True)
+	if rootPos != [0.0, 0.0, 0.0]:
+	    result = cmds.confirmDialog(title = "Publish", icon = "warning",  message = "Character's root bone is not at the origin. It is recommended that the character's root bone be at the origin of the scene.", button = ["Ignore", "Adjust Character"], defaultButton = "Adjust Character", cancelButton = "Ignore", dismissString = "Ignore")
+	    if result == "Adjust Character":
+		return	    
+	    
+	    
+	#change flow chart buttons:
+	imagePath = self.mayaToolsDir + "/General/Icons/ART/"
+
+	
+	if cmds.window("publish_ui_window", exists = True):
+	    cmds.deleteUI("publish_ui_window")
+	    
+	self.widgets["publishUIWindow"] = cmds.window("publish_ui_window", w = 300, h = 310, mxb = False, mnb = False, sizeable = False, title = "Publish Character")
+	self.widgets["publishUIMainLayout"] = cmds.columnLayout(w = 300, h =310, co = ["both", 5])
+	
+	#banner image
+	cmds.image(w = 300, h = 50, image = self.mayaToolsDir + "/General/Icons/ART/artBanner300px.bmp", parent = self.widgets["publishUIMainLayout"])
+	
+	self.widgets["publishUIForm"] = cmds.formLayout(w = 300, h = 310, parent = self.widgets["publishUIMainLayout"])
+	
+	
+	#create the objects for the UI
+	self.widgets["publishUIProjectLabel"] = cmds.text(label = "Select Project:", h = 30)
+	self.widgets["publishUI_ProjOM"] = cmds.optionMenu(w = 160, h = 30, cc = self.getCharactersForProject)
+	self.widgets["publishUI_NewProjBtn"] = cmds.symbolButton(w = 30, h = 30, image = imagePath + "new.bmp", ann = "Add New Project", c = self.publish_UI_addNewProj)
+	
+	self.widgets["publishUISearchLabel"] = cmds.text(label = "Search:", h = 20)
+	self.widgets["publishUI_search"] = cmds.textField(w = 190, h = 20)
+	
+	
+	self.widgets["publishUI_textScrollList"] = cmds.textScrollList(w = 275, h = 150, ams = False, selectCommand = self.updateCharacterNameField, ann = "List of existing characters in the selected project.")
+	self.widgets["publishUICharNameLabel"] = cmds.text(label = "Character Name:", h = 30, font = "boldLabelFont")
+	self.widgets["publishUICharNameTF"] = cmds.textField(w = 275, h = 30, ann = "Name the character rig will be built with.")
+	
+	cmds.separator(h = 10, parent = self.widgets["publishUIMainLayout"])
+	self.widgets["publishUIButtonRow"] = cmds.rowColumnLayout(nc = 2, w = 300, parent = self.widgets["publishUIMainLayout"], co = [[1, "both", 5], [2, "both", 5]])
+	self.widgets["publishUIPublishBtn"] = cmds.button(w = 140, h = 50, label = "Build", c = self.publish_UI_TakeScreenshot, parent = self.widgets["publishUIButtonRow"], ann = "Build the character rig with current settings.")
+	self.widgets["publishUICancelBtn"] = cmds.button(w = 140, h = 50, label = "Cancel", c = self.publish_UI_Cancel, parent = self.widgets["publishUIButtonRow"])
+	
+	self.widgets["publishUIAdvancedFrame"] = cmds.frameLayout(w = 280, parent = self.widgets["publishUIForm"], cll = True, cl = True, label = "Advanced", bs = "etchedIn", h = 250, cc = self.publishUICollapseFrame, ec = self.publishUIExpandFrame)
+	self.widgets["publishUIAdvancedLayout"] = cmds.columnLayout(w = 280, parent = self.widgets["publishUIAdvancedFrame"], co = ["both", 10], rs = 10)
+	cmds.text(label = "", parent = self.widgets["publishUIAdvancedLayout"])
+	
+	#Pre Build Script
+	self.widgets["publishUIPreScriptBtn"] = cmds.button(w = 260, h = 30, label = "Add Pre-Build Script", parent = self.widgets["publishUIAdvancedLayout"], c = self.publishUI_AddPreScript, ann = "Add a script to execute before the rig has been built. Valid script types are MEL and Python scripts.")
+	cmds.text(label = "Script To Execute:")
+	self.widgets["publishUIPreScriptField"] = cmds.textField(w = 260, h = 30, parent = self.widgets["publishUIAdvancedLayout"])
+	
+	#Post Build Script
+	self.widgets["publishUIPostScriptBtn"] = cmds.button(w = 260, h = 30, label = "Add Post-Build Script", parent = self.widgets["publishUIAdvancedLayout"], c = self.publishUI_AddPostScript, ann = "Add a script to execute after the rig has been built. Valid script types are MEL and Python scripts.")
+	cmds.text(label = "Script To Execute:")
+	self.widgets["publishUIPostScriptField"] = cmds.textField(w = 260, h = 30, parent = self.widgets["publishUIAdvancedLayout"])
+	
+	#add change command to the search textField
+	cmds.textField(self.widgets["publishUI_search"], edit = True, cc = partial(self.publish_searchList, self.widgets["publishUI_search"], self.widgets["publishUI_textScrollList"]))
+	
+	#attach the objects to the UI
+	cmds.formLayout(self.widgets["publishUIForm"], edit = True, af = [(self.widgets["publishUIProjectLabel"], "left", 10), (self.widgets["publishUIProjectLabel"], "top", 10)])
+	cmds.formLayout(self.widgets["publishUIForm"], edit = True, af = [(self.widgets["publishUI_ProjOM"], "left", 100), (self.widgets["publishUI_ProjOM"], "top", 10)])
+	cmds.formLayout(self.widgets["publishUIForm"], edit = True, af = [(self.widgets["publishUI_NewProjBtn"], "right", 10), (self.widgets["publishUI_NewProjBtn"], "top", 10)])
+	
+	cmds.formLayout(self.widgets["publishUIForm"], edit = True, af = [(self.widgets["publishUISearchLabel"], "left", 10), (self.widgets["publishUISearchLabel"], "top", 45)])
+	cmds.formLayout(self.widgets["publishUIForm"], edit = True, af = [(self.widgets["publishUI_search"], "right", 10), (self.widgets["publishUI_search"], "top", 45)])
+	
+	cmds.formLayout(self.widgets["publishUIForm"], edit = True, af = [(self.widgets["publishUI_textScrollList"], "left", 10), (self.widgets["publishUI_textScrollList"], "top", 80)])
+	
+	cmds.formLayout(self.widgets["publishUIForm"], edit = True, af = [(self.widgets["publishUICharNameLabel"], "left", 10), (self.widgets["publishUICharNameLabel"], "top", 230)])
+	cmds.formLayout(self.widgets["publishUIForm"], edit = True, af = [(self.widgets["publishUICharNameTF"], "right", 10), (self.widgets["publishUICharNameTF"], "top", 255)])
+	
+	cmds.formLayout(self.widgets["publishUIForm"], edit = True, af = [(self.widgets["publishUIAdvancedFrame"], "left", 10), (self.widgets["publishUIAdvancedFrame"], "top", 290)])
+	
+	#cmds.formLayout(self.widgets["publishUIForm"], edit = True, af = [(self.widgets["publishUIButtonRow"], "left", 10), (self.widgets["publishUIButtonRow"], "bottom", 10)])
+	
+
+	#add progress bar
+	cmds.separator(h = 10, parent = self.widgets["publishUIMainLayout"])
+	self.widgets["publishUI_ProgressBar"] = cmds.progressBar(parent = self.widgets["publishUIMainLayout"], w = 300, h = 20, progress = 0)
+	cmds.separator(h = 10, parent = self.widgets["publishUIMainLayout"])
+
+	cmds.showWindow(self.widgets["publishUIWindow"])
+	
+	#populate the projects
+	self.publish_UI_GetProjects()
+	
+	#populate the existing character names into the textScrollList for the project
+	self.getCharactersForProject()
+	
+	
+	
+	#check if pre script path exists on the SkeletonSettings_Cache, and if so, populate pre script textField
+	try:
+	    if cmds.objExists("SkeletonSettings_Cache.preScriptPath"):
+		scriptPath = cmds.getAttr("SkeletonSettings_Cache.preScriptPath")
+		
+		if os.path.exists(scriptPath):
+		    cmds.textField(self.widgets["publishUIPreScriptField"] , edit = True, text = scriptPath)
+		    
+		else:
+		    cmds.confirmDialog(title = "Pre Build Script", icon = "warning", message = "This file has a pre-build script associated with it, but the file could not be found on your drive.")
+
+	except:
+	    pass
+	
+	
+	
+	#check if post script path exists on the SkeletonSettings_Cache, and if so, populate post script textField
+	try:
+	    if cmds.objExists("SkeletonSettings_Cache.postScriptPath"):
+		scriptPath = cmds.getAttr("SkeletonSettings_Cache.postScriptPath")
+		
+		if os.path.exists(scriptPath):
+		    cmds.textField(self.widgets["publishUIPostScriptField"] , edit = True, text = scriptPath)
+		    
+		else:
+		    cmds.confirmDialog(title = "Post Build Script", icon = "warning", message = "This file has a post-build script associated with it, but the file could not be found on your drive.")
+
+	except:
+	    pass
+
+		
+	
+	
+	
+	
+    # # # # # # # # # # # # # # # # # # # # # # # # # # # # # # # # # # # # # # # # # # # # # # # # # # # # # # # # # # # # # # # 
+    # # # # # # # # # # # # # # # # # # # # # # # # # # # # # # # # # # # # # # # # # # # # # # # # # # # # # # # # # # # # # # # 
+    # # # # # # # # # # # # # # # # # # # # # # # # # # # # # # # # # # # # # # # # # # # # # # # # # # # # # # # # # # # # # # # 
+    # # # # # # # # # # # # # # # # # # # # # # # # # # # # # # # # # # # # # # # # # # # # # # # # # # # # # # # # # # # # # # # 
+    # # # # # # # # # # # # # # # # # # # # # # # # # # # # # # # # # # # # # # # # # # # # # # # # # # # # # # # # # # # # # # # 
+    def publish_searchList(self, textField, textScrollList, *args):
+	
+	#populate the existing character names into the textScrollList for the project
+	self.getCharactersForProject()
+	
+	#filter results
+	self.searchForKeyword(textField, textScrollList)
+	
+	
+	
+	
+    # # # # # # # # # # # # # # # # # # # # # # # # # # # # # # # # # # # # # # # # # # # # # # # # # # # # # # # # # # # # # # # 
+    # # # # # # # # # # # # # # # # # # # # # # # # # # # # # # # # # # # # # # # # # # # # # # # # # # # # # # # # # # # # # # # 
+    # # # # # # # # # # # # # # # # # # # # # # # # # # # # # # # # # # # # # # # # # # # # # # # # # # # # # # # # # # # # # # # 
+    # # # # # # # # # # # # # # # # # # # # # # # # # # # # # # # # # # # # # # # # # # # # # # # # # # # # # # # # # # # # # # # 
+    # # # # # # # # # # # # # # # # # # # # # # # # # # # # # # # # # # # # # # # # # # # # # # # # # # # # # # # # # # # # # # # 
+    def searchForKeyword(self, textField, textScrollList, *args):
+	
+	keyword = cmds.textField(textField, q = True, text = True)
+	
+	items = cmds.textScrollList(textScrollList, q = True, allItems = True)
+	
+	matches = []
+	
+	for item in items:
+	    if item.find(keyword) != -1:
+		matches.append(item)
+		
+	cmds.textScrollList(textScrollList, edit = True, ra = True)
+	
+	for match in matches:
+	    cmds.textScrollList(textScrollList, edit = True, append = match)
+	
+    # # # # # # # # # # # # # # # # # # # # # # # # # # # # # # # # # # # # # # # # # # # # # # # # # # # # # # # # # # # # # # # 
+    # # # # # # # # # # # # # # # # # # # # # # # # # # # # # # # # # # # # # # # # # # # # # # # # # # # # # # # # # # # # # # # 
+    # # # # # # # # # # # # # # # # # # # # # # # # # # # # # # # # # # # # # # # # # # # # # # # # # # # # # # # # # # # # # # # 
+    # # # # # # # # # # # # # # # # # # # # # # # # # # # # # # # # # # # # # # # # # # # # # # # # # # # # # # # # # # # # # # # 
+    # # # # # # # # # # # # # # # # # # # # # # # # # # # # # # # # # # # # # # # # # # # # # # # # # # # # # # # # # # # # # # # 
+    def publishUI_AddPreScript(self, *args):
+	
+	#launch the file dialog window
+	fileName = cmds.fileDialog2(startingDirectory = cmds.internalVar(usd = True), ff = "*.py;;*.mel", fm = 1, okCaption = "Load Script")[0]
+	
+	#edit the textField to have the path to the script
+	cmds.textField(self.widgets["publishUIPreScriptField"], edit = True, text = fileName)
+	
+    # # # # # # # # # # # # # # # # # # # # # # # # # # # # # # # # # # # # # # # # # # # # # # # # # # # # # # # # # # # # # # # 
+    # # # # # # # # # # # # # # # # # # # # # # # # # # # # # # # # # # # # # # # # # # # # # # # # # # # # # # # # # # # # # # # 
+    # # # # # # # # # # # # # # # # # # # # # # # # # # # # # # # # # # # # # # # # # # # # # # # # # # # # # # # # # # # # # # # 
+    # # # # # # # # # # # # # # # # # # # # # # # # # # # # # # # # # # # # # # # # # # # # # # # # # # # # # # # # # # # # # # # 
+    # # # # # # # # # # # # # # # # # # # # # # # # # # # # # # # # # # # # # # # # # # # # # # # # # # # # # # # # # # # # # # # 
+    def publishUI_AddPostScript(self, *args):
+	
+	#launch the file dialog window
+	fileName = cmds.fileDialog2(startingDirectory = cmds.internalVar(usd = True), ff = "*.py;;*.mel", fm = 1, okCaption = "Load Script")[0]
+	
+	#edit the textField to have the path to the script
+	cmds.textField(self.widgets["publishUIPostScriptField"], edit = True, text = fileName)
+	
+	
+    # # # # # # # # # # # # # # # # # # # # # # # # # # # # # # # # # # # # # # # # # # # # # # # # # # # # # # # # # # # # # # # 
+    # # # # # # # # # # # # # # # # # # # # # # # # # # # # # # # # # # # # # # # # # # # # # # # # # # # # # # # # # # # # # # # 
+    # # # # # # # # # # # # # # # # # # # # # # # # # # # # # # # # # # # # # # # # # # # # # # # # # # # # # # # # # # # # # # # 
+    # # # # # # # # # # # # # # # # # # # # # # # # # # # # # # # # # # # # # # # # # # # # # # # # # # # # # # # # # # # # # # # 
+    # # # # # # # # # # # # # # # # # # # # # # # # # # # # # # # # # # # # # # # # # # # # # # # # # # # # # # # # # # # # # # # 
+    def publishUICollapseFrame(self, *args):
+	cmds.window(self.widgets["publishUIWindow"], edit = True, h =310)
+	cmds.columnLayout(self.widgets["publishUIMainLayout"], edit = True, h =310)
+	cmds.formLayout(self.widgets["publishUIForm"], edit = True, h =310)
+	
+    # # # # # # # # # # # # # # # # # # # # # # # # # # # # # # # # # # # # # # # # # # # # # # # # # # # # # # # # # # # # # # # 
+    # # # # # # # # # # # # # # # # # # # # # # # # # # # # # # # # # # # # # # # # # # # # # # # # # # # # # # # # # # # # # # # 
+    # # # # # # # # # # # # # # # # # # # # # # # # # # # # # # # # # # # # # # # # # # # # # # # # # # # # # # # # # # # # # # # 
+    # # # # # # # # # # # # # # # # # # # # # # # # # # # # # # # # # # # # # # # # # # # # # # # # # # # # # # # # # # # # # # # 
+    # # # # # # # # # # # # # # # # # # # # # # # # # # # # # # # # # # # # # # # # # # # # # # # # # # # # # # # # # # # # # # # 
+    def publishUIExpandFrame(self, *args):
+	cmds.window(self.widgets["publishUIWindow"], edit = True, h =550)
+	cmds.columnLayout(self.widgets["publishUIMainLayout"], edit = True, h =550)
+	cmds.formLayout(self.widgets["publishUIForm"], edit = True, h =550)	
+	
+    # # # # # # # # # # # # # # # # # # # # # # # # # # # # # # # # # # # # # # # # # # # # # # # # # # # # # # # # # # # # # # # 
+    # # # # # # # # # # # # # # # # # # # # # # # # # # # # # # # # # # # # # # # # # # # # # # # # # # # # # # # # # # # # # # # 
+    # # # # # # # # # # # # # # # # # # # # # # # # # # # # # # # # # # # # # # # # # # # # # # # # # # # # # # # # # # # # # # # 
+    # # # # # # # # # # # # # # # # # # # # # # # # # # # # # # # # # # # # # # # # # # # # # # # # # # # # # # # # # # # # # # # 
+    # # # # # # # # # # # # # # # # # # # # # # # # # # # # # # # # # # # # # # # # # # # # # # # # # # # # # # # # # # # # # # # 
+    def updateCharacterNameField(self, *args):
+	
+	selected = cmds.textScrollList(self.widgets["publishUI_textScrollList"], q = True, si = True)[0]
+	cmds.textField(self.widgets["publishUICharNameTF"], edit = True, text = selected)
+	
+	
+	
+    # # # # # # # # # # # # # # # # # # # # # # # # # # # # # # # # # # # # # # # # # # # # # # # # # # # # # # # # # # # # # # # 
+    # # # # # # # # # # # # # # # # # # # # # # # # # # # # # # # # # # # # # # # # # # # # # # # # # # # # # # # # # # # # # # # 
+    # # # # # # # # # # # # # # # # # # # # # # # # # # # # # # # # # # # # # # # # # # # # # # # # # # # # # # # # # # # # # # # 
+    # # # # # # # # # # # # # # # # # # # # # # # # # # # # # # # # # # # # # # # # # # # # # # # # # # # # # # # # # # # # # # # 
+    # # # # # # # # # # # # # # # # # # # # # # # # # # # # # # # # # # # # # # # # # # # # # # # # # # # # # # # # # # # # # # # 
+    def getCharactersForProject(self, *args):
+	
+	#remove all entries from the list
+	try:
+	    cmds.textScrollList(self.widgets["publishUI_textScrollList"], edit = True, removeAll = True)
+	
+	except:
+	    pass
+	
+	
+	project = cmds.optionMenu(self.widgets["publishUI_ProjOM"], q = True, value = True)
+	projectDir = self.mayaToolsDir + "/General/ART/Projects/" + project + "/AnimRigs/"
+	
+	try:
+	    characters = os.listdir(projectDir)
+	    
+	    for character in characters:
+		niceName = character.partition(".")[0]
+		cmds.textScrollList(self.widgets["publishUI_textScrollList"], edit = True, append = niceName)
+		
+	    
+		
+	except:
+	    pass
+	
+	    
+	
+	
+    # # # # # # # # # # # # # # # # # # # # # # # # # # # # # # # # # # # # # # # # # # # # # # # # # # # # # # # # # # # # # # # 
+    # # # # # # # # # # # # # # # # # # # # # # # # # # # # # # # # # # # # # # # # # # # # # # # # # # # # # # # # # # # # # # # 
+    # # # # # # # # # # # # # # # # # # # # # # # # # # # # # # # # # # # # # # # # # # # # # # # # # # # # # # # # # # # # # # # 
+    # # # # # # # # # # # # # # # # # # # # # # # # # # # # # # # # # # # # # # # # # # # # # # # # # # # # # # # # # # # # # # # 
+    # # # # # # # # # # # # # # # # # # # # # # # # # # # # # # # # # # # # # # # # # # # # # # # # # # # # # # # # # # # # # # # 
+    def publish_UI_GetProjects(self):
+	
+	children = cmds.optionMenu(self.widgets["publishUI_ProjOM"], q = True, itemListLong = True)
+	if children != None:
+	    for child in children:
+		cmds.deleteUI(child)
+	    
+	
+	projectPath = self.mayaToolsDir + "/General/ART/Projects/"
+	
+	if not os.path.exists(projectPath):
+	    os.makedirs(projectPath)
+	    
+	projects = os.listdir(projectPath)
+	
+	if len(projects) > 0:
+	
+	    for proj in sorted(projects):
+		cmds.menuItem(label = proj, parent = self.widgets["publishUI_ProjOM"])
+	    
+	    
+    # # # # # # # # # # # # # # # # # # # # # # # # # # # # # # # # # # # # # # # # # # # # # # # # # # # # # # # # # # # # # # # 
+    # # # # # # # # # # # # # # # # # # # # # # # # # # # # # # # # # # # # # # # # # # # # # # # # # # # # # # # # # # # # # # # 
+    # # # # # # # # # # # # # # # # # # # # # # # # # # # # # # # # # # # # # # # # # # # # # # # # # # # # # # # # # # # # # # # 
+    # # # # # # # # # # # # # # # # # # # # # # # # # # # # # # # # # # # # # # # # # # # # # # # # # # # # # # # # # # # # # # # 
+    # # # # # # # # # # # # # # # # # # # # # # # # # # # # # # # # # # # # # # # # # # # # # # # # # # # # # # # # # # # # # # # 
+    def publish_UI_addNewProj(self, *args):
+	
+	result = cmds.promptDialog(title='Add New Project', message='Enter Project Name:', button=['Accept', 'Cancel'],defaultButton='Accept',cancelButton='Cancel',dismissString='Cancel')
+
+	if result == 'Accept':
+	    proj = cmds.promptDialog(query=True, text=True)
+	    
+	    directory = self.mayaToolsDir + "/General/ART/Projects/" + proj
+	    if not os.path.exists(directory):
+		os.makedirs(directory)
+		os.makedirs(directory + "/Animations")
+		os.makedirs(directory + "/Poses")
+	    self.publish_UI_GetProjects()
+	
+
+    # # # # # # # # # # # # # # # # # # # # # # # # # # # # # # # # # # # # # # # # # # # # # # # # # # # # # # # # # # # # # # # 
+    # # # # # # # # # # # # # # # # # # # # # # # # # # # # # # # # # # # # # # # # # # # # # # # # # # # # # # # # # # # # # # # 
+    # # # # # # # # # # # # # # # # # # # # # # # # # # # # # # # # # # # # # # # # # # # # # # # # # # # # # # # # # # # # # # # 
+    # # # # # # # # # # # # # # # # # # # # # # # # # # # # # # # # # # # # # # # # # # # # # # # # # # # # # # # # # # # # # # # 
+    # # # # # # # # # # # # # # # # # # # # # # # # # # # # # # # # # # # # # # # # # # # # # # # # # # # # # # # # # # # # # # # 
+    def publish_UI_Cancel(self, *args):
+	
+	cmds.deleteUI(self.widgets["publishUIWindow"])
+	
+    # # # # # # # # # # # # # # # # # # # # # # # # # # # # # # # # # # # # # # # # # # # # # # # # # # # # # # # # # # # # # # # 
+    # # # # # # # # # # # # # # # # # # # # # # # # # # # # # # # # # # # # # # # # # # # # # # # # # # # # # # # # # # # # # # # 
+    # # # # # # # # # # # # # # # # # # # # # # # # # # # # # # # # # # # # # # # # # # # # # # # # # # # # # # # # # # # # # # # 
+    # # # # # # # # # # # # # # # # # # # # # # # # # # # # # # # # # # # # # # # # # # # # # # # # # # # # # # # # # # # # # # # 
+    # # # # # # # # # # # # # # # # # # # # # # # # # # # # # # # # # # # # # # # # # # # # # # # # # # # # # # # # # # # # # # # 
+    def screenshot_thumbnailCreationAndCheck(self, project, characterName, handCtrlSpace, *args):
+	
+	sourceControl = False
+	
+	#check if source control is on
+	settingsLocation = self.mayaToolsDir + "/General/Scripts/projectSettings.txt"
+	if os.path.exists(settingsLocation):
+	    f = open(settingsLocation, 'r')
+	    settings = cPickle.load(f)
+	    f.close()	    
+	    sourceControl = settings.get("UseSourceControl")
+	    
+	#check to see if a thumbnail already exists
+	if os.path.exists(self.mayaToolsDir + "/General/Icons/ART/Thumbnails/" + project + "/" + characterName + ".bmp"):
+	    #now check to see if the cases of the file match
+	    caseMatch = False
+	    
+	    files = os.listdir(self.mayaToolsDir + "/General/Icons/ART/Thumbnails/" + project + "/")
+	    for f in files:
+		if f == characterName + ".bmp":
+		    caseMatch = True
+		    
+	    
+	    #if it was a case match, ask if they want to overwrite
+	    if caseMatch:
+		result = cmds.confirmDialog( icon = "question", title='Thumbnail Creation', message='Character already has a thumbnail', button=['Overwrite','Use Existing'], defaultButton='Overwrite', cancelButton='Use Existing', dismissString='Use Existing' )
+	    
+		if result == "Overwrite":
+		    
+		    #first, delete the files
+		    try:
+			os.remove(self.mayaToolsDir + "/General/Icons/ART/Thumbnails/" + project + "/" + characterName + ".bmp")
+			os.remove(self.mayaToolsDir + "/General/Icons/ART/Thumbnails/" + project + "/" + characterName + "_small.bmp")
+			
+		    except:
+			
+			if sourceControl == False:
+			    cmds.confirmDialog(title = "Publish", icon = "critical", message = "Cannot overwrite existing files. Make sure files are not read only. Aborting operation.")
+			    return
+			
+			else:
+			    #if the user is using perforce:
+			    #check to see if the file is currently checked out
+			    result = cmds.confirmDialog(title = "Publish", icon = "critical", message = "Could not overwrite file. File may exist already and be marked as read only.", button = ["Check Out File", "Cancel"])
+			    
+			    if result == "Check Out File":
+				import perforceUtils
+				reload(perforceUtils)
+				writeable = perforceUtils.p4_checkOutCurrentFile(self.mayaToolsDir + "/General/Icons/ART/Thumbnails/" + project + "/" + characterName + ".bmp")
+				writeable = perforceUtils.p4_checkOutCurrentFile(self.mayaToolsDir + "/General/Icons/ART/Thumbnails/" + project + "/" + characterName + "_small.bmp")
+				
+				if writeable:
+				    try:
+					os.remove(self.mayaToolsDir + "/General/Icons/ART/Thumbnails/" + project + "/" + characterName + ".bmp")
+					os.remove(self.mayaToolsDir + "/General/Icons/ART/Thumbnails/" + project + "/" + characterName + "_small.bmp")
+					
+				    except:
+					cmds.confirmDialog(title = "Publish", icon = "critical", message = "Perforce operation unsucessful. Could not save file.")
+					return
+				    
+				else:
+				    cmds.warning("Operation aborted.")
+				    return
+				
+			    else:
+				cmds.warning("Operation Aborted.")
+				return
+			    
+				
+		    
+		    
+		    #Create a window with a model editor 
+	    
+		    self.widgets["screenshotWindow"] =  cmds.window(title = "Character Thumbnail Creator", w = 400, h = 350, sizeable = False, titleBar = False)
+		    self.widgets["screenshotFormLayout"] = cmds.formLayout()
+		    self.widgets["screenshot_modelEditor"] = cmds.modelEditor(parent = self.widgets["screenshotFormLayout"])
+		    self.widgets["screenshot_button_frame"] = cmds.frameLayout(h = 488, borderStyle = "etchedIn", lv = False, collapsable = False,  parent = self.widgets["screenshotFormLayout"])
+		    self.widgets["screenshot_button_column"] = cmds.columnLayout(parent = self.widgets["screenshot_button_frame"])
+		    
+		    
+		    
+		    #Create the screenshot button
+		    self.widgets["screenshot_button"] = cmds.symbolButton(w = 100, h = 50, image = self.mayaToolsDir + "/General/Icons/ART/screenshot_take_shot.bmp", parent = self.widgets["screenshot_button_column"], c = partial(self.screenshot_take, project, characterName, handCtrlSpace))
+		    
+		    
+		    #add sliders for foward, backward, and side and a refocus button
+		    self.widgets["refocusCameraButton"] = cmds.button(w = 100, h = 50, label='Re-Focus Camera', parent = self.widgets["screenshot_button_column"], c = self.resetCamera)
+		    
+		    #cam controls
+		    self.widgets["camControlsFormLayout"] = cmds.formLayout(w = 100, parent = self.widgets["screenshot_button_column"])
+		    self.widgets["camLabel1"] = cmds.text(label = "Camera Controls:")
+		    self.widgets["camLabel2"] = cmds.text(label = "Zoom Controls:")
+		    self.widgets["camUpButton"] = cmds.button(w = 20, h = 20, label = "^", parent = self.widgets["camControlsFormLayout"], c = partial(self.cameraControls, True, False, False, False, False, False))
+		    self.widgets["camDownButton"] = cmds.button(w = 20, h = 20, label = "v", parent = self.widgets["camControlsFormLayout"], c = partial(self.cameraControls, False, True, False, False, False, False))
+		    self.widgets["camRightButton"] = cmds.button(w = 20, h = 20, label = ">", parent = self.widgets["camControlsFormLayout"], c = partial(self.cameraControls, False, False, True, False, False, False))
+		    self.widgets["camLeftButton"] = cmds.button(w = 20, h = 20, label = "<", parent = self.widgets["camControlsFormLayout"], c = partial(self.cameraControls, False, False, False, True, False, False))
+		    self.widgets["camZoomInButton"] = cmds.button(w = 20, h = 20, label = "+", parent = self.widgets["camControlsFormLayout"], c = partial(self.cameraControls, False, False, False, False, True, False))
+		    self.widgets["camZoomOutButton"] = cmds.button(w = 20, h = 20, label = "-", parent = self.widgets["camControlsFormLayout"], c = partial(self.cameraControls, False, False, False, False, False, True))
+		    self.widgets["screenshotCancelButton"] = cmds.button(w = 100, h = 50, label = "Cancel", c = self.screenshotCancel)
+		    
+		    cmds.formLayout(self.widgets["camControlsFormLayout"], edit = True, af = [(self.widgets["camUpButton"], "left", 40), (self.widgets["camUpButton"], "top", 50)])
+		    cmds.formLayout(self.widgets["camControlsFormLayout"], edit = True, af = [(self.widgets["camDownButton"], "left", 40), (self.widgets["camDownButton"], "top", 90)])
+		    cmds.formLayout(self.widgets["camControlsFormLayout"], edit = True, af = [(self.widgets["camRightButton"], "left", 60), (self.widgets["camRightButton"], "top", 70)])
+		    cmds.formLayout(self.widgets["camControlsFormLayout"], edit = True, af = [(self.widgets["camLeftButton"], "left", 20), (self.widgets["camLeftButton"], "top", 70)])
+		    cmds.formLayout(self.widgets["camControlsFormLayout"], edit = True, af = [(self.widgets["camLabel1"], "left", 10), (self.widgets["camLabel1"], "top", 30)])
+		    cmds.formLayout(self.widgets["camControlsFormLayout"], edit = True, af = [(self.widgets["camLabel2"], "left", 15), (self.widgets["camLabel2"], "top", 120)])
+		    cmds.formLayout(self.widgets["camControlsFormLayout"], edit = True, af = [(self.widgets["camZoomInButton"], "right", 20), (self.widgets["camZoomInButton"], "top", 140)])
+		    cmds.formLayout(self.widgets["camControlsFormLayout"], edit = True, af = [(self.widgets["camZoomOutButton"], "left", 20), (self.widgets["camZoomOutButton"], "top", 140)])
+		    cmds.formLayout(self.widgets["camControlsFormLayout"], edit = True, af = [(self.widgets["screenshotCancelButton"], "left", 0), (self.widgets["screenshotCancelButton"], "bottom", 20)])
+	    
+		    
+		    
+		    
+		    #attach elements to form layout
+		    cmds.formLayout(self.widgets["screenshotFormLayout"], edit = True, af = [(self.widgets["screenshot_button_frame"], 'top', 5), (self.widgets["screenshot_button_frame"], 'left', 5)],  attachNone=[(self.widgets["screenshot_button_frame"], 'bottom'), (self.widgets["screenshot_button_frame"], 'right')])
+		    cmds.formLayout(self.widgets["screenshotFormLayout"], edit = True, af = [(self.widgets["screenshot_modelEditor"], 'top', 5), (self.widgets["screenshot_modelEditor"], 'bottom', 5), (self.widgets["screenshot_modelEditor"], 'right', 5)], attachControl=(self.widgets["screenshot_modelEditor"], 'left', 5, self.widgets["screenshot_button_frame"]))
+		    
+		    #Create a camera for the editor
+		    self.screenshot_camera = cmds.camera(name = "screenshot_cam", worldUp = (0, 0, 1), rotation = (90, 0, 0))
+		    cmds.setAttr(self.screenshot_camera[1] + ".focalLength", 70)
+		    cmds.setAttr(self.screenshot_camera[0] + ".tx", 0)
+		    if cmds.objExists("head"):
+			cmds.select("head")
+			
+		    cmds.viewFit(self.screenshot_camera[0])
+		    cmds.xform(self.screenshot_camera[0], relative = True, t = [0, -60, 0])
+		    cmds.select(clear = True)
+		    
+		    
+		    #Attach the camera to the model editor.
+		    cmds.modelEditor( self.widgets["screenshot_modelEditor"], edit = True, camera = self.screenshot_camera[0] )
+		    cmds.showWindow( self.widgets["screenshotWindow"] )
+	    
+	    
+		    #turn on smooth shading & textures
+		    cmds.modelEditor(self.widgets["screenshot_modelEditor"], edit = True, displayAppearance = "smoothShaded", displayTextures = True, headsUpDisplay = False, cameras = False, grid = False, joints = False, textures = True, lights = True, shadows = True )
+		    
+		    #use viewport 2.0
+		    mel.eval( 'setRendererInModelPanel "ogsRenderer" ' + self.widgets["screenshot_modelEditor"] ) 
+	    
+		
+		else:
+		    #publish the character
+		    self.publish(project, characterName, handCtrlSpace)
+		    
+	    else:
+		try:
+		    os.remove(self.mayaToolsDir + "/General/Icons/ART/Thumbnails/" + project + "/" + characterName + ".bmp")
+		    os.remove(self.mayaToolsDir + "/General/Icons/ART/Thumbnails/" + project + "/" + characterName + "_small.bmp")
+		    self.screenshot_thumbnailCreationAndCheck(project, characterName, handCtrlSpace)
+		    
+		except:
+		    if sourceControl == False:
+			cmds.confirmDialog(title = "Publish", icon = "critical", message = "Cannot overwrite existing files. Make sure files are not read only. Aborting operation.")
+			return
+		    
+		    else:
+			#if the user is using perforce:
+			#check to see if the file is currently checked out
+			result = cmds.confirmDialog(title = "Publish", icon = "critical", message = "Could not overwrite file. File may exist already and be marked as read only.", button = ["Check Out File", "Cancel"])
+			
+			if result == "Check Out File":
+			    import perforceUtils
+			    reload(perforceUtils)
+			    perforceUtils.p4_checkOutCurrentFile(self.mayaToolsDir + "/General/Icons/ART/Thumbnails/" + project + "/" + characterName + ".bmp")
+			    perforceUtils.p4_checkOutCurrentFile(self.mayaToolsDir + "/General/Icons/ART/Thumbnails/" + project + "/" + characterName + "_small.bmp")
+	    
+			    
+			    try:
+				os.remove(self.mayaToolsDir + "/General/Icons/ART/Thumbnails/" + project + "/" + characterName + ".bmp")
+				os.remove(self.mayaToolsDir + "/General/Icons/ART/Thumbnails/" + project + "/" + characterName + "_small.bmp")
+				
+			    except:
+				cmds.confirmDialog(title = "Publish", icon = "critical", message = "Perforce operation unsucessful. Could not save file.")
+				return	
+		    
+	    
+		
+	#create new thumbnails
+	else:
+	    
+	    #make sure files don't exist first, and if they do, remove
+	    if os.path.exists(self.mayaToolsDir + "/General/Icons/ART/Thumbnails/" + project + "/" + characterName + ".bmp"):
+		try:
+		    os.remove(self.mayaToolsDir + "/General/Icons/ART/Thumbnails/" + project + "/" + characterName + ".bmp")
+		    os.remove(self.mayaToolsDir + "/General/Icons/ART/Thumbnails/" + project + "/" + characterName + "_small.bmp")
+		    
+		except:
+		    
+		    if sourceControl == False:
+			cmds.confirmDialog(title = "Publish", icon = "critical", message = "Cannot overwrite existing files. Make sure files are not read only. Aborting operation.")
+			return
+		    
+		    else:
+			#if the user is using perforce:
+			#check to see if the file is currently checked out
+			result = cmds.confirmDialog(title = "Publish", icon = "critical", message = "Could not overwrite file. File may exist already and be marked as read only.", button = ["Check Out File", "Cancel"])
+			
+			if result == "Check Out File":
+			    import perforceUtils
+			    reload(perforceUtils)
+			    perforceUtils.p4_checkOutCurrentFile(self.mayaToolsDir + "/General/Icons/ART/Thumbnails/" + project + "/" + characterName + ".bmp")
+			    perforceUtils.p4_checkOutCurrentFile(self.mayaToolsDir + "/General/Icons/ART/Thumbnails/" + project + "/" + characterName + "_small.bmp")
+	    
+			    
+			    try:
+				os.remove(self.mayaToolsDir + "/General/Icons/ART/Thumbnails/" + project + "/" + characterName + ".bmp")
+				os.remove(self.mayaToolsDir + "/General/Icons/ART/Thumbnails/" + project + "/" + characterName + "_small.bmp")
+				
+			    except:
+				cmds.confirmDialog(title = "Publish", icon = "critical", message = "Perforce operation unsucessful. Could not save file.")
+				return	
+		    
+	    
+	    #Create a window with a model editor 
+    
+	    self.widgets["screenshotWindow"] =  cmds.window(title = "Character Thumbnail Creator", w = 400, h = 350, sizeable = False, titleBar = False)
+	    self.widgets["screenshotFormLayout"] = cmds.formLayout()
+	    self.widgets["screenshot_modelEditor"] = cmds.modelEditor(parent = self.widgets["screenshotFormLayout"])
+	    self.widgets["screenshot_button_frame"] = cmds.frameLayout(h = 488, borderStyle = "etchedIn", lv = False, collapsable = False,  parent = self.widgets["screenshotFormLayout"])
+	    self.widgets["screenshot_button_column"] = cmds.columnLayout(parent = self.widgets["screenshot_button_frame"])
+	    
+	    
+	    
+	    #Create the screenshot button
+	    self.widgets["screenshot_button"] = cmds.symbolButton(w = 100, h = 50, image = self.mayaToolsDir + "/General/Icons/ART/screenshot_take_shot.bmp", parent = self.widgets["screenshot_button_column"], c = partial(self.screenshot_take, project, characterName, handCtrlSpace))
+	    
+	    
+	    #add sliders for foward, backward, and side and a refocus button
+	    self.widgets["refocusCameraButton"] = cmds.button(w = 100, h = 50, label='Re-Focus Camera', parent = self.widgets["screenshot_button_column"], c = self.resetCamera)
+	    
+	    #cam controls
+	    self.widgets["camControlsFormLayout"] = cmds.formLayout(w = 100, parent = self.widgets["screenshot_button_column"])
+	    self.widgets["camLabel1"] = cmds.text(label = "Camera Controls:")
+	    self.widgets["camLabel2"] = cmds.text(label = "Zoom Controls:")
+	    self.widgets["camUpButton"] = cmds.button(w = 20, h = 20, label = "^", parent = self.widgets["camControlsFormLayout"], c = partial(self.cameraControls, True, False, False, False, False, False))
+	    self.widgets["camDownButton"] = cmds.button(w = 20, h = 20, label = "v", parent = self.widgets["camControlsFormLayout"], c = partial(self.cameraControls, False, True, False, False, False, False))
+	    self.widgets["camRightButton"] = cmds.button(w = 20, h = 20, label = ">", parent = self.widgets["camControlsFormLayout"], c = partial(self.cameraControls, False, False, True, False, False, False))
+	    self.widgets["camLeftButton"] = cmds.button(w = 20, h = 20, label = "<", parent = self.widgets["camControlsFormLayout"], c = partial(self.cameraControls, False, False, False, True, False, False))
+	    self.widgets["camZoomInButton"] = cmds.button(w = 20, h = 20, label = "+", parent = self.widgets["camControlsFormLayout"], c = partial(self.cameraControls, False, False, False, False, True, False))
+	    self.widgets["camZoomOutButton"] = cmds.button(w = 20, h = 20, label = "-", parent = self.widgets["camControlsFormLayout"], c = partial(self.cameraControls, False, False, False, False, False, True))
+	    self.widgets["screenshotCancelButton"] = cmds.button(w = 100, h = 30, label = "Cancel", c = self.screenshotCancel)
+	    
+	    self.widgets["useLighting"] = cmds.iconTextCheckBox(w = 30, h = 30, onc = self.publish_UI_TakeScreenshot_LightsOn, ofc = self.publish_UI_TakeScreenshot_LightsOff, ann = "Use Lighting", image = self.mayaToolsDir + "/General/Icons/ART/light_off.bmp", selectionImage = self.mayaToolsDir + "/General/Icons/ART/light_on.bmp", v = False)
+	    self.widgets["useHiRender"] = cmds.iconTextCheckBox(w = 60, h = 30, onc = self.publish_UI_TakeScreenshot_hqRenderOn, ofc = self.publish_UI_TakeScreenshot_hqRenderOff, ann = "Use High Quality Renderer", v = True, image = self.mayaToolsDir + "/General/Icons/ART/hq_off.bmp", selectionImage = self.mayaToolsDir + "/General/Icons/ART/hq_on.bmp")
+	    
+	    cmds.formLayout(self.widgets["camControlsFormLayout"], edit = True, af = [(self.widgets["camUpButton"], "left", 40), (self.widgets["camUpButton"], "top", 50)])
+	    cmds.formLayout(self.widgets["camControlsFormLayout"], edit = True, af = [(self.widgets["camDownButton"], "left", 40), (self.widgets["camDownButton"], "top", 90)])
+	    cmds.formLayout(self.widgets["camControlsFormLayout"], edit = True, af = [(self.widgets["camRightButton"], "left", 60), (self.widgets["camRightButton"], "top", 70)])
+	    cmds.formLayout(self.widgets["camControlsFormLayout"], edit = True, af = [(self.widgets["camLeftButton"], "left", 20), (self.widgets["camLeftButton"], "top", 70)])
+	    cmds.formLayout(self.widgets["camControlsFormLayout"], edit = True, af = [(self.widgets["camLabel1"], "left", 10), (self.widgets["camLabel1"], "top", 30)])
+	    cmds.formLayout(self.widgets["camControlsFormLayout"], edit = True, af = [(self.widgets["camLabel2"], "left", 15), (self.widgets["camLabel2"], "top", 120)])
+	    cmds.formLayout(self.widgets["camControlsFormLayout"], edit = True, af = [(self.widgets["camZoomInButton"], "right", 20), (self.widgets["camZoomInButton"], "top", 140)])
+	    cmds.formLayout(self.widgets["camControlsFormLayout"], edit = True, af = [(self.widgets["camZoomOutButton"], "left", 20), (self.widgets["camZoomOutButton"], "top", 140)])
+	    cmds.formLayout(self.widgets["camControlsFormLayout"], edit = True, af = [(self.widgets["useLighting"], "left", 5), (self.widgets["useLighting"], "top", 165)])
+	    cmds.formLayout(self.widgets["camControlsFormLayout"], edit = True, af = [(self.widgets["useHiRender"], "right", 5), (self.widgets["useHiRender"], "top", 165)])
+	    
+	    cmds.formLayout(self.widgets["camControlsFormLayout"], edit = True, af = [(self.widgets["screenshotCancelButton"], "left", 0), (self.widgets["screenshotCancelButton"], "top", 200)])
+    
+	    
+	    
+	    
+	    #attach elements to form layout
+	    cmds.formLayout(self.widgets["screenshotFormLayout"], edit = True, af = [(self.widgets["screenshot_button_frame"], 'top', 5), (self.widgets["screenshot_button_frame"], 'left', 5)],  attachNone=[(self.widgets["screenshot_button_frame"], 'bottom'), (self.widgets["screenshot_button_frame"], 'right')])
+	    cmds.formLayout(self.widgets["screenshotFormLayout"], edit = True, af = [(self.widgets["screenshot_modelEditor"], 'top', 5), (self.widgets["screenshot_modelEditor"], 'bottom', 5), (self.widgets["screenshot_modelEditor"], 'right', 5)], attachControl=(self.widgets["screenshot_modelEditor"], 'left', 5, self.widgets["screenshot_button_frame"]))
+	    
+	    #Create a camera for the editor
+	    self.screenshot_camera = cmds.camera(name = "screenshot_cam", worldUp = (0, 0, 1), rotation = (90, 0, 0))
+	    cmds.setAttr(self.screenshot_camera[1] + ".focalLength", 70)
+	    cmds.setAttr(self.screenshot_camera[0] + ".tx", 0)
+	    if cmds.objExists("head"):
+		cmds.select("head")
+		
+	    cmds.viewFit(self.screenshot_camera[0])
+	    cmds.xform(self.screenshot_camera[0], relative = True, t = [0, -60, 0])
+	    cmds.select(clear = True)
+	    
+	    
+	    #Attach the camera to the model editor.
+	    cmds.modelEditor( self.widgets["screenshot_modelEditor"], edit = True, camera = self.screenshot_camera[0] )
+	    cmds.showWindow( self.widgets["screenshotWindow"] )
+    
+    
+	    #turn on smooth shading & textures
+	    cmds.modelEditor(self.widgets["screenshot_modelEditor"], edit = True, displayAppearance = "smoothShaded", displayTextures = True, headsUpDisplay = False, cameras = False, grid = False, joints = False, textures = True, dl = "default", lights = False, shadows = False )
+	    
+	    #use viewport 2.0
+	    mel.eval( 'setRendererInModelPanel "ogsRenderer" ' + self.widgets["screenshot_modelEditor"] ) 	    
+
+	
+    # # # # # # # # # # # # # # # # # # # # # # # # # # # # # # # # # # # # # # # # # # # # # # # # # # # # # # # # # # # # # # # 
+    # # # # # # # # # # # # # # # # # # # # # # # # # # # # # # # # # # # # # # # # # # # # # # # # # # # # # # # # # # # # # # # 
+    # # # # # # # # # # # # # # # # # # # # # # # # # # # # # # # # # # # # # # # # # # # # # # # # # # # # # # # # # # # # # # # 
+    # # # # # # # # # # # # # # # # # # # # # # # # # # # # # # # # # # # # # # # # # # # # # # # # # # # # # # # # # # # # # # # 
+    # # # # # # # # # # # # # # # # # # # # # # # # # # # # # # # # # # # # # # # # # # # # # # # # # # # # # # # # # # # # # # # 
+    def publish_UI_TakeScreenshot_LightsOn(self, *args):
+	
+	cmds.modelEditor(self.widgets["screenshot_modelEditor"], edit = True, dl = "all", shadows = True )
+	
+	
+    # # # # # # # # # # # # # # # # # # # # # # # # # # # # # # # # # # # # # # # # # # # # # # # # # # # # # # # # # # # # # # # 
+    # # # # # # # # # # # # # # # # # # # # # # # # # # # # # # # # # # # # # # # # # # # # # # # # # # # # # # # # # # # # # # # 
+    # # # # # # # # # # # # # # # # # # # # # # # # # # # # # # # # # # # # # # # # # # # # # # # # # # # # # # # # # # # # # # # 
+    # # # # # # # # # # # # # # # # # # # # # # # # # # # # # # # # # # # # # # # # # # # # # # # # # # # # # # # # # # # # # # # 
+    # # # # # # # # # # # # # # # # # # # # # # # # # # # # # # # # # # # # # # # # # # # # # # # # # # # # # # # # # # # # # # # 
+    def publish_UI_TakeScreenshot_LightsOff(self, *args):
+	
+	cmds.modelEditor(self.widgets["screenshot_modelEditor"], edit = True, dl = "default", shadows = False )
+	
+	
+    # # # # # # # # # # # # # # # # # # # # # # # # # # # # # # # # # # # # # # # # # # # # # # # # # # # # # # # # # # # # # # # 
+    # # # # # # # # # # # # # # # # # # # # # # # # # # # # # # # # # # # # # # # # # # # # # # # # # # # # # # # # # # # # # # # 
+    # # # # # # # # # # # # # # # # # # # # # # # # # # # # # # # # # # # # # # # # # # # # # # # # # # # # # # # # # # # # # # # 
+    # # # # # # # # # # # # # # # # # # # # # # # # # # # # # # # # # # # # # # # # # # # # # # # # # # # # # # # # # # # # # # # 
+    # # # # # # # # # # # # # # # # # # # # # # # # # # # # # # # # # # # # # # # # # # # # # # # # # # # # # # # # # # # # # # # 
+    def publish_UI_TakeScreenshot_hqRenderOn(self, *args):
+	
+	#use viewport 2.0
+	mel.eval( 'setRendererInModelPanel "ogsRenderer" ' + self.widgets["screenshot_modelEditor"] ) 
+	
+
+    # # # # # # # # # # # # # # # # # # # # # # # # # # # # # # # # # # # # # # # # # # # # # # # # # # # # # # # # # # # # # # # 
+    # # # # # # # # # # # # # # # # # # # # # # # # # # # # # # # # # # # # # # # # # # # # # # # # # # # # # # # # # # # # # # # 
+    # # # # # # # # # # # # # # # # # # # # # # # # # # # # # # # # # # # # # # # # # # # # # # # # # # # # # # # # # # # # # # # 
+    # # # # # # # # # # # # # # # # # # # # # # # # # # # # # # # # # # # # # # # # # # # # # # # # # # # # # # # # # # # # # # # 
+    # # # # # # # # # # # # # # # # # # # # # # # # # # # # # # # # # # # # # # # # # # # # # # # # # # # # # # # # # # # # # # # 
+    def publish_UI_TakeScreenshot_hqRenderOff(self, *args):
+	
+	#use viewport 2.0
+	mel.eval( 'setRendererInModelPanel "base_OpenGL_Renderer" ' + self.widgets["screenshot_modelEditor"] ) 
+	
+	
+    # # # # # # # # # # # # # # # # # # # # # # # # # # # # # # # # # # # # # # # # # # # # # # # # # # # # # # # # # # # # # # # 
+    # # # # # # # # # # # # # # # # # # # # # # # # # # # # # # # # # # # # # # # # # # # # # # # # # # # # # # # # # # # # # # # 
+    # # # # # # # # # # # # # # # # # # # # # # # # # # # # # # # # # # # # # # # # # # # # # # # # # # # # # # # # # # # # # # # 
+    # # # # # # # # # # # # # # # # # # # # # # # # # # # # # # # # # # # # # # # # # # # # # # # # # # # # # # # # # # # # # # # 
+    # # # # # # # # # # # # # # # # # # # # # # # # # # # # # # # # # # # # # # # # # # # # # # # # # # # # # # # # # # # # # # # 
+    def publish_UI_TakeScreenshot(self, *args):
+	
+	#get info from publish UI
+	project = cmds.optionMenu(self.widgets["publishUI_ProjOM"], q = True, v = True)
+	
+	if project == None:
+	    cmds.confirmDialog(title = "Publish", message = "No Project selected")
+	    return	    
+
+	characterName = cmds.textField(self.widgets["publishUICharNameTF"], q = True, text = True)
+	if characterName == "":
+	    return
+	
+	handCtrlSpace = True                        
+	
+	if characterName.find(" ") != -1:
+	    characterName = characterName.replace(" ", "_")
+	    
+
+	path = (self.mayaToolsDir + "/General/ART/Projects/" + project + "/AnimRigs/" + characterName + ".mb")
+	
+	if os.path.exists(path):
+	    
+	    #even if the path exists, need to check for case sensitivity
+	    caseMatch = False
+	    files = os.listdir(self.mayaToolsDir + "/General/ART/Projects/" + project + "/AnimRigs/")
+	    for f in files:
+		if f == characterName + ".mb":
+		    caseMatch = True
+	    
+	    if caseMatch:
+		#now ask if we want to overwrite(only if cases match)
+		result = cmds.confirmDialog(icon = "question",  title='Naming Conflict', message='A Character already exists with that name.', button=['Overwrite','Cancel'], defaultButton='Overwrite', cancelButton='Cancel', dismissString='Cancel' )
+
+		if result == "Cancel":
+		    return
+		    
+		else:
+		    #overwrite the character and check for thumbnails
+		    self.screenshot_thumbnailCreationAndCheck(project, characterName, handCtrlSpace)
+		    
+	    #if there was no case match
+	    else:
+		files = os.listdir(self.mayaToolsDir + "/General/ART/Projects/" + project + "/AnimRigs/")
+		for f in files:
+		    if f.partition(".")[0].lower() == characterName.lower():
+			characterName = f.partition(".")[0]
+		
+		self.screenshot_thumbnailCreationAndCheck(project, characterName, handCtrlSpace)
+			
+	else:
+	    #if there was no such character, move onto thumbnail creation
+	    self.screenshot_thumbnailCreationAndCheck(project, characterName, handCtrlSpace)
+
+
+    # # # # # # # # # # # # # # # # # # # # # # # # # # # # # # # # # # # # # # # # # # # # # # # # # # # # # # # # # # # # # # # 
+    # # # # # # # # # # # # # # # # # # # # # # # # # # # # # # # # # # # # # # # # # # # # # # # # # # # # # # # # # # # # # # # 
+    # # # # # # # # # # # # # # # # # # # # # # # # # # # # # # # # # # # # # # # # # # # # # # # # # # # # # # # # # # # # # # # 
+    # # # # # # # # # # # # # # # # # # # # # # # # # # # # # # # # # # # # # # # # # # # # # # # # # # # # # # # # # # # # # # # 
+    # # # # # # # # # # # # # # # # # # # # # # # # # # # # # # # # # # # # # # # # # # # # # # # # # # # # # # # # # # # # # # # 
+    def screenshotCancel(self, *args):
+	cmds.deleteUI(self.widgets["screenshotWindow"] )
+	cmds.delete(self.screenshot_camera[0])
+	
+	
+    # # # # # # # # # # # # # # # # # # # # # # # # # # # # # # # # # # # # # # # # # # # # # # # # # # # # # # # # # # # # # # # 
+    # # # # # # # # # # # # # # # # # # # # # # # # # # # # # # # # # # # # # # # # # # # # # # # # # # # # # # # # # # # # # # # 
+    # # # # # # # # # # # # # # # # # # # # # # # # # # # # # # # # # # # # # # # # # # # # # # # # # # # # # # # # # # # # # # # 
+    # # # # # # # # # # # # # # # # # # # # # # # # # # # # # # # # # # # # # # # # # # # # # # # # # # # # # # # # # # # # # # # 
+    # # # # # # # # # # # # # # # # # # # # # # # # # # # # # # # # # # # # # # # # # # # # # # # # # # # # # # # # # # # # # # # 
+    def resetCamera(self, *args):
+	if cmds.objExists("head"):
+	    cmds.select("head")
+	    
+	cmds.viewFit(self.screenshot_camera[0])
+	cmds.xform(self.screenshot_camera[0], relative = True, t = [0, -60, 0])
+	cmds.select(clear = True)
+	
+
+    # # # # # # # # # # # # # # # # # # # # # # # # # # # # # # # # # # # # # # # # # # # # # # # # # # # # # # # # # # # # # # # 
+    # # # # # # # # # # # # # # # # # # # # # # # # # # # # # # # # # # # # # # # # # # # # # # # # # # # # # # # # # # # # # # # 
+    # # # # # # # # # # # # # # # # # # # # # # # # # # # # # # # # # # # # # # # # # # # # # # # # # # # # # # # # # # # # # # # 
+    # # # # # # # # # # # # # # # # # # # # # # # # # # # # # # # # # # # # # # # # # # # # # # # # # # # # # # # # # # # # # # # 
+    # # # # # # # # # # # # # # # # # # # # # # # # # # # # # # # # # # # # # # # # # # # # # # # # # # # # # # # # # # # # # # # 
+    def cameraControls(self, up, down, left, right, zoomIn, zoomOut, *args):
+	
+	if up:
+	    cmds.select(self.screenshot_camera[0])
+	    cmds.xform(self.screenshot_camera[0], relative = True, t = [0, 0, 5])
+	    cmds.select(clear = True)
+	    
+	if down:
+	    cmds.select(self.screenshot_camera[0])
+	    cmds.xform(self.screenshot_camera[0], relative = True, t = [0, 0, -5])
+	    cmds.select(clear = True)
+	    
+	if left:
+	    cmds.select(self.screenshot_camera[0])
+	    cmds.xform(self.screenshot_camera[0], relative = True, t = [-5, 0, 0])
+	    cmds.select(clear = True)
+	    
+	if right:
+	    cmds.select(self.screenshot_camera[0])
+	    cmds.xform(self.screenshot_camera[0], relative = True, t = [5, 0, 0])
+	    cmds.select(clear = True)
+	    
+	if zoomIn:
+	    cmds.select(self.screenshot_camera[0])
+	    cmds.xform(self.screenshot_camera[0], relative = True, t = [0, 5, 0])
+	    cmds.select(clear = True)
+	    
+	if zoomOut:
+	    cmds.select(self.screenshot_camera[0])
+	    cmds.xform(self.screenshot_camera[0], relative = True, t = [0, -5, 0])
+	    cmds.select(clear = True)
+	    
+	    
+    # # # # # # # # # # # # # # # # # # # # # # # # # # # # # # # # # # # # # # # # # # # # # # # # # # # # # # # # # # # # # # # 
+    # # # # # # # # # # # # # # # # # # # # # # # # # # # # # # # # # # # # # # # # # # # # # # # # # # # # # # # # # # # # # # # 
+    # # # # # # # # # # # # # # # # # # # # # # # # # # # # # # # # # # # # # # # # # # # # # # # # # # # # # # # # # # # # # # # 
+    # # # # # # # # # # # # # # # # # # # # # # # # # # # # # # # # # # # # # # # # # # # # # # # # # # # # # # # # # # # # # # # 
+    # # # # # # # # # # # # # # # # # # # # # # # # # # # # # # # # # # # # # # # # # # # # # # # # # # # # # # # # # # # # # # # 
+    def cameraZoom(self, *args):
+	val = cmds.floatSlider(self.widgets["camZoomSlider"], q = True, v = True)
+	cmds.xform(self.screenshot_camera[0], relative = True, t = [0, val, 0])
+
+	
+    # # # # # # # # # # # # # # # # # # # # # # # # # # # # # # # # # # # # # # # # # # # # # # # # # # # # # # # # # # # # # # # 
+    # # # # # # # # # # # # # # # # # # # # # # # # # # # # # # # # # # # # # # # # # # # # # # # # # # # # # # # # # # # # # # # 
+    # # # # # # # # # # # # # # # # # # # # # # # # # # # # # # # # # # # # # # # # # # # # # # # # # # # # # # # # # # # # # # # 
+    # # # # # # # # # # # # # # # # # # # # # # # # # # # # # # # # # # # # # # # # # # # # # # # # # # # # # # # # # # # # # # # 
+    # # # # # # # # # # # # # # # # # # # # # # # # # # # # # # # # # # # # # # # # # # # # # # # # # # # # # # # # # # # # # # # 
+    def screenshot_take(self, project, characterName, handCtrlSpace, *args):
+	
+	if not os.path.exists(self.mayaToolsDir + "/General/Icons/ART/Thumbnails/" + project):
+	    os.makedirs(self.mayaToolsDir + "/General/Icons/ART/Thumbnails/" + project)
+	    
+
+	
+	#grab the current render globals
+	format = cmds.getAttr("defaultRenderGlobals.imageFormat")
+	
+	#set to render out bmp
+	cmds.setAttr("defaultRenderGlobals.imageFormat", 20)
+	
+	#get the current background color settings
+	settingsPath = cmds.internalVar(upd = True) + "userRGBColors2.mel"
+	if os.path.exists(settingsPath):
+	    
+	    #read the background color lines
+	    f = open(settingsPath, 'r')
+	    lines = f.readlines()
+	    background = lines[0].rpartition("background ")[2].partition(";")[0]
+	    backgroundTop = lines[1].rpartition("backgroundTop ")[2].partition(";")[0]
+	    backgroundBottom = lines[2].rpartition("backgroundBottom ")[2].partition(";")[0]
+	    f.close()
+	    
+	#change the background color settings
+	melCode = "displayRGBColor \"background\" .27 .27 .27 ; displayRGBColor \"backgroundTop\" .27 .27 .27 ; displayRGBColor \"backgroundBottom\" .27 .27 .27 ; "
+	mel.eval(melCode)
+	
+
+	#playblast 1 frame to capture the thumbnail
+	if cmds.playblast(activeEditor = True) != self.widgets["screenshot_modelEditor"]:
+	    cmds.modelEditor(self.widgets["screenshot_modelEditor"], edit = True, activeView = True)
+
+	cmds.playblast(frame = 1, format = "image", cf = self.mayaToolsDir + "/General/Icons/ART/Thumbnails/" + project + "/" + characterName + ".bmp", v = False, wh = [200, 200], p = 100)
+	
+	cmds.playblast(frame = 1, format = "image", cf = self.mayaToolsDir + "/General/Icons/ART/Thumbnails/" + project + "/" + characterName + "_small.bmp", v = False, wh = [50, 50], p = 100)
+	
+	#restore old settings
+	cmds.setAttr("defaultRenderGlobals.imageFormat", format)
+
+	#delete the UI
+	cmds.deleteUI(self.widgets["screenshotWindow"])
+
+	#delete the camera
+	cmds.delete(self.screenshot_camera)
+	
+	#restore the background color settings
+	
+	if os.path.exists(settingsPath):
+	    melCode = "displayRGBColor \"background\"" + background + "; displayRGBColor \"backgroundTop\"" + backgroundTop + "; displayRGBColor \"backgroundBottom\"" + backgroundBottom + "; "
+	    mel.eval(melCode)
+	
+	#publish the character
+	self.publish(project, characterName, handCtrlSpace)
+	
+	
+    # # # # # # # # # # # # # # # # # # # # # # # # # # # # # # # # # # # # # # # # # # # # # # # # # # # # # # # # # # # # # # # 
+    # # # # # # # # # # # # # # # # # # # # # # # # # # # # # # # # # # # # # # # # # # # # # # # # # # # # # # # # # # # # # # # 
+    # # # # # # # # # # # # # # # # # # # # # # # # # # # # # # # # # # # # # # # # # # # # # # # # # # # # # # # # # # # # # # # 
+    # # # # # # # # # # # # # # # # # # # # # # # # # # # # # # # # # # # # # # # # # # # # # # # # # # # # # # # # # # # # # # # 
+    # # # # # # # # # # # # # # # # # # # # # # # # # # # # # # # # # # # # # # # # # # # # # # # # # # # # # # # # # # # # # # # 
+    def deletePreviewSkeleton(self, *args):
+	if cmds.objExists("preview_root"):
+	    cmds.delete("preview_root")	
+	
+    # # # # # # # # # # # # # # # # # # # # # # # # # # # # # # # # # # # # # # # # # # # # # # # # # # # # # # # # # # # # # # # 
+    # # # # # # # # # # # # # # # # # # # # # # # # # # # # # # # # # # # # # # # # # # # # # # # # # # # # # # # # # # # # # # # 
+    # # # # # # # # # # # # # # # # # # # # # # # # # # # # # # # # # # # # # # # # # # # # # # # # # # # # # # # # # # # # # # # 
+    # # # # # # # # # # # # # # # # # # # # # # # # # # # # # # # # # # # # # # # # # # # # # # # # # # # # # # # # # # # # # # # 
+    # # # # # # # # # # # # # # # # # # # # # # # # # # # # # # # # # # # # # # # # # # # # # # # # # # # # # # # # # # # # # # # 
+    def buildPreviewSkeleton(self, *args):
+	
+	
+	if cmds.objExists("preview_root"):
+	    cmds.delete("preview_root")
+	    
+	
+	#build the root
+	cmds.select(clear = True)
+	joint = cmds.joint(name = "preview_root")
+	constraint = cmds.parentConstraint("root_lra", joint)[0]
+	cmds.delete(constraint)
+	cmds.select(clear = True)
+
+	parentInfo = []
+	
+	for mover in self.geoMovers:
+	    mover = mover.partition("|")[0]
+	    children = cmds.listRelatives(mover, children = True, type = "transform")
+	    if children != None or mover.find("ball") == 0:
+
+	    
+		prefix = mover.partition("_geo_mover")[0]
+		suffix = mover.partition("_geo_mover")[2]
+		lra = prefix + suffix + "_lra"
+
+		    
+		jointName = "preview_" + prefix + suffix
+		moverGrp = prefix + "_mover" + suffix + "_grp"
+		parent = cmds.listRelatives(moverGrp, parent = True)[0]
+		jpPrefix = parent.partition("_mover")[0]
+		jpSuffix = parent.partition("_mover")[2]
+		jointParent = "preview_" + jpPrefix + jpSuffix
+
+		cmds.select(clear = True)
+		joint = cmds.joint(name = jointName)
+		constraint = cmds.parentConstraint(lra, joint)[0]
+		cmds.delete(constraint)
+		cmds.select(clear = True)
+		parentInfo.append([joint, jointParent, lra])
+	    
+
+
+
+	for item in parentInfo:
+	    joint = item[0]
+	    parent = item[1]
+
+	    if parent != "":
+		cmds.parent(joint, parent)
+		
+	cmds.makeIdentity("preview_root", r = 1, apply = True)
+	
+	
+
+	#duplicate the root, delete the old root, and rename the dupe root
+	newRoot = cmds.duplicate("preview_root")[0]
+	cmds.delete("preview_root")
+	cmds.rename(newRoot, "preview_root")
+	
+	#reconstrain to lras
+	for item in parentInfo:
+	    joint = item[0]
+	    lra = item[2]
+	    
+	    constraint = cmds.parentConstraint(lra, joint)
+	    
+    # # # # # # # # # # # # # # # # # # # # # # # # # # # # # # # # # # # # # # # # # # # # # # # # # # # # # # # # # # # # # # # 
+    # # # # # # # # # # # # # # # # # # # # # # # # # # # # # # # # # # # # # # # # # # # # # # # # # # # # # # # # # # # # # # # 
+    # # # # # # # # # # # # # # # # # # # # # # # # # # # # # # # # # # # # # # # # # # # # # # # # # # # # # # # # # # # # # # # 
+    # # # # # # # # # # # # # # # # # # # # # # # # # # # # # # # # # # # # # # # # # # # # # # # # # # # # # # # # # # # # # # # 
+    # # # # # # # # # # # # # # # # # # # # # # # # # # # # # # # # # # # # # # # # # # # # # # # # # # # # # # # # # # # # # # # 
+    def buildFinalSkeleton(self):
+	
+	cmds.lockNode("JointMover", lock = False)
+	
+	cmds.rename("root", "joint_mover_root")
+	
+	#build the root
+	cmds.select(clear = True)
+	joint = cmds.joint(name = "root")
+	constraint = cmds.parentConstraint("root_lra", joint)[0]
+	cmds.delete(constraint)
+	cmds.select(clear = True)
+
+	parentInfo = []
+	
+	for mover in self.geoMovers:
+	    mover = mover.partition("|")[0]
+	    children = cmds.listRelatives(mover, children = True, type = "transform")
+	    if children != None or mover.find("ball") == 0:
+
+	    
+		prefix = mover.partition("_geo_mover")[0]
+		suffix = mover.partition("_geo_mover")[2]
+		lra = prefix + suffix + "_lra"
+
+		    
+		jointName = prefix + suffix
+		moverGrp = prefix + "_mover" + suffix + "_grp"
+		parent = cmds.listRelatives(moverGrp, parent = True)[0]
+		jpPrefix = parent.partition("_mover")[0]
+		jpSuffix = parent.partition("_mover")[2]
+		jointParent = jpPrefix + jpSuffix
+
+		cmds.select(clear = True)
+		joint = cmds.joint(name = jointName)
+		constraint = cmds.parentConstraint(lra, joint)[0]
+		cmds.delete(constraint)
+		cmds.select(clear = True)
+		parentInfo.append([joint, jointParent])
+	    
+
+
+
+	for item in parentInfo:
+	    joint = item[0]
+	    parent = item[1]
+	    
+	    if parent != "":
+		cmds.parent(joint, parent)
+		
+	cmds.makeIdentity("root", r = 1, apply = True)
+	
+	
+
+	#duplicate the root, delete the old root, and rename the dupe root
+	newRoot = cmds.duplicate("root")[0]
+	cmds.delete("root")
+	cmds.rename(newRoot, "root")
+	
+	
+	#lock it down --this is causing issues with binding meshes to the skeleton. 
+	#It was an effort to prevent the user from reparenting after the joints have been created.
+	#cmds.select("root", hi = True)
+	#joints = cmds.ls(sl = True)
+	
+	#for joint in joints:
+	    #cmds.lockNode(joint, lock = True)
+	    
+	#cmds.lockNode("JointMover", lock = True)
+	    
+    # # # # # # # # # # # # # # # # # # # # # # # # # # # # # # # # # # # # # # # # # # # # # # # # # # # # # # # # # # # # # # # 
+    # # # # # # # # # # # # # # # # # # # # # # # # # # # # # # # # # # # # # # # # # # # # # # # # # # # # # # # # # # # # # # # 
+    # # # # # # # # # # # # # # # # # # # # # # # # # # # # # # # # # # # # # # # # # # # # # # # # # # # # # # # # # # # # # # # 
+    # # # # # # # # # # # # # # # # # # # # # # # # # # # # # # # # # # # # # # # # # # # # # # # # # # # # # # # # # # # # # # # 
+    # # # # # # # # # # # # # # # # # # # # # # # # # # # # # # # # # # # # # # # # # # # # # # # # # # # # # # # # # # # # # # # 
+    def projectWeights(self, *args):
+	#auto weight meshes
+	meshes = cmds.ls(type = "mesh")
+	skinnableGeo = []
+	for mesh in meshes:
+	    parent = cmds.listRelatives(mesh, parent = True, type = "transform")[0]
+	    if parent != None:
+		if parent.find("proxy_geo") == -1:
+		    if parent.find("lra") == -1:
+			if parent.find("extra_joints") == -1:
+			    if parent.find("skin_geo") == -1:
+				skinnableGeo.append(parent)
+				
+	if len(skinnableGeo) > 0:
+	    self.autoWeightMeshes(skinnableGeo)
+	    
+	else:
+	    cmds.confirmDialog(icon = "information", title = "Project Weights", message = "No meshes found to project proxy mesh weights onto.")
+	    
+	
+	
+    # # # # # # # # # # # # # # # # # # # # # # # # # # # # # # # # # # # # # # # # # # # # # # # # # # # # # # # # # # # # # # # 
+    # # # # # # # # # # # # # # # # # # # # # # # # # # # # # # # # # # # # # # # # # # # # # # # # # # # # # # # # # # # # # # # 
+    # # # # # # # # # # # # # # # # # # # # # # # # # # # # # # # # # # # # # # # # # # # # # # # # # # # # # # # # # # # # # # # 
+    # # # # # # # # # # # # # # # # # # # # # # # # # # # # # # # # # # # # # # # # # # # # # # # # # # # # # # # # # # # # # # # 
+    # # # # # # # # # # # # # # # # # # # # # # # # # # # # # # # # # # # # # # # # # # # # # # # # # # # # # # # # # # # # # # # 
+    def autoWeightMeshes(self, geo):
+	
+	if cmds.window("autoWeightMeshes_UI", exists = True):
+	    cmds.deleteUI("autoWeightMeshes_UI")
+	window = cmds.window("autoWeightMeshes_UI", title = "Skin Weight", w = 300, h = 400, titleBarMenu = False)
+	mainLayout = cmds.columnLayout(w = 300, h = 400, co = ["both", 5], rs = 10)
+	cmds.textScrollList("autoWeightMeshList", w = 280, h = 300, ams = True)
+	
+	for obj in set(geo):
+	    cmds.textScrollList("autoWeightMeshList", edit =True, append = obj)
+	cmds.button(label = "Skin Weight Selected", w = 280, ann = "puts a first pass weighting on selected meshes using your proxy geometry as a guide for the weights", c = self.autoWeightMeshes_Accept)
+	cmds.button(label = "Skip This Step", w = 280, c = self.autoWeightMeshes_Cancel)
+	
+	cmds.showWindow(window)
+	
+	
+    # # # # # # # # # # # # # # # # # # # # # # # # # # # # # # # # # # # # # # # # # # # # # # # # # # # # # # # # # # # # # # # 
+    # # # # # # # # # # # # # # # # # # # # # # # # # # # # # # # # # # # # # # # # # # # # # # # # # # # # # # # # # # # # # # # 
+    # # # # # # # # # # # # # # # # # # # # # # # # # # # # # # # # # # # # # # # # # # # # # # # # # # # # # # # # # # # # # # # 
+    # # # # # # # # # # # # # # # # # # # # # # # # # # # # # # # # # # # # # # # # # # # # # # # # # # # # # # # # # # # # # # # 
+    # # # # # # # # # # # # # # # # # # # # # # # # # # # # # # # # # # # # # # # # # # # # # # # # # # # # # # # # # # # # # # # 
+    def autoWeightMeshes_Cancel(self, *args):
+	cmds.deleteUI("autoWeightMeshes_UI")
+	
+    # # # # # # # # # # # # # # # # # # # # # # # # # # # # # # # # # # # # # # # # # # # # # # # # # # # # # # # # # # # # # # # 
+    # # # # # # # # # # # # # # # # # # # # # # # # # # # # # # # # # # # # # # # # # # # # # # # # # # # # # # # # # # # # # # # 
+    # # # # # # # # # # # # # # # # # # # # # # # # # # # # # # # # # # # # # # # # # # # # # # # # # # # # # # # # # # # # # # # 
+    # # # # # # # # # # # # # # # # # # # # # # # # # # # # # # # # # # # # # # # # # # # # # # # # # # # # # # # # # # # # # # # 
+    # # # # # # # # # # # # # # # # # # # # # # # # # # # # # # # # # # # # # # # # # # # # # # # # # # # # # # # # # # # # # # # 
+    def autoWeightMeshes_Accept(self, *args):
+	
+	#get selected items in text scroll List
+	meshes = cmds.textScrollList("autoWeightMeshList", q = True, si = True)
+	noSkin = True
+	
+	
+	#unlock joints
+	cmds.select("root", hi = True)
+	cmds.lockNode(lock = False)
+	cmds.select(clear = True)
+	
+	
+	if meshes:
+	
+	    for mesh in meshes:
+		cmds.select("root", hi = True)
+		cmds.select(mesh, add = True)
+		
+		#check if the geometry that is selected is already skinned or not
+		skinClusters = cmds.ls(type = 'skinCluster')
+		skinnedGeometry = []
+		for cluster in skinClusters:
+		    geometry = cmds.skinCluster(cluster, q = True, g = True)[0]
+		    geoTransform = cmds.listRelatives(geometry, parent = True)[0]
+		    
+		    if geoTransform == mesh:
+			skin = cluster
+			noSkin = False
+			break
+		    
+		if noSkin == True:
+		    skin = cmds.skinCluster(sm = 0, tsb = True)[0]
+		    
+		
+		
+		#copy weights from proxy
+		cmds.select("skin_geo*")
+		cmds.select(mesh, add = True)
+		cmds.copySkinWeights(noMirror = True,  surfaceAssociation = "closestPoint", influenceAssociation = "closestBone")
+		
+		cmds.select(clear = True)
+		#cmds.skinCluster(skin, edit = True, sw = .5, smoothWeightsMaxIterations = 5)
+		
+	cmds.select(clear = True)
+	cmds.deleteUI("autoWeightMeshes_UI")
+	    
+	    
+
+    # # # # # # # # # # # # # # # # # # # # # # # # # # # # # # # # # # # # # # # # # # # # # # # # # # # # # # # # # # # # # # # 
+    # # # # # # # # # # # # # # # # # # # # # # # # # # # # # # # # # # # # # # # # # # # # # # # # # # # # # # # # # # # # # # # 
+    # # # # # # # # # # # # # # # # # # # # # # # # # # # # # # # # # # # # # # # # # # # # # # # # # # # # # # # # # # # # # # # 
+    # # # # # # # # # # # # # # # # # # # # # # # # # # # # # # # # # # # # # # # # # # # # # # # # # # # # # # # # # # # # # # # 
+    # # # # # # # # # # # # # # # # # # # # # # # # # # # # # # # # # # # # # # # # # # # # # # # # # # # # # # # # # # # # # # # 
+    def skinProxyGeo(self):
+	
+	#unlock joints
+	cmds.select("root", hi = True)
+	joints = cmds.ls(sl = True)
+	for joint in joints:
+	    cmds.lockNode(joint, lock = False)
+	    
+	    
+	skinGrp = cmds.group(empty = True, name = "Proxy_Geo_Skin_Grp")
+	skinGeo = []
+	
+	for mover in self.geoMovers:
+	    mover = mover.partition("|")[0]
+	    children = cmds.listRelatives(mover, children = True, type = "transform")
+	    
+	    if children != None:
+		joint = mover.partition("_geo_mover")[0]
+		suffix = mover.partition("_geo_mover")[2]
+		
+		if suffix.find("_") == 0:
+		    suffix = suffix.partition("_")[2]
+		    joint = joint + "_" + suffix
+		    
+		
+
+		
+		for child in children:
+		    #dupe the child
+		    newName = child.partition("proxy")[2]
+		    if newName:
+			newName = "skin" + newName
+		    else:
+			newName = child.partition("skeleton_mesh")[2]
+			newName = "skeleton_skin_mesh" + newName
+			
+		    dupe = cmds.duplicate(child, name = newName)[0]
+		    skinGeo.append(dupe)
+		    
+		    cmds.setAttr(dupe + ".tx", keyable = True, lock = False)
+		    cmds.setAttr(dupe + ".ty", keyable = True, lock = False)
+		    cmds.setAttr(dupe + ".tz", keyable = True, lock = False)
+		    cmds.setAttr(dupe + ".rx", keyable = True, lock = False)
+		    cmds.setAttr(dupe + ".ry", keyable = True, lock = False)
+		    cmds.setAttr(dupe + ".rz", keyable = True, lock = False)
+		    cmds.setAttr(dupe + ".sx", keyable = True, lock = False)
+		    cmds.setAttr(dupe + ".sy", keyable = True, lock = False)
+		    cmds.setAttr(dupe + ".sz", keyable = True, lock = False)
+		    
+		    
+		    
+		    parent = cmds.listRelatives(dupe, parent = True)
+		    if parent != None:
+			cmds.parent(dupe, world = True)
+		    #freeze xforms
+		    cmds.makeIdentity(dupe, t = 1, r = 1, s = 1, apply = True)
+			
+			
+		    cmds.setAttr(child + ".v", 0)
+		    
+		    
+		    try:
+			cmds.select(joint, r = True)
+			cmds.skinCluster( joint, dupe, toSelectedBones = True )
+			cmds.parent(dupe, skinGrp)
+			
+		    except:
+			pass
+			
+		    
+	#add proxyGeoLayerItems to a display layer
+	proxyGeoLayerItems = []
+	for geo in skinGeo:
+	    if geo.find("skin_geo") != -1:
+		proxyGeoLayerItems.append(geo)
+
+	cmds.select(proxyGeoLayerItems)
+	cmds.createDisplayLayer(name = "proxy_geometry_layer", nr = True)
+	cmds.setAttr("proxy_geometry_layer.enabled", 1)
+	cmds.setAttr("proxy_geometry_layer.displayType", 2)
+
+	    
+	    
+	#lock joints
+	for joint in joints:
+	    cmds.lockNode(joint, lock = True)
+	    
+	cmds.select(clear = True)
+	
+	
+	
+
+    # # # # # # # # # # # # # # # # # # # # # # # # # # # # # # # # # # # # # # # # # # # # # # # # # # # # # # # # # # # # # # # 
+    # # # # # # # # # # # # # # # # # # # # # # # # # # # # # # # # # # # # # # # # # # # # # # # # # # # # # # # # # # # # # # # 
+    # # # # # # # # # # # # # # # # # # # # # # # # # # # # # # # # # # # # # # # # # # # # # # # # # # # # # # # # # # # # # # # 
+    # # # # # # # # # # # # # # # # # # # # # # # # # # # # # # # # # # # # # # # # # # # # # # # # # # # # # # # # # # # # # # # 
+    # # # # # # # # # # # # # # # # # # # # # # # # # # # # # # # # # # # # # # # # # # # # # # # # # # # # # # # # # # # # # # # 
+    def jointMoverUI_SaveTemplateUI(self, *args):
+	
+	if not os.path.exists(self.mayaToolsDir + "/General/ART/JointMoverTemplates"):
+	    os.makedirs(self.mayaToolsDir + "/General/ART/JointMoverTemplates")
+	    
+	fileName = cmds.fileDialog2(startingDirectory = self.mayaToolsDir + "/General/ART/JointMoverTemplates/", ff = "*.txt", fm = 0, okCaption = "Save")
+	
+	if fileName != None:
+	    #check path to make sure the path is in fact, in the JointMoverTemplates directory
+	    if fileName[0].find("JointMoverTemplates") != -1:
+		
+		if fileName[0].rpartition(".")[2] != "txt":
+		    self.jointMoverUI_SaveTemplate(fileName[0] + ".txt")
+		    
+		else:
+		    self.jointMoverUI_SaveTemplate(fileName[0])
+		    
+	    else:
+		cmds.confirmDialog(icon = "warning", title = "Save Template", message = "Please Save the template in the JointMoverTemplates directory.")
+		return
+	    
+		
+
+		
+    
+    # # # # # # # # # # # # # # # # # # # # # # # # # # # # # # # # # # # # # # # # # # # # # # # # # # # # # # # # # # # # # # # 
+    # # # # # # # # # # # # # # # # # # # # # # # # # # # # # # # # # # # # # # # # # # # # # # # # # # # # # # # # # # # # # # # 
+    # # # # # # # # # # # # # # # # # # # # # # # # # # # # # # # # # # # # # # # # # # # # # # # # # # # # # # # # # # # # # # # 
+    # # # # # # # # # # # # # # # # # # # # # # # # # # # # # # # # # # # # # # # # # # # # # # # # # # # # # # # # # # # # # # # 
+    # # # # # # # # # # # # # # # # # # # # # # # # # # # # # # # # # # # # # # # # # # # # # # # # # # # # # # # # # # # # # # # 
+    def jointMoverUI_SaveTemplate(self, path):
+	
+	sourceControl = False
+	
+	#check if user has source control on
+	settingsLocation = self.mayaToolsDir + "/General/Scripts/projectSettings.txt"
+	if os.path.exists(settingsLocation):
+	    f = open(settingsLocation, 'r')
+	    settings = cPickle.load(f)
+	    f.close()	    
+	    sourceControl = settings.get("UseSourceControl")	
+
+	#turn off symmetry mode
+	if cmds.iconTextCheckBox(self.widgets["symmetry"], q = True, v = True):
+	    cmds.iconTextCheckBox(self.widgets["symmetry"], e = True, v = False)
+	    self.jointMoverUI_SymmetryModeOff()
+	    
+	    
+	fileCheckedOut = False
+	
+	try:
+	    f = open(path, 'w')
+	    
+	except:
+	    if sourceControl == False:
+		cmds.confirmDialog(title = "Save Template", icon = "critical", message = "Could not save file. File may exist already and be marked as read only.")
+		return
+	    
+	    else:
+		#check to see if the file is currently checked out
+		result = cmds.confirmDialog(title = "Save Template", icon = "critical", message = "Could not save file. File may exist already and be marked as read only.", button = ["Check Out File", "Cancel"])
+				
+		if result == "Check Out File":
+		    import perforceUtils
+		    reload(perforceUtils)
+		    perforceUtils.p4_checkOutCurrentFile(path)
+				    
+		    try:
+			f = open(path, 'w')
+			fileCheckedOut = True
+			
+		    except:
+			cmds.confirmDialog(title = "Save Template", icon = "critical", message = "Perforce operation unsucessful. Could not save file.")
+			return
+		    
+        
+        #global movers first
+        for mover in self.globalMovers:
+            mover = mover.partition("|")[0]
+            tx = cmds.getAttr(mover + ".tx")
+            ty = cmds.getAttr(mover + ".ty")
+            tz = cmds.getAttr(mover + ".tz")
+            rx = cmds.getAttr(mover + ".rx")
+            ry = cmds.getAttr(mover + ".ry")
+            rz = cmds.getAttr(mover + ".rz")
+            sx = cmds.getAttr(mover + ".sx")
+            sy = cmds.getAttr(mover + ".sy")
+            sz = cmds.getAttr(mover + ".sz")
+            
+            info = [str(mover), tx, ty, tz, rx, ry, rz, sx, sy, sz]
+            f.write(str(info) + "\n")
+            
+            
+        #offset movers
+        for mover in self.offsetMovers:
+            mover = mover.partition("|")[0]
+            tx = cmds.getAttr(mover + ".tx")
+            ty = cmds.getAttr(mover + ".ty")
+            tz = cmds.getAttr(mover + ".tz")
+            rx = cmds.getAttr(mover + ".rx")
+            ry = cmds.getAttr(mover + ".ry")
+            rz = cmds.getAttr(mover + ".rz")
+            sx = cmds.getAttr(mover + ".sx")
+            sy = cmds.getAttr(mover + ".sy")
+            sz = cmds.getAttr(mover + ".sz")
+            
+            info = [str(mover), tx, ty, tz, rx, ry, rz, sx, sy, sz]
+            f.write(str(info) + "\n")
+	    
+        #geo movers
+        for mover in self.geoMovers:
+            mover = mover.partition("|")[0]
+            tx = cmds.getAttr(mover + ".tx")
+            ty = cmds.getAttr(mover + ".ty")
+            tz = cmds.getAttr(mover + ".tz")
+            rx = cmds.getAttr(mover + ".rx")
+            ry = cmds.getAttr(mover + ".ry")
+            rz = cmds.getAttr(mover + ".rz")
+            sx = cmds.getAttr(mover + ".sx")
+            sy = cmds.getAttr(mover + ".sy")
+            sz = cmds.getAttr(mover + ".sz")
+            
+            info = [str(mover), tx, ty, tz, rx, ry, rz, sx, sy, sz]
+            f.write(str(info) + "\n")
+	    
+        #lra movers
+        for mover in self.lraControls:
+            tx = cmds.getAttr(mover + ".tx")
+            ty = cmds.getAttr(mover + ".ty")
+            tz = cmds.getAttr(mover + ".tz")
+            rx = cmds.getAttr(mover + ".rx")
+            ry = cmds.getAttr(mover + ".ry")
+            rz = cmds.getAttr(mover + ".rz")
+            sx = cmds.getAttr(mover + ".sx")
+            sy = cmds.getAttr(mover + ".sy")
+            sz = cmds.getAttr(mover + ".sz")
+            
+            info = [str(mover), tx, ty, tz, rx, ry, rz, sx, sy, sz]
+            f.write(str(info) + "\n")
+            
+        f.close()
+	
+	
+	#if user has source control on, submit to p4
+	if sourceControl:
+	    import perforceUtils
+	    reload(perforceUtils)
+	    
+	    if fileCheckedOut == True:
+		result = cmds.promptDialog(title = "Perforce", message = "Please Enter a Description..", button = ["Accept", "Cancel"], defaultButton = "Accept", dismissString = "Cancel", cancelButton = "Cancel")
+		if result == "Accept":
+		    desc = cmds.promptDialog(q = True, text = True)
+		    
+		    result = perforceUtils.p4_submitCurrentFile(path, desc)
+		    
+		    if result == False:
+			result = cmds.confirmDialog(title = "Save Template", icon = "question", message = "Would you like to add this template to Perforce?", button = ["Yes", "No"])
+			
+			if result == "Yes":
+			    import perforceUtils
+			    reload(perforceUtils)
+			    perforceUtils.p4_addAndSubmitCurrentFile(path, desc)
+	    
+	    else:
+		if sourceControl == True:
+		    result = cmds.confirmDialog(title = "Save Template", icon = "question", message = "Would you like to add this template to Perforce?", button = ["Yes", "No"])
+					    
+		    if result == "Yes":
+			result = cmds.promptDialog(title = "Perforce", message = "Please Enter a Description..", button = ["Accept", "Cancel"], defaultButton = "Accept", dismissString = "Cancel", cancelButton = "Cancel")
+			if result == "Accept":
+			    desc = cmds.promptDialog(q = True, text = True)
+			    
+			    
+			import perforceUtils
+			reload(perforceUtils)
+			perforceUtils.p4_addAndSubmitCurrentFile(path, desc)
+			
+	
+	cmds.headsUpMessage( 'Template Saved' )
+			
+	
+	
+	
+    # # # # # # # # # # # # # # # # # # # # # # # # # # # # # # # # # # # # # # # # # # # # # # # # # # # # # # # # # # # # # # # 
+    # # # # # # # # # # # # # # # # # # # # # # # # # # # # # # # # # # # # # # # # # # # # # # # # # # # # # # # # # # # # # # # 
+    # # # # # # # # # # # # # # # # # # # # # # # # # # # # # # # # # # # # # # # # # # # # # # # # # # # # # # # # # # # # # # # 
+    # # # # # # # # # # # # # # # # # # # # # # # # # # # # # # # # # # # # # # # # # # # # # # # # # # # # # # # # # # # # # # # 
+    # # # # # # # # # # # # # # # # # # # # # # # # # # # # # # # # # # # # # # # # # # # # # # # # # # # # # # # # # # # # # # # 
+    def jointMoverUI_LoadTemplateUI(self, *args):
+	
+        if cmds.window("loadTemplateUI", exists = True):
+            cmds.deleteUI("loadTemplateUI")
+            
+        window = cmds.window("loadTemplateUI", w = 300, h = 200, title = "Load Template", titleBarMenu = False, sizeable = False)
+        mainLayout = cmds.columnLayout(w = 300, h = 300)
+        menuBarLayout = cmds.menuBarLayout(h = 30, w = 300, ebg = True, bgc = [.1, .1, .1], parent = mainLayout)
+        cmds.menu(label = "File")
+        cmds.menuItem(label = "Remove Selected Template", c = self.removeSelectedTemplate)
+
+        #banner image
+	cmds.image(w = 300, h = 50, image = self.mayaToolsDir + "/General/Icons/ART/artBanner300px.bmp", parent = mainLayout)
+	
+	
+	
+        cmds.textScrollList("templateTextScrollList", w = 300, h = 150, numberOfRows = 100, allowMultiSelection = False, parent = mainLayout)
+	
+	#buttons + layout
+	cmds.separator(h = 10, parent = mainLayout)
+	buttonLayout = cmds.rowColumnLayout(parent = mainLayout, w = 300, rs = [1, 10], nc = 2, cat = [(1, "left", 20), (2, "both", 20)])
+        cmds.button(label = "Load Template", w = 120, h = 50, c = self.jointMoverUI_LoadTemplate, parent = buttonLayout)
+        cmds.button(label = "Cancel", w = 120, h = 50, c = self.exitLoadTemplateUI, parent = buttonLayout)
+        cmds.separator(h = 10, parent = mainLayout)
+	
+        #get out template files
+	if not os.path.exists(self.mayaToolsDir + "/General/ART/JointMoverTemplates"):
+		    os.makedirs(self.mayaToolsDir + "/General/ART/JointMoverTemplates")	
+
+        files = os.listdir(self.mayaToolsDir + "/General/ART/JointMoverTemplates/")
+        
+        for file in files:
+            niceName = file.rpartition(".")[0]
+            cmds.textScrollList("templateTextScrollList", edit = True, append = niceName)
+        
+        cmds.showWindow(window)
+	
+	
+    # # # # # # # # # # # # # # # # # # # # # # # # # # # # # # # # # # # # # # # # # # # # # # # # # # # # # # # # # # # # # # # 
+    # # # # # # # # # # # # # # # # # # # # # # # # # # # # # # # # # # # # # # # # # # # # # # # # # # # # # # # # # # # # # # # 
+    # # # # # # # # # # # # # # # # # # # # # # # # # # # # # # # # # # # # # # # # # # # # # # # # # # # # # # # # # # # # # # # 
+    # # # # # # # # # # # # # # # # # # # # # # # # # # # # # # # # # # # # # # # # # # # # # # # # # # # # # # # # # # # # # # # 
+    # # # # # # # # # # # # # # # # # # # # # # # # # # # # # # # # # # # # # # # # # # # # # # # # # # # # # # # # # # # # # # # 
+    def jointMoverUI_LoadTemplate(self, *args):
+	
+	
+	if cmds.iconTextCheckBox(self.widgets["symmetry"], q = True, v = True):
+	    cmds.iconTextCheckBox(self.widgets["symmetry"], e = True, v = False)
+	    self.jointMoverUI_SymmetryModeOff()
+	    
+	    
+        #get the selected template from the UI
+        if cmds.textScrollList("templateTextScrollList", q = True, exists = True):
+	    try:
+		selectedFile = cmds.textScrollList("templateTextScrollList", q = True, si = True)[0]
+		path = self.mayaToolsDir + "/General/ART/JointMoverTemplates/" + str(selectedFile) + ".txt"
+		
+	    except:
+		cmds.confirmDialog(icon = "warning", title = "Load Template", message = "Please select a template from the list.")
+		return
+	    
+	    if os.path.exists(path):
+		f = open(path, 'r')
+		lines = f.readlines()
+		
+		for line in lines:
+		    #convert the string that is a list, to a python list
+		    list = eval(line)
+		    control = list[0] 
+		    tx = list[1]
+		    ty = list[2]
+		    tz = list[3]
+		    rx = list[4]
+		    ry = list[5]
+		    rz = list[6]
+		    sx = list[7]
+		    sy = list[8]
+		    sz = list[9]
+		    
+		    if cmds.objExists(control):
+			
+			for attr in [".tx", ".ty", ".tz", ".rx", ".ry", ".rz", ".sx", ".sy", ".sz"]:
+			    if cmds.getAttr(control + attr, lock = True) == False:
+				if attr == ".tx":
+				    val = tx
+				if attr == ".ty":
+				    val = ty
+				if attr == ".tz":
+				    val = tz
+				if attr == ".rx":
+				    val = rx
+				if attr == ".ry":
+				    val = ry
+				if attr == ".rz":
+				    val = rz
+				if attr == ".sx":
+				    val = sx
+				if attr == ".sy":
+				    val = sy
+				if attr == ".sz":
+				    val = sz
+				    
+				try:
+				    cmds.setAttr(control + attr, val)
+				except:
+				    pass
+ 
+	cmds.deleteUI("loadTemplateUI")
+	
+	
+	
+    # # # # # # # # # # # # # # # # # # # # # # # # # # # # # # # # # # # # # # # # # # # # # # # # # # # # # # # # # # # # # # # 
+    # # # # # # # # # # # # # # # # # # # # # # # # # # # # # # # # # # # # # # # # # # # # # # # # # # # # # # # # # # # # # # # 
+    # # # # # # # # # # # # # # # # # # # # # # # # # # # # # # # # # # # # # # # # # # # # # # # # # # # # # # # # # # # # # # # 
+    # # # # # # # # # # # # # # # # # # # # # # # # # # # # # # # # # # # # # # # # # # # # # # # # # # # # # # # # # # # # # # # 
+    # # # # # # # # # # # # # # # # # # # # # # # # # # # # # # # # # # # # # # # # # # # # # # # # # # # # # # # # # # # # # # # 
+    def rangeOfMotion_UI(self, *args):
+	
+	#check to see if window exists, if so delete. Then create new instance
+	if cmds.window("rangeOfMotion_UI", exists = True):
+	    cmds.deleteUI("rangeOfMotion_UI")
+	    
+	self.widgets["rom_window"] = cmds.window("rangeOfMotion_UI", w = 600, h = 400, title = "Range Of Motion Generator", mxb = False, mnb = False, sizeable = False)
+	
+	#create the main layout
+	self.widgets["rom_formLayout"] = cmds.formLayout(w = 600, h = 400)
+	
+	#create the node outliner and attach to the form
+	self.widgets["rom_outliner"] = cmds.nodeOutliner(w = 250, h = 400, ms = True, sc = self.rangeOfMotion_OutlinerSelect)
+	cmds.formLayout(self.widgets["rom_formLayout"], edit = True, af = [(self.widgets["rom_outliner"], "left", 5), (self.widgets["rom_outliner"], "top", 5)])
+	
+	
+	#create the rest of the controls, like the rotation checkboxes, interval fields, generate button, etc.
+	self.widgets["rom_scrollField"] = cmds.scrollField(w = 300, h = 100, enable = False, editable=False, wordWrap=True, text= "Choose the joints from the list to generate range of motion animation for. Then select which attributes you would like to have motion generated on. Finally, choose the frame interval between keys." )
+	self.widgets["rom_rotX_cb"] = cmds.checkBox(label = "rotateX", value = True, cc = self.rangeOfMotion_OutlinerSelect)
+	self.widgets["rom_rotY_cb"] = cmds.checkBox(label = "rotateY", value = True, cc = self.rangeOfMotion_OutlinerSelect)
+	self.widgets["rom_rotZ_cb"] = cmds.checkBox(label = "rotateZ", value = True, cc = self.rangeOfMotion_OutlinerSelect)
+	self.widgets["rom_intervalField"] = cmds.intFieldGrp( numberOfFields=1, label='Interval Between Frames: ', value1 = 10, cal = (1, "left") , changeCommand = self.rangeOfMotion_OutlinerSelect)
+	self.widgets["rom_animLength"] = cmds.text(label='Current Animation Length: ', align = "left")
+	self.widgets["rom_generateButton"] = cmds.button(w = 300, h = 50, label = "Generate Range of Motion Keys", c = self.rangeOfMotion_Generate)
+	self.widgets["animLabel"] = cmds.text(label = "~Rotation Attributes to Receive Animation~", font = "obliqueLabelFont")
+	self.widgets["groupItemsCB"] = cmds.checkBox(label = "Solve as One Item", value = False, cc = self.rangeOfMotion_OutlinerSelect)
+	self.widgets["rom_startFrame"] = cmds.intFieldGrp( numberOfFields=1, label='Animation Start Frame: ', value1 = 0, cal = (1, "left"))
+
+	
+	#place widgets
+	cmds.formLayout(self.widgets["rom_formLayout"], edit = True, af = [(self.widgets["rom_scrollField"], "left", 275), (self.widgets["rom_scrollField"], "top", 5)])
+	cmds.formLayout(self.widgets["rom_formLayout"], edit = True, af = [(self.widgets["rom_rotX_cb"], "left", 275), (self.widgets["rom_rotX_cb"], "top", 150)])
+	cmds.formLayout(self.widgets["rom_formLayout"], edit = True, af = [(self.widgets["rom_rotY_cb"], "left", 385), (self.widgets["rom_rotY_cb"], "top", 150)])
+	cmds.formLayout(self.widgets["rom_formLayout"], edit = True, af = [(self.widgets["rom_rotZ_cb"], "left", 495), (self.widgets["rom_rotZ_cb"], "top", 150)])
+	cmds.formLayout(self.widgets["rom_formLayout"], edit = True, af = [(self.widgets["rom_intervalField"], "left", 275), (self.widgets["rom_intervalField"], "top", 225)])
+	cmds.formLayout(self.widgets["rom_formLayout"], edit = True, af = [(self.widgets["rom_generateButton"], "left", 275), (self.widgets["rom_generateButton"], "top", 345)])
+	cmds.formLayout(self.widgets["rom_formLayout"], edit = True, af = [(self.widgets["animLabel"], "left", 312), (self.widgets["animLabel"], "top", 120)])
+	cmds.formLayout(self.widgets["rom_formLayout"], edit = True, af = [(self.widgets["rom_animLength"], "left", 275), (self.widgets["rom_animLength"], "top", 300)])
+	cmds.formLayout(self.widgets["rom_formLayout"], edit = True, af = [(self.widgets["groupItemsCB"], "left", 275), (self.widgets["groupItemsCB"], "top", 180)])
+	cmds.formLayout(self.widgets["rom_formLayout"], edit = True, af = [(self.widgets["rom_startFrame"], "left", 275), (self.widgets["rom_startFrame"], "top", 255)])
+
+
+	
+	#show the window
+	cmds.showWindow(self.widgets["rom_window"])
+	
+	
+	#add objects to the node outliner
+	cmds.select("root", hi = True)
+	selection = cmds.ls(sl = True)
+	for obj in selection:
+	    cmds.nodeOutliner(self.widgets["rom_outliner"], edit = True, a = obj)
+	    
+	cmds.select(clear = True)
+	    
+    
+    # # # # # # # # # # # # # # # # # # # # # # # # # # # # # # # # # # # # # # # # # # # # # # # # # # # # # # # # # # # # # # # 
+    # # # # # # # # # # # # # # # # # # # # # # # # # # # # # # # # # # # # # # # # # # # # # # # # # # # # # # # # # # # # # # # 
+    # # # # # # # # # # # # # # # # # # # # # # # # # # # # # # # # # # # # # # # # # # # # # # # # # # # # # # # # # # # # # # # 
+    # # # # # # # # # # # # # # # # # # # # # # # # # # # # # # # # # # # # # # # # # # # # # # # # # # # # # # # # # # # # # # # 
+    # # # # # # # # # # # # # # # # # # # # # # # # # # # # # # # # # # # # # # # # # # # # # # # # # # # # # # # # # # # # # # # 
+    def rangeOfMotion_OutlinerSelect(self, *args):
+	
+	#get the current selection
+	animLength = 0
+	currentSelection = cmds.nodeOutliner(self.widgets["rom_outliner"], q = True, cs = True)
+	
+	if currentSelection != None:
+	    #create a variable our generate function can access that has the info we need
+	    self.romInfo = []
+	    
+	    numItems = len(currentSelection)
+	    
+	    #get the value of the checkboxes
+	    xVal = cmds.checkBox(self.widgets["rom_rotX_cb"], q = True, v = True)
+	    yVal = cmds.checkBox(self.widgets["rom_rotY_cb"], q = True, v = True)
+	    zVal = cmds.checkBox(self.widgets["rom_rotZ_cb"], q = True, v = True)
+	    
+	    #get the frame interval
+	    frameInt = cmds.intFieldGrp(self.widgets["rom_intervalField"], q = True, value1 = True)
+	    
+	    #start computing the anim length
+	    mult = 0
+	    if xVal:
+		mult += 1
+		
+	    if yVal:
+		mult += 1
+		
+	    if zVal:
+		mult += 1
+		
+	    groupSolve = cmds.checkBox(self.widgets["groupItemsCB"], q = True, value = True)
+	    if groupSolve:
+		numItems = 1
+		
+	    animLength += (((mult * 3) * numItems) * frameInt)
+	    
+	    #update current anim length field
+	    cmds.text(self.widgets["rom_animLength"], edit = True, label = 'Current Animation Length: ' + str(animLength) + " frames")
+	    
+	    #fill our list
+	    self.romInfo.append([currentSelection, xVal, yVal, zVal, frameInt, animLength])
+    
+   
+   
+    # # # # # # # # # # # # # # # # # # # # # # # # # # # # # # # # # # # # # # # # # # # # # # # # # # # # # # # # # # # # # # # 
+    # # # # # # # # # # # # # # # # # # # # # # # # # # # # # # # # # # # # # # # # # # # # # # # # # # # # # # # # # # # # # # # 
+    # # # # # # # # # # # # # # # # # # # # # # # # # # # # # # # # # # # # # # # # # # # # # # # # # # # # # # # # # # # # # # # 
+    # # # # # # # # # # # # # # # # # # # # # # # # # # # # # # # # # # # # # # # # # # # # # # # # # # # # # # # # # # # # # # # 
+    # # # # # # # # # # # # # # # # # # # # # # # # # # # # # # # # # # # # # # # # # # # # # # # # # # # # # # # # # # # # # # # 
+    def rangeOfMotion_Clear(self, *args): 
+	cmds.currentTime(0)
+	#wipe existing keys
+	cmds.select("root", hi = True)
+	selection = cmds.ls(sl = True)
+	cmds.cutKey( selection, time=(-10000, 10000), option="keys" )
+	cmds.select(clear = True)
+	
+	
+	
+    # # # # # # # # # # # # # # # # # # # # # # # # # # # # # # # # # # # # # # # # # # # # # # # # # # # # # # # # # # # # # # # 
+    # # # # # # # # # # # # # # # # # # # # # # # # # # # # # # # # # # # # # # # # # # # # # # # # # # # # # # # # # # # # # # # 
+    # # # # # # # # # # # # # # # # # # # # # # # # # # # # # # # # # # # # # # # # # # # # # # # # # # # # # # # # # # # # # # # 
+    # # # # # # # # # # # # # # # # # # # # # # # # # # # # # # # # # # # # # # # # # # # # # # # # # # # # # # # # # # # # # # # 
+    # # # # # # # # # # # # # # # # # # # # # # # # # # # # # # # # # # # # # # # # # # # # # # # # # # # # # # # # # # # # # # # 
+    def rangeOfMotion_Generate(self, *args):
+	#simple tool that will generate a range of motion animation
+	try:
+	    #breakdown info
+	    info = self.romInfo[0]
+	    itemsToKey = info[0]
+	    xVal = info[1]
+	    yVal = info[2]
+	    zVal = info[3]
+	    frameInt = info[4]
+	    animLength = info[5]
+	    
+	except:
+	    cmds.confirmDialog(icon = "warning", title = "Range Of Motion", message = "You must select at least one joint from the list!")
+	    return
+	
+	
+
+	#set the timeline to the anim length 0 - animLength
+	startFrame = cmds.intFieldGrp(self.widgets["rom_startFrame"], q = True, v = True)[0]
+	cmds.playbackOptions(min = 0, max = startFrame + animLength)
+	
+	#wipe existing keys
+	cmds.cutKey( itemsToKey, time=(-100,animLength * 10), option="keys" )
+	
+	#for items in selected, set keyframes
+
+	
+	cmds.setKeyframe(itemsToKey, t = startFrame)
+	time = startFrame
+	
+	#find out if group solve is checked
+	groupSolve = cmds.checkBox(self.widgets["groupItemsCB"], q = True, value = True)
+	if groupSolve:
+	    
+	    
+	    if xVal:
+		cmds.currentTime(time)
+		for item in itemsToKey:
+		    cmds.setAttr(item + ".rotateX", 0)
+		    cmds.setAttr(item + ".rotateY", 0)
+		    cmds.setAttr(item + ".rotateZ", 0)
+		    cmds.setKeyframe(item + ".rotate")
+		    
+		time = time + frameInt
+		cmds.currentTime(time)
+		
+		for item in itemsToKey:
+		    cmds.setAttr(item + ".rotateX", 30)
+		    cmds.setAttr(item + ".rotateY", 0)
+		    cmds.setAttr(item + ".rotateZ", 0)
+		    cmds.setKeyframe(item + ".rotate")
+		
+		
+		time = time + frameInt
+		cmds.currentTime(time)
+		
+		for item in itemsToKey:
+		    cmds.setAttr(item + ".rotateX", -30)
+		    cmds.setAttr(item + ".rotateY", 0)
+		    cmds.setAttr(item + ".rotateZ", 0)
+		    cmds.setKeyframe(item + ".rotate")
+		
+		
+		time = time + frameInt
+		cmds.currentTime(time)
+		
+		for item in itemsToKey:
+		    cmds.setAttr(item + ".rotateX", 0)
+		    cmds.setAttr(item + ".rotateY", 0)
+		    cmds.setAttr(item + ".rotateZ", 0)
+		    cmds.setKeyframe(item + ".rotate")
+		    
+
+	    if yVal:
+		
+		cmds.currentTime(time)
+		for item in itemsToKey:
+		    cmds.setAttr(item + ".rotateX", 0)
+		    cmds.setAttr(item + ".rotateY", 0)
+		    cmds.setAttr(item + ".rotateZ", 0)
+		    cmds.setKeyframe(item + ".rotate")
+		    
+		time = time + frameInt
+		cmds.currentTime(time)
+		
+		for item in itemsToKey:
+		    cmds.setAttr(item + ".rotateX", 0)
+		    cmds.setAttr(item + ".rotateY", 30)
+		    cmds.setAttr(item + ".rotateZ", 0)
+		    cmds.setKeyframe(item + ".rotate")
+		
+		
+		time = time + frameInt
+		cmds.currentTime(time)
+		
+		for item in itemsToKey:
+		    cmds.setAttr(item + ".rotateX", 0)
+		    cmds.setAttr(item + ".rotateY", -30)
+		    cmds.setAttr(item + ".rotateZ", 0)
+		    cmds.setKeyframe(item + ".rotate")
+		
+		
+		time = time + frameInt
+		cmds.currentTime(time)
+		
+		for item in itemsToKey:
+		    cmds.setAttr(item + ".rotateX", 0)
+		    cmds.setAttr(item + ".rotateY", 0)
+		    cmds.setAttr(item + ".rotateZ", 0)
+		    cmds.setKeyframe(item + ".rotate")
+		    
+
+	    if zVal:
+		
+		cmds.currentTime(time)
+		for item in itemsToKey:
+		    cmds.setAttr(item + ".rotateX", 0)
+		    cmds.setAttr(item + ".rotateY", 0)
+		    cmds.setAttr(item + ".rotateZ", 0)
+		    cmds.setKeyframe(item + ".rotate")
+		    
+		time = time + frameInt
+		cmds.currentTime(time)
+		
+		for item in itemsToKey:
+		    cmds.setAttr(item + ".rotateX", 0)
+		    cmds.setAttr(item + ".rotateY", 0)
+		    cmds.setAttr(item + ".rotateZ", 30)
+		    cmds.setKeyframe(item + ".rotate")
+		
+		
+		time = time + frameInt
+		cmds.currentTime(time)
+		
+		for item in itemsToKey:
+		    cmds.setAttr(item + ".rotateX", 0)
+		    cmds.setAttr(item + ".rotateY", 0)
+		    cmds.setAttr(item + ".rotateZ", -30)
+		    cmds.setKeyframe(item + ".rotate")
+		
+		
+		time = time + frameInt
+		cmds.currentTime(time)
+		
+		for item in itemsToKey:
+		    cmds.setAttr(item + ".rotateX", 0)
+		    cmds.setAttr(item + ".rotateY", 0)
+		    cmds.setAttr(item + ".rotateZ", 0)
+		    cmds.setKeyframe(item + ".rotate")
+		    
+		
+
+
+	
+	else:
+	    if xVal:
+		for item in itemsToKey:
+		    
+		    cmds.currentTime(time)
+		    cmds.setAttr(item + ".rotateX", 0)
+		    cmds.setAttr(item + ".rotateY", 0)
+		    cmds.setAttr(item + ".rotateZ", 0)
+		    cmds.setKeyframe(item + ".rotate")
+		    time = time + frameInt
+		    
+		    cmds.currentTime(time)
+		    cmds.setAttr(item + ".rotateX", 90)
+		    cmds.setAttr(item + ".rotateY", 0)
+		    cmds.setAttr(item + ".rotateZ", 0)
+		    cmds.setKeyframe(item + ".rotate")
+		    time = time + frameInt
+		    
+		    cmds.currentTime(time)
+		    cmds.setAttr(item + ".rotateX", -90)
+		    cmds.setAttr(item + ".rotateY", 0)
+		    cmds.setAttr(item + ".rotateZ", 0)
+		    cmds.setKeyframe(item + ".rotate")
+		    time = time + frameInt
+		    
+		    cmds.currentTime(time)
+		    cmds.setAttr(item + ".rotateX", 0)
+		    cmds.setAttr(item + ".rotateY", 0)
+		    cmds.setAttr(item + ".rotateZ", 0)
+		    cmds.setKeyframe(item + ".rotate")
+		    
+	    if yVal:
+		for item in itemsToKey:
+		    cmds.currentTime(time)
+		    cmds.setAttr(item + ".rotateX", 0)
+		    cmds.setAttr(item + ".rotateY", 0)
+		    cmds.setAttr(item + ".rotateZ", 0)
+		    cmds.setKeyframe(item + ".rotate")
+		    time = time + frameInt
+		    
+		    cmds.currentTime(time)
+		    cmds.setAttr(item + ".rotateX", 0)
+		    cmds.setAttr(item + ".rotateY", 90)
+		    cmds.setAttr(item + ".rotateZ", 0)
+		    cmds.setKeyframe(item + ".rotate")
+		    time = time + frameInt
+		    
+		    cmds.currentTime(time)
+		    cmds.setAttr(item + ".rotateX", 0)
+		    cmds.setAttr(item + ".rotateY", -90)
+		    cmds.setAttr(item + ".rotateZ", 0)
+		    cmds.setKeyframe(item + ".rotate")
+		    time = time + frameInt
+		    
+		    cmds.currentTime(time)
+		    cmds.setAttr(item + ".rotateX", 0)
+		    cmds.setAttr(item + ".rotateY", 0)
+		    cmds.setAttr(item + ".rotateZ", 0)
+		    cmds.setKeyframe(item + ".rotate")
+		
+    
+		
+	    if zVal:
+		for item in itemsToKey:
+		    cmds.currentTime(time)
+		    cmds.setAttr(item + ".rotateX", 0)
+		    cmds.setAttr(item + ".rotateY", 0)
+		    cmds.setAttr(item + ".rotateZ", 0)
+		    cmds.setKeyframe(item + ".rotate")
+		    time = time + frameInt
+		    
+		    cmds.currentTime(time)
+		    cmds.setAttr(item + ".rotateX", 0)
+		    cmds.setAttr(item + ".rotateY", 0)
+		    cmds.setAttr(item + ".rotateZ", 90)
+		    cmds.setKeyframe(item + ".rotate")
+		    time = time + frameInt
+		    
+		    cmds.currentTime(time)
+		    cmds.setAttr(item + ".rotateX", 0)
+		    cmds.setAttr(item + ".rotateY", 0)
+		    cmds.setAttr(item + ".rotateZ", -90)
+		    cmds.setKeyframe(item + ".rotate")
+		    time = time + frameInt
+		    
+		    cmds.currentTime(time)
+		    cmds.setAttr(item + ".rotateX", 0)
+		    cmds.setAttr(item + ".rotateY", 0)
+		    cmds.setAttr(item + ".rotateZ", 0)
+		    cmds.setKeyframe(item + ".rotate")
+		    
+		    
+
+
+    # # # # # # # # # # # # # # # # # # # # # # # # # # # # # # # # # # # # # # # # # # # # # # # # # # # # # # # # # # # # # # # 
+    # # # # # # # # # # # # # # # # # # # # # # # # # # # # # # # # # # # # # # # # # # # # # # # # # # # # # # # # # # # # # # # 
+    # # # # # # # # # # # # # # # # # # # # # # # # # # # # # # # # # # # # # # # # # # # # # # # # # # # # # # # # # # # # # # # 
+    # # # # # # # # # # # # # # # # # # # # # # # # # # # # # # # # # # # # # # # # # # # # # # # # # # # # # # # # # # # # # # # 
+    # # # # # # # # # # # # # # # # # # # # # # # # # # # # # # # # # # # # # # # # # # # # # # # # # # # # # # # # # # # # # # # 
+    def paintWeights(self, *args):
+	#writing our own paint weights tool so we can dock it in the UI
+	
+	currentSelection = cmds.ls(sl = True, flatten = True)
+	#check to see if the user actually selected something
+	if cmds.ls(sl = True) == []:
+	    cmds.warning("You must select some geometry to use the paint skin weights operation.")
+	    return
+	
+	
+	#first, clear out the influencelist tree view
+	cmds.treeView(self.widgets["paint_influencesList_list"], edit = True, removeAll = True)
+	self.skinCluster = ""
+	imagePath = self.mayaToolsDir + "/General/Icons/ART/"
+
+	#if our context tool does not exist, create it. Else, switch to it
+	if cmds.artAttrSkinPaintCtx("ART_artAttrSkinPaint", exists = True) == False:
+	    self.paintWeightsTool = cmds.artAttrSkinPaintCtx("ART_artAttrSkinPaint", ch = True )
+	
+	
+	#setup some defaults when launching mode
+	cmds.setToolTo('ART_artAttrSkinPaint')
+	cmds.artAttrSkinPaintCtx("ART_artAttrSkinPaint", edit = True, radius = 5.0, value = 1.0, opacity = 1.0, sao= "additive")
+	cmds.radioButton(self.widgets["paint_artAttrOper_radioButton2"], edit = True, sl = True)
+	self.paintWeightsPaintMode()
+	cmds.optionMenu(self.widgets["paint_skinTools_paintToolMode"], e = True, sl = 1)
+	cmds.floatSliderGrp("paint_opacitySlider", edit = True, v = 1.0)
+	cmds.floatSliderGrp("paint_attrValueSlider", edit = True, v = 1.0)
+	cmds.optionMenu(self.widgets["paint_lSortLayout_sortMethod"], q = True, select = 1)
+	
+	#set previous values if they exist
+	if self.paintWeightSettings !=  []:
+	    
+	    cmds.radioButton(self.widgets[str(self.paintWeightSettings[0])], edit = True, sl = True)
+	    cmds.artAttrSkinPaintCtx("ART_artAttrSkinPaint", edit = True, radius = self.paintWeightSettings[1])
+	    cmds.artAttrSkinPaintCtx("ART_artAttrSkinPaint", edit = True, value = self.paintWeightSettings[2])
+	    cmds.artAttrSkinPaintCtx("ART_artAttrSkinPaint", edit = True, opacity = self.paintWeightSettings[3])
+	    cmds.artAttrSkinPaintCtx("ART_artAttrSkinPaint", edit = True, sao = self.paintWeightSettings[4])
+	    cmds.floatSliderGrp("paint_opacitySlider", edit = True, v = self.paintWeightSettings[3])
+	    cmds.floatSliderGrp("paint_attrValueSlider", edit = True, v = self.paintWeightSettings[2])	    
+
+
+
+	#build influence list
+	selected = cmds.ls(sl = True)[0]
+	
+	if selected.find(".vtx") != -1:
+	    selected = selected.partition(".vtx")[0]
+	skinClusters = cmds.ls(type = 'skinCluster')
+	self.influences = []
+	
+	#go through each found skin cluster, and if we find a skin cluster whose geometry matches our selection, get influences
+	for cluster in skinClusters:
+	    geometry = cmds.skinCluster(cluster, q = True, g = True)[0]
+	    geoTransform = cmds.listRelatives(geometry, parent = True)[0]
+	    if geoTransform == selected:
+		self.skinCluster = cluster
+		self.influences = cmds.skinCluster(cluster, q= True, inf = True)
+		
+
+	#make a list of all of our joints
+	cmds.select("root", hi = True)
+	self.allJoints = cmds.ls(sl = True)
+	cmds.select(clear =True)
+	cmds.select(selected)
+
+	#add items to UI
+	for joint in self.allJoints:
+	    if joint == "root":
+		parent = ""
+	    else:
+		parent = cmds.listRelatives(joint, parent = True)[0]
+
+		
+	    cmds.treeView(self.widgets["paint_influencesList_list"], edit = True, addItem = (joint,parent))
+	    cmds.treeView(self.widgets["paint_influencesList_list"], edit = True, buttonStyle = (joint, 1, "2StateButton"))
+	    cmds.treeView(self.widgets["paint_influencesList_list"], edit = True, pressCommand = (1, self.paintWeights_ToggleLockState))
+	    cmds.treeView(self.widgets["paint_influencesList_list"], edit = True, image = (joint, 1, imagePath + "LockOFF.bmp"))
+
+	    
+	#if a joint is not an influence in the skinCluster, hide that item in the tree
+	
+	for joint in self.allJoints:
+	    if joint not in self.influences:
+		cmds.treeView(self.widgets["paint_influencesList_list"], edit = True, itemVisible = (joint, False))
+		
+
+	#perform hotkey check/setup
+	self.paintWeightsModeHotkeyChange()
+	
+	#create scriptJob
+	self.paintWeightsModeScriptJob()
+	
+	#find any locked influences and reflect that in the UI
+	for item in self.influences:
+	    if cmds.skinCluster(self.skinCluster, q = True, inf = item, lw = True):
+		cmds.treeView(self.widgets["paint_influencesList_list"], edit = True, lbc = [item, .5, .5, .5], cs = True)
+		cmds.treeView(self.widgets["paint_influencesList_list"], edit = True, image = (item, 1, imagePath + "LockON.bmp"), buttonState = [item, 1, "buttonDown"])
+		
+	#enable add and remove influence buttons
+	cmds.symbolButton(self.widgets["addInfluenceButton"], edit = True, enable = True)
+	cmds.symbolButton(self.widgets["removeInfluenceButton"], edit = True, enable = True)
+	
+	#select the first item
+	try:
+	    visibleItems = cmds.treeView(self.widgets["paint_influencesList_list"], q = True, itemVisible = True)
+	    cmds.treeView(self.widgets["paint_influencesList_list"], edit = True, selectItem = (visibleItems[0], True))
+	    self.updatedPaintInfluenceSelection()
+	    
+	except:
+	    cmds.confirmDialog(icon = "warning", title = "Paint Skin Weights", message = "Mesh has no skinCluster assigned.")
+	
+
+	#reselect selection
+	cmds.select(currentSelection)
+
+    # # # # # # # # # # # # # # # # # # # # # # # # # # # # # # # # # # # # # # # # # # # # # # # # # # # # # # # # # # # # # # # 
+    # # # # # # # # # # # # # # # # # # # # # # # # # # # # # # # # # # # # # # # # # # # # # # # # # # # # # # # # # # # # # # # 
+    # # # # # # # # # # # # # # # # # # # # # # # # # # # # # # # # # # # # # # # # # # # # # # # # # # # # # # # # # # # # # # # 
+    # # # # # # # # # # # # # # # # # # # # # # # # # # # # # # # # # # # # # # # # # # # # # # # # # # # # # # # # # # # # # # # 
+    # # # # # # # # # # # # # # # # # # # # # # # # # # # # # # # # # # # # # # # # # # # # # # # # # # # # # # # # # # # # # # # 
+    def paintWeights_ToggleLockState(self, *args):
+	
+	imagePath = self.mayaToolsDir + "/General/Icons/ART/"
+	
+	#get current button state
+	selectedInf = args[0]
+	cmds.treeView(self.widgets["paint_influencesList_list"], edit = True, clearSelection = True)
+	cmds.treeView(self.widgets["paint_influencesList_list"], edit = True, selectItem = [selectedInf, True])
+	
+	currentState = args[1]
+		
+	if int(currentState) == 1:
+	    #turn on lock weights
+	    if selectedInf in self.influences:
+		cmds.skinCluster(self.skinCluster, edit = True, inf = selectedInf, lw = True)
+		cmds.treeView(self.widgets["paint_influencesList_list"], edit = True, lbc = [selectedInf, .5, .5, .5], cs = True)
+		cmds.treeView(self.widgets["paint_influencesList_list"], edit = True, image = (selectedInf, 1, imagePath + "LockON.bmp"))
+	    
+	else:
+	    #turn off lock weights
+	    if selectedInf in self.influences:
+		cmds.skinCluster(self.skinCluster, edit = True, inf = selectedInf, lw = False)
+		cmds.treeView(self.widgets["paint_influencesList_list"], edit = True, lbc = [selectedInf, -1, -1, -1,], cs = True)
+		cmds.treeView(self.widgets["paint_influencesList_list"], edit = True, image = (selectedInf, 1, imagePath + "LockOFF.bmp"))
+
+	
+	
+    # # # # # # # # # # # # # # # # # # # # # # # # # # # # # # # # # # # # # # # # # # # # # # # # # # # # # # # # # # # # # # # 
+    # # # # # # # # # # # # # # # # # # # # # # # # # # # # # # # # # # # # # # # # # # # # # # # # # # # # # # # # # # # # # # # 
+    # # # # # # # # # # # # # # # # # # # # # # # # # # # # # # # # # # # # # # # # # # # # # # # # # # # # # # # # # # # # # # # 
+    # # # # # # # # # # # # # # # # # # # # # # # # # # # # # # # # # # # # # # # # # # # # # # # # # # # # # # # # # # # # # # # 
+    # # # # # # # # # # # # # # # # # # # # # # # # # # # # # # # # # # # # # # # # # # # # # # # # # # # # # # # # # # # # # # # 
+    def paintWeightsModeScriptJob(self):
+	
+	#create a script job when entering paint weights mode that will look for user changing tool, so we can call exitPaintWeightsMode
+	scriptJobNum = cmds.scriptJob(event = ["ToolChanged", self.exitPaintWeightsMode], runOnce = True)
+	
+	
+    # # # # # # # # # # # # # # # # # # # # # # # # # # # # # # # # # # # # # # # # # # # # # # # # # # # # # # # # # # # # # # # 
+    # # # # # # # # # # # # # # # # # # # # # # # # # # # # # # # # # # # # # # # # # # # # # # # # # # # # # # # # # # # # # # # 
+    # # # # # # # # # # # # # # # # # # # # # # # # # # # # # # # # # # # # # # # # # # # # # # # # # # # # # # # # # # # # # # # 
+    # # # # # # # # # # # # # # # # # # # # # # # # # # # # # # # # # # # # # # # # # # # # # # # # # # # # # # # # # # # # # # # 
+    # # # # # # # # # # # # # # # # # # # # # # # # # # # # # # # # # # # # # # # # # # # # # # # # # # # # # # # # # # # # # # # 
+    def exitPaintWeightsMode(self, *args):
+	
+	#store values user has set
+	#paintWeightSettings = mode, radius, value, opacity, sao
+	mode = "paint_artAttrOper_radioButton2"
+	sao = "additive"
+	
+	if cmds.radioButton(self.widgets["paint_artAttrOper_radioButton1"], q = True, sl = True):
+	    mode = "paint_artAttrOper_radioButton1"
+	    sao = "absolute"
+	if cmds.radioButton(self.widgets["paint_artAttrOper_radioButton2"], q = True, sl = True):
+	    mode = "paint_artAttrOper_radioButton2"
+	    sao = "additive"
+	if cmds.radioButton(self.widgets["paint_artAttrOper_radioButton3"], q = True, sl = True):
+	    mode = "paint_artAttrOper_radioButton3"
+	    sao = "scale"
+	if cmds.radioButton(self.widgets["paint_artAttrOper_radioButton4"], q = True, sl = True):
+	    mode = "paint_artAttrOper_radioButton4"
+	    sao = "smooth"
+	    
+	radius = cmds.artAttrSkinPaintCtx("ART_artAttrSkinPaint", q = True, radius = True)
+	value = cmds.artAttrSkinPaintCtx("ART_artAttrSkinPaint", q = True, value = True)
+	opacity = cmds.artAttrSkinPaintCtx("ART_artAttrSkinPaint", q = True, opacity = True)
+	
+	self.paintWeightSettings = [mode, radius, value, opacity, sao]
+
+	
+	#reset hotkeys to original values
+	cmds.hotkey( keyShortcut='n',name= self.nKeyPress )
+	cmds.hotkey( keyShortcut='n', releaseName= self.nKeyRelease)
+	
+	#clear influence list
+	cmds.treeView(self.widgets["paint_influencesList_list"], edit = True, removeAll = True)
+	
+	#disable add and remove infs buttons
+	cmds.symbolButton(self.widgets["addInfluenceButton"], edit = True, enable = False)
+	cmds.symbolButton(self.widgets["removeInfluenceButton"], edit = True, enable = False)
+	
+    # # # # # # # # # # # # # # # # # # # # # # # # # # # # # # # # # # # # # # # # # # # # # # # # # # # # # # # # # # # # # # # 
+    # # # # # # # # # # # # # # # # # # # # # # # # # # # # # # # # # # # # # # # # # # # # # # # # # # # # # # # # # # # # # # # 
+    # # # # # # # # # # # # # # # # # # # # # # # # # # # # # # # # # # # # # # # # # # # # # # # # # # # # # # # # # # # # # # # 
+    # # # # # # # # # # # # # # # # # # # # # # # # # # # # # # # # # # # # # # # # # # # # # # # # # # # # # # # # # # # # # # # 
+    # # # # # # # # # # # # # # # # # # # # # # # # # # # # # # # # # # # # # # # # # # # # # # # # # # # # # # # # # # # # # # # 
+    def updatedPaintInfluenceSelection(self, *args):
+	#this function notes a selection change in our tree view, and updates the artisan tool to now be working with the new influence
+	
+	inf = cmds.treeView(self.widgets["paint_influencesList_list"], q = True, si = True)[0]
+	cmds.artAttrSkinPaintCtx("ART_artAttrSkinPaint", edit = True, influence = inf)
+	string = "source artAttrSkinJointMenu;artSkinSelectInfluence(\"artAttrSkinPaintCtx\", \"" + inf + "\");"
+	mel.eval(string)
+	
+	
+    # # # # # # # # # # # # # # # # # # # # # # # # # # # # # # # # # # # # # # # # # # # # # # # # # # # # # # # # # # # # # # # 
+    # # # # # # # # # # # # # # # # # # # # # # # # # # # # # # # # # # # # # # # # # # # # # # # # # # # # # # # # # # # # # # # 
+    # # # # # # # # # # # # # # # # # # # # # # # # # # # # # # # # # # # # # # # # # # # # # # # # # # # # # # # # # # # # # # # 
+    # # # # # # # # # # # # # # # # # # # # # # # # # # # # # # # # # # # # # # # # # # # # # # # # # # # # # # # # # # # # # # # 
+    # # # # # # # # # # # # # # # # # # # # # # # # # # # # # # # # # # # # # # # # # # # # # # # # # # # # # # # # # # # # # # # 
+    def changePaintOperation(self, operation, *args):
+	#changes paint operation from the radio buttons(add, replace, scale, smooth)
+	
+	if cmds.radioButton(self.widgets["paint_artAttrOper_radioButton1"], q = True, sl = True):
+	    cmds.artAttrSkinPaintCtx('ART_artAttrSkinPaint', edit = True, sao= "absolute")
+	    
+	if cmds.radioButton(self.widgets["paint_artAttrOper_radioButton2"], q = True, sl = True):
+	    cmds.artAttrSkinPaintCtx('ART_artAttrSkinPaint', edit = True, sao= "additive")
+	    
+	if cmds.radioButton(self.widgets["paint_artAttrOper_radioButton3"], q = True, sl = True):
+	    cmds.artAttrSkinPaintCtx('ART_artAttrSkinPaint', edit = True, sao= "scale")
+	    
+	if cmds.radioButton(self.widgets["paint_artAttrOper_radioButton4"], q = True, sl = True):
+	    cmds.artAttrSkinPaintCtx('ART_artAttrSkinPaint', edit = True, sao= "smooth")
+	    
+	    
+    # # # # # # # # # # # # # # # # # # # # # # # # # # # # # # # # # # # # # # # # # # # # # # # # # # # # # # # # # # # # # # # 
+    # # # # # # # # # # # # # # # # # # # # # # # # # # # # # # # # # # # # # # # # # # # # # # # # # # # # # # # # # # # # # # # 
+    # # # # # # # # # # # # # # # # # # # # # # # # # # # # # # # # # # # # # # # # # # # # # # # # # # # # # # # # # # # # # # # 
+    # # # # # # # # # # # # # # # # # # # # # # # # # # # # # # # # # # # # # # # # # # # # # # # # # # # # # # # # # # # # # # # 
+    # # # # # # # # # # # # # # # # # # # # # # # # # # # # # # # # # # # # # # # # # # # # # # # # # # # # # # # # # # # # # # # 
+    def changeBrushStamp(self, stamp, *args):
+	#changes the brush stamp(4 buttons at top)
+	cmds.symbolCheckBox(self.widgets["paint_spGaussianChBx"], edit = True, v = 0)
+	cmds.symbolCheckBox(self.widgets["paint_spPolyBrushChBx"], edit = True, v = 0)
+	cmds.symbolCheckBox(self.widgets["paint_spSolidChBx"], edit = True, v = 0)
+	cmds.symbolCheckBox(self.widgets["paint_spRectBrushChBx"], edit = True, v = 0)
+	
+	
+	if stamp == "gaussian":
+	    cmds.symbolCheckBox(self.widgets["paint_spGaussianChBx"], edit = True, v = 1)
+	    
+	if stamp == "poly":
+	    cmds.symbolCheckBox(self.widgets["paint_spPolyBrushChBx"], edit = True, v = 1)
+	    
+	if stamp == "solid":
+	    cmds.symbolCheckBox(self.widgets["paint_spSolidChBx"], edit = True, v = 1)
+	    
+	if stamp == "square":
+	    cmds.symbolCheckBox(self.widgets["paint_spRectBrushChBx"], edit = True, v = 1)
+
+	cmds.artAttrSkinPaintCtx('ART_artAttrSkinPaint', edit = True, stampProfile= stamp)
+	
+	
+    # # # # # # # # # # # # # # # # # # # # # # # # # # # # # # # # # # # # # # # # # # # # # # # # # # # # # # # # # # # # # # # 
+    # # # # # # # # # # # # # # # # # # # # # # # # # # # # # # # # # # # # # # # # # # # # # # # # # # # # # # # # # # # # # # # 
+    # # # # # # # # # # # # # # # # # # # # # # # # # # # # # # # # # # # # # # # # # # # # # # # # # # # # # # # # # # # # # # # 
+    # # # # # # # # # # # # # # # # # # # # # # # # # # # # # # # # # # # # # # # # # # # # # # # # # # # # # # # # # # # # # # # 
+    # # # # # # # # # # # # # # # # # # # # # # # # # # # # # # # # # # # # # # # # # # # # # # # # # # # # # # # # # # # # # # # 
+    def changeBrushOpac(self, *args):
+    
+	value = cmds.floatSliderGrp("paint_opacitySlider", q = True, value = True)
+	cmds.artAttrSkinPaintCtx('ART_artAttrSkinPaint', edit = True, opacity= value)
+	
+	
+    # # # # # # # # # # # # # # # # # # # # # # # # # # # # # # # # # # # # # # # # # # # # # # # # # # # # # # # # # # # # # # # 
+    # # # # # # # # # # # # # # # # # # # # # # # # # # # # # # # # # # # # # # # # # # # # # # # # # # # # # # # # # # # # # # # 
+    # # # # # # # # # # # # # # # # # # # # # # # # # # # # # # # # # # # # # # # # # # # # # # # # # # # # # # # # # # # # # # # 
+    # # # # # # # # # # # # # # # # # # # # # # # # # # # # # # # # # # # # # # # # # # # # # # # # # # # # # # # # # # # # # # # 
+    # # # # # # # # # # # # # # # # # # # # # # # # # # # # # # # # # # # # # # # # # # # # # # # # # # # # # # # # # # # # # # # 
+    def changeBrushValue(self, *args):
+	
+	value = cmds.floatSliderGrp("paint_attrValueSlider", q = True, value = True)
+	cmds.artAttrSkinPaintCtx('ART_artAttrSkinPaint', edit = True, value= value)
+
+
+    # # # # # # # # # # # # # # # # # # # # # # # # # # # # # # # # # # # # # # # # # # # # # # # # # # # # # # # # # # # # # # # 
+    # # # # # # # # # # # # # # # # # # # # # # # # # # # # # # # # # # # # # # # # # # # # # # # # # # # # # # # # # # # # # # # 
+    # # # # # # # # # # # # # # # # # # # # # # # # # # # # # # # # # # # # # # # # # # # # # # # # # # # # # # # # # # # # # # # 
+    # # # # # # # # # # # # # # # # # # # # # # # # # # # # # # # # # # # # # # # # # # # # # # # # # # # # # # # # # # # # # # # 
+    # # # # # # # # # # # # # # # # # # # # # # # # # # # # # # # # # # # # # # # # # # # # # # # # # # # # # # # # # # # # # # # 
+    def paintFloodWeights(self, *args):
+	
+	value = cmds.floatSliderGrp("paint_attrValueSlider", q = True, value = True)
+	currentMode = cmds.artAttrSkinPaintCtx('ART_artAttrSkinPaint', q = True, sao = True)
+	if currentMode != "additive":
+	    cmds.artAttrSkinPaintCtx('ART_artAttrSkinPaint', edit = True, clear = True)
+	    
+	else:
+	    cmds.artAttrSkinPaintCtx('ART_artAttrSkinPaint', edit = True, clear = value)
+	
+	
+    # # # # # # # # # # # # # # # # # # # # # # # # # # # # # # # # # # # # # # # # # # # # # # # # # # # # # # # # # # # # # # # 
+    # # # # # # # # # # # # # # # # # # # # # # # # # # # # # # # # # # # # # # # # # # # # # # # # # # # # # # # # # # # # # # # 
+    # # # # # # # # # # # # # # # # # # # # # # # # # # # # # # # # # # # # # # # # # # # # # # # # # # # # # # # # # # # # # # # 
+    # # # # # # # # # # # # # # # # # # # # # # # # # # # # # # # # # # # # # # # # # # # # # # # # # # # # # # # # # # # # # # # 
+    # # # # # # # # # # # # # # # # # # # # # # # # # # # # # # # # # # # # # # # # # # # # # # # # # # # # # # # # # # # # # # # 
+    def paintWeightsModeHotkeyChange(self):
+	
+	#get original hotkey commands and store them
+	self.nKeyPress = cmds.hotkey( 'n', query=True, name = True )
+	self.bKeyPress = cmds.hotkey( 'b', query=True, name = True )
+	self.nKeyRelease = cmds.hotkey( 'n', query=True, releaseName = True )
+	self.bKeyRelease = cmds.hotkey( 'b', query=True, releaseName = True )
+	
+	
+	#set new commands
+	newNPressKey = cmds.nameCommand(ann = "ART Paint Weights ModifyValuePress", sourceType = "mel", c= "artAttrSkinPaintCtx -e -dragSlider value ART_artAttrSkinPaint;")
+	cmds.hotkey( keyShortcut='n',name= newNPressKey )
+	
+	newNReleaseKey = cmds.nameCommand(ann = "ART Paint Weights ModifyValueRelease", sourceType = "mel", c= "artAttrSkinPaintCtx -e -dragSlider none ART_artAttrSkinPaint;floatSliderGrp -e -v `artAttrSkinPaintCtx -q -value ART_artAttrSkinPaint` paint_attrValueSlider;")
+	cmds.hotkey( keyShortcut='n', releaseName= newNReleaseKey)
+	
+	
+    # # # # # # # # # # # # # # # # # # # # # # # # # # # # # # # # # # # # # # # # # # # # # # # # # # # # # # # # # # # # # # # 
+    # # # # # # # # # # # # # # # # # # # # # # # # # # # # # # # # # # # # # # # # # # # # # # # # # # # # # # # # # # # # # # # 
+    # # # # # # # # # # # # # # # # # # # # # # # # # # # # # # # # # # # # # # # # # # # # # # # # # # # # # # # # # # # # # # # 
+    # # # # # # # # # # # # # # # # # # # # # # # # # # # # # # # # # # # # # # # # # # # # # # # # # # # # # # # # # # # # # # # 
+    # # # # # # # # # # # # # # # # # # # # # # # # # # # # # # # # # # # # # # # # # # # # # # # # # # # # # # # # # # # # # # # 
+    def paintWeightsDrawBrushToggle(self, *args):
+	if cmds.checkBox(self.widgets["paint_displaySettings_drawBrushCB"], q = True, v = True) == False:
+	    cmds.artAttrSkinPaintCtx('ART_artAttrSkinPaint', edit = True, outline = False)
+	    
+	else:
+	    cmds.artAttrSkinPaintCtx('ART_artAttrSkinPaint', edit = True, outline = True)
+	
+	
+    # # # # # # # # # # # # # # # # # # # # # # # # # # # # # # # # # # # # # # # # # # # # # # # # # # # # # # # # # # # # # # # 
+    # # # # # # # # # # # # # # # # # # # # # # # # # # # # # # # # # # # # # # # # # # # # # # # # # # # # # # # # # # # # # # # 
+    # # # # # # # # # # # # # # # # # # # # # # # # # # # # # # # # # # # # # # # # # # # # # # # # # # # # # # # # # # # # # # # 
+    # # # # # # # # # # # # # # # # # # # # # # # # # # # # # # # # # # # # # # # # # # # # # # # # # # # # # # # # # # # # # # # 
+    # # # # # # # # # # # # # # # # # # # # # # # # # # # # # # # # # # # # # # # # # # # # # # # # # # # # # # # # # # # # # # # 
+    def paintWeightsDrawWireToggle(self, *args):
+	
+	if cmds.checkBox(self.widgets["paint_displaySettings_wireframeCB"], q = True, v = True) == False:
+	    cmds.artAttrSkinPaintCtx('ART_artAttrSkinPaint', edit = True, sa = False)
+	    
+	else:
+	    cmds.artAttrSkinPaintCtx('ART_artAttrSkinPaint', edit = True, sa = True)
+	
+	
+    # # # # # # # # # # # # # # # # # # # # # # # # # # # # # # # # # # # # # # # # # # # # # # # # # # # # # # # # # # # # # # # 
+    # # # # # # # # # # # # # # # # # # # # # # # # # # # # # # # # # # # # # # # # # # # # # # # # # # # # # # # # # # # # # # # 
+    # # # # # # # # # # # # # # # # # # # # # # # # # # # # # # # # # # # # # # # # # # # # # # # # # # # # # # # # # # # # # # # 
+    # # # # # # # # # # # # # # # # # # # # # # # # # # # # # # # # # # # # # # # # # # # # # # # # # # # # # # # # # # # # # # # 
+    # # # # # # # # # # # # # # # # # # # # # # # # # # # # # # # # # # # # # # # # # # # # # # # # # # # # # # # # # # # # # # # 
+    def paintWeightsDrawColorToggle(self, *args):
+	
+	if cmds.checkBox(self.widgets["paint_displaySettings_colorCB"], q = True, v = True) == False:
+	    cmds.artAttrSkinPaintCtx('ART_artAttrSkinPaint', edit = True, cf = False)
+	    
+	else:
+	    cmds.artAttrSkinPaintCtx('ART_artAttrSkinPaint', edit = True, cf = True)
+	    
+	    
+    # # # # # # # # # # # # # # # # # # # # # # # # # # # # # # # # # # # # # # # # # # # # # # # # # # # # # # # # # # # # # # # 
+    # # # # # # # # # # # # # # # # # # # # # # # # # # # # # # # # # # # # # # # # # # # # # # # # # # # # # # # # # # # # # # # 
+    # # # # # # # # # # # # # # # # # # # # # # # # # # # # # # # # # # # # # # # # # # # # # # # # # # # # # # # # # # # # # # # 
+    # # # # # # # # # # # # # # # # # # # # # # # # # # # # # # # # # # # # # # # # # # # # # # # # # # # # # # # # # # # # # # # 
+    # # # # # # # # # # # # # # # # # # # # # # # # # # # # # # # # # # # # # # # # # # # # # # # # # # # # # # # # # # # # # # # 
+    def paintWeightsLockWeights(self, *args):
+	imagePath = self.mayaToolsDir + "/General/Icons/ART/"
+	#find selected influences in the tree view
+	items = cmds.treeView(self.widgets["paint_influencesList_list"], q = True, si = True)
+	
+	
+	if items:
+	    for item in items:
+		if item in self.influences:
+		    cmds.skinCluster(self.skinCluster, edit = True, inf = item, lw = True)
+		    cmds.treeView(self.widgets["paint_influencesList_list"], edit = True, lbc = [item, .5, .5, .5], cs = True)
+		    cmds.treeView(self.widgets["paint_influencesList_list"], edit = True, image = (item, 1, imagePath + "LockON.bmp"))
+	
+    # # # # # # # # # # # # # # # # # # # # # # # # # # # # # # # # # # # # # # # # # # # # # # # # # # # # # # # # # # # # # # # 
+    # # # # # # # # # # # # # # # # # # # # # # # # # # # # # # # # # # # # # # # # # # # # # # # # # # # # # # # # # # # # # # # 
+    # # # # # # # # # # # # # # # # # # # # # # # # # # # # # # # # # # # # # # # # # # # # # # # # # # # # # # # # # # # # # # # 
+    # # # # # # # # # # # # # # # # # # # # # # # # # # # # # # # # # # # # # # # # # # # # # # # # # # # # # # # # # # # # # # # 
+    # # # # # # # # # # # # # # # # # # # # # # # # # # # # # # # # # # # # # # # # # # # # # # # # # # # # # # # # # # # # # # # 
+    def paintWeightsUnlockWeights(self, *args):
+	imagePath = self.mayaToolsDir + "/General/Icons/ART/"
+	#find selected influences in the tree view
+	items = cmds.treeView(self.widgets["paint_influencesList_list"], q = True, si = True)
+	
+	if items:
+	    for item in items:
+		if item in self.influences:
+		    cmds.skinCluster(self.skinCluster, edit = True, inf = item, lw = False)
+		    cmds.treeView(self.widgets["paint_influencesList_list"], edit = True, lbc = [item, -1, -1, -1,], cs = True)
+		    cmds.treeView(self.widgets["paint_influencesList_list"], edit = True, image = (item, 1, imagePath + "LockOFF.bmp"))
+	
+	
+	
+    # # # # # # # # # # # # # # # # # # # # # # # # # # # # # # # # # # # # # # # # # # # # # # # # # # # # # # # # # # # # # # # 
+    # # # # # # # # # # # # # # # # # # # # # # # # # # # # # # # # # # # # # # # # # # # # # # # # # # # # # # # # # # # # # # # 
+    # # # # # # # # # # # # # # # # # # # # # # # # # # # # # # # # # # # # # # # # # # # # # # # # # # # # # # # # # # # # # # # 
+    # # # # # # # # # # # # # # # # # # # # # # # # # # # # # # # # # # # # # # # # # # # # # # # # # # # # # # # # # # # # # # # 
+    # # # # # # # # # # # # # # # # # # # # # # # # # # # # # # # # # # # # # # # # # # # # # # # # # # # # # # # # # # # # # # # 
+    def paintWeightsSkinWeightCopy(self, *args):
+	try:
+	    mel.eval("artAttrSkinWeightCopy;")
+	except:
+	    cmds.confirmDialog(icon = "warning", title = "Copy Skin Weights", message = "Nothing is selected.  Please select a vertex to copy weights from.")
+	
+	
+    # # # # # # # # # # # # # # # # # # # # # # # # # # # # # # # # # # # # # # # # # # # # # # # # # # # # # # # # # # # # # # # 
+    # # # # # # # # # # # # # # # # # # # # # # # # # # # # # # # # # # # # # # # # # # # # # # # # # # # # # # # # # # # # # # # 
+    # # # # # # # # # # # # # # # # # # # # # # # # # # # # # # # # # # # # # # # # # # # # # # # # # # # # # # # # # # # # # # # 
+    # # # # # # # # # # # # # # # # # # # # # # # # # # # # # # # # # # # # # # # # # # # # # # # # # # # # # # # # # # # # # # # 
+    # # # # # # # # # # # # # # # # # # # # # # # # # # # # # # # # # # # # # # # # # # # # # # # # # # # # # # # # # # # # # # # 
+    def paintWeightsSkinWeightPaste(self, *args):
+	
+	try:
+	    mel.eval("artAttrSkinWeightPaste;")
+	    
+	except:
+	    cmds.confirmDialog(icon = "warning", title = "Paste Skin Weights", message = "Unable to paste weights.  The data from the copy does not appear to be valid.")
+	
+	
+    # # # # # # # # # # # # # # # # # # # # # # # # # # # # # # # # # # # # # # # # # # # # # # # # # # # # # # # # # # # # # # # 
+    # # # # # # # # # # # # # # # # # # # # # # # # # # # # # # # # # # # # # # # # # # # # # # # # # # # # # # # # # # # # # # # 
+    # # # # # # # # # # # # # # # # # # # # # # # # # # # # # # # # # # # # # # # # # # # # # # # # # # # # # # # # # # # # # # # 
+    # # # # # # # # # # # # # # # # # # # # # # # # # # # # # # # # # # # # # # # # # # # # # # # # # # # # # # # # # # # # # # # 
+    # # # # # # # # # # # # # # # # # # # # # # # # # # # # # # # # # # # # # # # # # # # # # # # # # # # # # # # # # # # # # # # 
+    def paintWeightsHammerVerts(self, *args):
+	try:
+	    mel.eval("weightHammerVerts;")
+	    
+	except:
+	    cmds.confirmDialog(icon = "warning", title = "Hammer Skin Weights", message = "The weight hammer works on polygon vertices.  Please select at least 1 vertex.")
+	    
+	
+	
+    # # # # # # # # # # # # # # # # # # # # # # # # # # # # # # # # # # # # # # # # # # # # # # # # # # # # # # # # # # # # # # # 
+    # # # # # # # # # # # # # # # # # # # # # # # # # # # # # # # # # # # # # # # # # # # # # # # # # # # # # # # # # # # # # # # 
+    # # # # # # # # # # # # # # # # # # # # # # # # # # # # # # # # # # # # # # # # # # # # # # # # # # # # # # # # # # # # # # # 
+    # # # # # # # # # # # # # # # # # # # # # # # # # # # # # # # # # # # # # # # # # # # # # # # # # # # # # # # # # # # # # # # 
+    # # # # # # # # # # # # # # # # # # # # # # # # # # # # # # # # # # # # # # # # # # # # # # # # # # # # # # # # # # # # # # # 
+    def paintWeightsMoveInfs_UI(self, *args):
+	
+	try:
+	    if cmds.window("paintWeightsMoveInfsUI", exists = True):
+		cmds.deleteUI("paintWeightsMoveInfsUI")
+		
+	    window = cmds.window("paintWeightsMoveInfsUI", w = 300, h = 60, title = "Move Influences", titleBarMenu = False, sizeable = False)
+	    mainLayout = cmds.columnLayout(w = 300, h = 60)
+	    
+	    cmds.separator(h = 10)
+	    moveFrom = cmds.optionMenu("pwmoveFrom", label = "Move From:", w = 300)
+	    cmds.separator(h = 5)
+	    moveTo = cmds.optionMenu("pwmoveTo", label = "Move To:    ", w = 300)
+	    cmds.separator(h = 10)
+	    
+	    rowLayout = cmds.rowColumnLayout(w = 300, h = 30, nc = 2, cw = [(1, 150), (2, 150)])
+	    cmds.button(label = "Move Influences", w = 150, c = self.paintWeightMoveInfs)
+	    cmds.button(label = "Cancel", w = 150, c = self.paintWeightMoveInfs_Delete)
+	    
+	    
+	    #get selected items in tree view:
+	    items = cmds.treeView("paint_influencesList_list", q = True, si = True)
+	    if len(items) <= 1:
+		cmds.warning("Only 2 influences should be selected")
+		return
+		
+	    if len(items) > 2:
+		cmds.warning("Only 2 influences should be selected")
+		return
+	    else:
+		for item in items:
+		    cmds.menuItem(label = item, parent = moveFrom)
+		    cmds.menuItem(label = item, parent = moveTo)
+		    
+		#set selection in each option menu
+		cmds.optionMenu(moveFrom, edit = True, select = 1)
+		cmds.optionMenu(moveTo, edit = True, select = 2)
+	    
+    
+	    cmds.showWindow(window)
+	    
+	except:
+	    cmds.confirmDialog(icon = "warning", title = "Move Skin Weights", message = "Please select a source influence and target influence in the influence list.")
+	
+	
+    # # # # # # # # # # # # # # # # # # # # # # # # # # # # # # # # # # # # # # # # # # # # # # # # # # # # # # # # # # # # # # # 
+    # # # # # # # # # # # # # # # # # # # # # # # # # # # # # # # # # # # # # # # # # # # # # # # # # # # # # # # # # # # # # # # 
+    # # # # # # # # # # # # # # # # # # # # # # # # # # # # # # # # # # # # # # # # # # # # # # # # # # # # # # # # # # # # # # # 
+    # # # # # # # # # # # # # # # # # # # # # # # # # # # # # # # # # # # # # # # # # # # # # # # # # # # # # # # # # # # # # # # 
+    # # # # # # # # # # # # # # # # # # # # # # # # # # # # # # # # # # # # # # # # # # # # # # # # # # # # # # # # # # # # # # # 
+    def paintWeightMoveInfs_Delete(self, *args):
+	cmds.deleteUI("paintWeightsMoveInfsUI")
+	
+	
+    # # # # # # # # # # # # # # # # # # # # # # # # # # # # # # # # # # # # # # # # # # # # # # # # # # # # # # # # # # # # # # # 
+    # # # # # # # # # # # # # # # # # # # # # # # # # # # # # # # # # # # # # # # # # # # # # # # # # # # # # # # # # # # # # # # 
+    # # # # # # # # # # # # # # # # # # # # # # # # # # # # # # # # # # # # # # # # # # # # # # # # # # # # # # # # # # # # # # # 
+    # # # # # # # # # # # # # # # # # # # # # # # # # # # # # # # # # # # # # # # # # # # # # # # # # # # # # # # # # # # # # # # 
+    # # # # # # # # # # # # # # # # # # # # # # # # # # # # # # # # # # # # # # # # # # # # # # # # # # # # # # # # # # # # # # # 
+    def paintWeightMoveInfs(self, *args):
+	
+	#get influences from UI
+	inf1 = cmds.optionMenu("pwmoveFrom", q = True, v = True)
+	inf2 = cmds.optionMenu("pwmoveTo", q = True, v = True)
+	
+	cmds.skinPercent(self.skinCluster, tmw = [inf1, inf2])
+	
+	cmds.deleteUI("paintWeightsMoveInfsUI")
+	
+	
+    # # # # # # # # # # # # # # # # # # # # # # # # # # # # # # # # # # # # # # # # # # # # # # # # # # # # # # # # # # # # # # # 
+    # # # # # # # # # # # # # # # # # # # # # # # # # # # # # # # # # # # # # # # # # # # # # # # # # # # # # # # # # # # # # # # 
+    # # # # # # # # # # # # # # # # # # # # # # # # # # # # # # # # # # # # # # # # # # # # # # # # # # # # # # # # # # # # # # # 
+    # # # # # # # # # # # # # # # # # # # # # # # # # # # # # # # # # # # # # # # # # # # # # # # # # # # # # # # # # # # # # # # 
+    # # # # # # # # # # # # # # # # # # # # # # # # # # # # # # # # # # # # # # # # # # # # # # # # # # # # # # # # # # # # # # # 
+    def paintWeightsPaintMode(self, *args):
+
+	cmds.artAttrSkinPaintCtx('ART_artAttrSkinPaint', edit = True, spm = 1)
+	cmds.selectMode( object=True )
+	
+	#collapse the influence section, unhide the selection tools frame and expand it
+	cmds.frameLayout(self.widgets["paint_influences_frame"], edit = True, collapse = False, vis = True)
+	cmds.frameLayout(self.widgets["paint_selectionTools_frame"], edit = True, collapse = True, vis = False)
+	
+	#if self.saveSelection has anything, select it
+	if self.saveSelection:
+	    cmds.select(self.saveSelection)
+	    self.saveSelection = []
+	
+	
+    # # # # # # # # # # # # # # # # # # # # # # # # # # # # # # # # # # # # # # # # # # # # # # # # # # # # # # # # # # # # # # # 
+    # # # # # # # # # # # # # # # # # # # # # # # # # # # # # # # # # # # # # # # # # # # # # # # # # # # # # # # # # # # # # # # 
+    # # # # # # # # # # # # # # # # # # # # # # # # # # # # # # # # # # # # # # # # # # # # # # # # # # # # # # # # # # # # # # # 
+    # # # # # # # # # # # # # # # # # # # # # # # # # # # # # # # # # # # # # # # # # # # # # # # # # # # # # # # # # # # # # # # 
+    # # # # # # # # # # # # # # # # # # # # # # # # # # # # # # # # # # # # # # # # # # # # # # # # # # # # # # # # # # # # # # # 
+    def paintWeightsSelectMode(self, *args):
+	cmds.artAttrSkinPaintCtx('ART_artAttrSkinPaint', edit = True, spm = 0)
+	cmds.selectMode( component=True )
+	
+	#collapse the influence section, unhide the selection tools frame and expand it
+	cmds.frameLayout(self.widgets["paint_influences_frame"], edit = True, collapse = True, vis = False)
+	cmds.frameLayout(self.widgets["paint_selectionTools_frame"], edit = True, collapse = False, vis = True)
+	
+	
+    # # # # # # # # # # # # # # # # # # # # # # # # # # # # # # # # # # # # # # # # # # # # # # # # # # # # # # # # # # # # # # # 
+    # # # # # # # # # # # # # # # # # # # # # # # # # # # # # # # # # # # # # # # # # # # # # # # # # # # # # # # # # # # # # # # 
+    # # # # # # # # # # # # # # # # # # # # # # # # # # # # # # # # # # # # # # # # # # # # # # # # # # # # # # # # # # # # # # # 
+    # # # # # # # # # # # # # # # # # # # # # # # # # # # # # # # # # # # # # # # # # # # # # # # # # # # # # # # # # # # # # # # 
+    # # # # # # # # # # # # # # # # # # # # # # # # # # # # # # # # # # # # # # # # # # # # # # # # # # # # # # # # # # # # # # # 
+    def paintWeightsWhichMode(self, *args):
+	
+	value = cmds.optionMenu(self.widgets["paint_skinTools_paintToolMode"], q = True, value = True)
+	
+	if value == "Paint":
+	    self.paintWeightsPaintMode()
+	    
+	if value == "Select":
+	    self.paintWeightsSelectMode()
+	    
+	    
+    # # # # # # # # # # # # # # # # # # # # # # # # # # # # # # # # # # # # # # # # # # # # # # # # # # # # # # # # # # # # # # # 
+    # # # # # # # # # # # # # # # # # # # # # # # # # # # # # # # # # # # # # # # # # # # # # # # # # # # # # # # # # # # # # # # 
+    # # # # # # # # # # # # # # # # # # # # # # # # # # # # # # # # # # # # # # # # # # # # # # # # # # # # # # # # # # # # # # # 
+    # # # # # # # # # # # # # # # # # # # # # # # # # # # # # # # # # # # # # # # # # # # # # # # # # # # # # # # # # # # # # # # 
+    # # # # # # # # # # # # # # # # # # # # # # # # # # # # # # # # # # # # # # # # # # # # # # # # # # # # # # # # # # # # # # # 
+    def paintWeightsSortList(self, *args):
+	#grab the selected influence in the tree view
+	selected = cmds.treeView(self.widgets["paint_influencesList_list"], q = True, si = True)
+	if selected != None:
+	    selected = selected[0]
+	    
+	else:
+	    return
+	
+	
+	
+	imagePath = self.mayaToolsDir + "/General/Icons/ART/"
+	value = cmds.optionMenu(self.widgets["paint_lSortLayout_sortMethod"], q = True, value = True)
+	cmds.treeView(self.widgets["paint_influencesList_list"], edit = True, removeAll = True)
+	              
+	              
+	if value == "Unlocked Only":
+	    self.paintWeightsSortUnlocked(selected)
+	    
+	else:
+	    if value == "Alphabetically":
+		newList = sorted(self.influences)
+		
+		for item in newList:
+		    cmds.treeView(self.widgets["paint_influencesList_list"], edit = True, addItem = (item,""))
+		    cmds.treeView(self.widgets["paint_influencesList_list"], edit = True, buttonStyle = (item, 1, "2StateButton"))
+		    cmds.treeView(self.widgets["paint_influencesList_list"], edit = True, image = (item, 1, imagePath + "LockOFF.bmp"))
+		    
+	    if value == "Flat":
+		
+		for inf in self.influences:
+		    infParent = cmds.listRelatives(inf, parent = True)[0]
+		    parent = ""
+		    if  infParent:
+			parent = ""
+			
+		    cmds.treeView(self.widgets["paint_influencesList_list"], edit = True, addItem = (inf,parent))
+		    cmds.treeView(self.widgets["paint_influencesList_list"], edit = True, buttonStyle = (inf, 1, "2StateButton"))
+		    cmds.treeView(self.widgets["paint_influencesList_list"], edit = True, image = (inf, 1, imagePath + "LockOFF.bmp"))
+		    
+		    
+	    if value == "By Hierarchy":
+
+		#add items to UI
+		for joint in self.allJoints:
+		    if joint == "root":
+			parent = ""
+		    else:
+			parent = cmds.listRelatives(joint, parent = True)[0]
+	
+			
+		    cmds.treeView(self.widgets["paint_influencesList_list"], edit = True, addItem = (joint,parent))
+		    cmds.treeView(self.widgets["paint_influencesList_list"], edit = True, buttonStyle = (joint, 1, "2StateButton"))
+		    cmds.treeView(self.widgets["paint_influencesList_list"], edit = True, image = (joint, 1, imagePath + "LockOFF.bmp"))
+	
+		    
+		#if a joint is not an influence in the skinCluster, hide that item in the tree
+		for joint in self.allJoints:
+		    if joint not in self.influences:
+			cmds.treeView(self.widgets["paint_influencesList_list"], edit = True, itemVisible = (joint, False))
+		
+
+		    
+		    
+	    #find any locked influences and reflect that in the UI
+	    for item in self.influences:
+		if cmds.skinCluster(self.skinCluster, q = True, inf = item, lw = True):
+		    cmds.treeView(self.widgets["paint_influencesList_list"], edit = True, lbc = [item, .5, .5, .5], cs = True)
+		    cmds.treeView(self.widgets["paint_influencesList_list"], edit = True, image = (item, 1, imagePath + "LockON.bmp"))
+		    
+	    #select the last selected influence
+	    cmds.treeView(self.widgets["paint_influencesList_list"], edit = True, selectItem = (selected, True))
+		    
+		    
+		    
+    # # # # # # # # # # # # # # # # # # # # # # # # # # # # # # # # # # # # # # # # # # # # # # # # # # # # # # # # # # # # # # # 
+    # # # # # # # # # # # # # # # # # # # # # # # # # # # # # # # # # # # # # # # # # # # # # # # # # # # # # # # # # # # # # # # 
+    # # # # # # # # # # # # # # # # # # # # # # # # # # # # # # # # # # # # # # # # # # # # # # # # # # # # # # # # # # # # # # # 
+    # # # # # # # # # # # # # # # # # # # # # # # # # # # # # # # # # # # # # # # # # # # # # # # # # # # # # # # # # # # # # # # 
+    # # # # # # # # # # # # # # # # # # # # # # # # # # # # # # # # # # # # # # # # # # # # # # # # # # # # # # # # # # # # # # # 
+    def paintWeightsSortUnlocked(self, selected, *args):
+	imagePath = self.mayaToolsDir + "/General/Icons/ART/"
+
+	for inf in self.influences:
+	    infParent = cmds.listRelatives(inf, parent = True)[0]
+	    parent = ""
+	    if  infParent:
+		parent = ""
+		
+	    cmds.treeView(self.widgets["paint_influencesList_list"], edit = True, addItem = (inf,parent))
+	    cmds.treeView(self.widgets["paint_influencesList_list"], edit = True, buttonStyle = (inf, 1, "2StateButton"))
+	    cmds.treeView(self.widgets["paint_influencesList_list"], edit = True, image = (inf, 1, imagePath + "LockOFF.bmp"))
+		    
+		    
+	for item in self.influences:
+	    if cmds.skinCluster(self.skinCluster, q = True, inf = item, lw = True):
+		cmds.treeView(self.widgets["paint_influencesList_list"], edit = True, itemVisible = (item, False))
+
+
+	#select the last selected influence
+	cmds.treeView(self.widgets["paint_influencesList_list"], edit = True, selectItem = (selected, True))
+	
+	
+		
+    # # # # # # # # # # # # # # # # # # # # # # # # # # # # # # # # # # # # # # # # # # # # # # # # # # # # # # # # # # # # # # # 
+    # # # # # # # # # # # # # # # # # # # # # # # # # # # # # # # # # # # # # # # # # # # # # # # # # # # # # # # # # # # # # # # 
+    # # # # # # # # # # # # # # # # # # # # # # # # # # # # # # # # # # # # # # # # # # # # # # # # # # # # # # # # # # # # # # # 
+    # # # # # # # # # # # # # # # # # # # # # # # # # # # # # # # # # # # # # # # # # # # # # # # # # # # # # # # # # # # # # # # 
+    # # # # # # # # # # # # # # # # # # # # # # # # # # # # # # # # # # # # # # # # # # # # # # # # # # # # # # # # # # # # # # # 
+    def mirrorSkinWeights(self, *args):
+	mel.eval("MirrorSkinWeightsOptions;")
+	
+	
+    # # # # # # # # # # # # # # # # # # # # # # # # # # # # # # # # # # # # # # # # # # # # # # # # # # # # # # # # # # # # # # # 
+    # # # # # # # # # # # # # # # # # # # # # # # # # # # # # # # # # # # # # # # # # # # # # # # # # # # # # # # # # # # # # # # 
+    # # # # # # # # # # # # # # # # # # # # # # # # # # # # # # # # # # # # # # # # # # # # # # # # # # # # # # # # # # # # # # # 
+    # # # # # # # # # # # # # # # # # # # # # # # # # # # # # # # # # # # # # # # # # # # # # # # # # # # # # # # # # # # # # # # 
+    # # # # # # # # # # # # # # # # # # # # # # # # # # # # # # # # # # # # # # # # # # # # # # # # # # # # # # # # # # # # # # # 
+    def pruneSkinWeights(self, *args):
+	mel.eval("PruneSmallWeightsOptions;")
+	
+	
+    # # # # # # # # # # # # # # # # # # # # # # # # # # # # # # # # # # # # # # # # # # # # # # # # # # # # # # # # # # # # # # # 
+    # # # # # # # # # # # # # # # # # # # # # # # # # # # # # # # # # # # # # # # # # # # # # # # # # # # # # # # # # # # # # # # 
+    # # # # # # # # # # # # # # # # # # # # # # # # # # # # # # # # # # # # # # # # # # # # # # # # # # # # # # # # # # # # # # # 
+    # # # # # # # # # # # # # # # # # # # # # # # # # # # # # # # # # # # # # # # # # # # # # # # # # # # # # # # # # # # # # # # 
+    # # # # # # # # # # # # # # # # # # # # # # # # # # # # # # # # # # # # # # # # # # # # # # # # # # # # # # # # # # # # # # # 
+    def paintWeightsTabletPressureOn(self, *args):
+	cmds.artAttrSkinPaintCtx('ART_artAttrSkinPaint', edit = True, up = True)
+	
+	
+    # # # # # # # # # # # # # # # # # # # # # # # # # # # # # # # # # # # # # # # # # # # # # # # # # # # # # # # # # # # # # # # 
+    # # # # # # # # # # # # # # # # # # # # # # # # # # # # # # # # # # # # # # # # # # # # # # # # # # # # # # # # # # # # # # # 
+    # # # # # # # # # # # # # # # # # # # # # # # # # # # # # # # # # # # # # # # # # # # # # # # # # # # # # # # # # # # # # # # 
+    # # # # # # # # # # # # # # # # # # # # # # # # # # # # # # # # # # # # # # # # # # # # # # # # # # # # # # # # # # # # # # # 
+    # # # # # # # # # # # # # # # # # # # # # # # # # # # # # # # # # # # # # # # # # # # # # # # # # # # # # # # # # # # # # # # 
+    def paintWeightsTabletPressureOff(self, *args):
+	cmds.artAttrSkinPaintCtx('ART_artAttrSkinPaint', edit = True, up = False)
+    
+    # # # # # # # # # # # # # # # # # # # # # # # # # # # # # # # # # # # # # # # # # # # # # # # # # # # # # # # # # # # # # # # 
+    # # # # # # # # # # # # # # # # # # # # # # # # # # # # # # # # # # # # # # # # # # # # # # # # # # # # # # # # # # # # # # # 
+    # # # # # # # # # # # # # # # # # # # # # # # # # # # # # # # # # # # # # # # # # # # # # # # # # # # # # # # # # # # # # # # 
+    # # # # # # # # # # # # # # # # # # # # # # # # # # # # # # # # # # # # # # # # # # # # # # # # # # # # # # # # # # # # # # # 
+    # # # # # # # # # # # # # # # # # # # # # # # # # # # # # # # # # # # # # # # # # # # # # # # # # # # # # # # # # # # # # # # 
+    def paintWeightsBrushPresets(self, attr, value, *args):
+	
+	if attr == "opacity":
+	    cmds.artAttrSkinPaintCtx('ART_artAttrSkinPaint', edit = True, op = value)
+	    cmds.floatSliderGrp("paint_opacitySlider", edit = True, value = value)
+	    
+	if attr == "value":
+	    cmds.artAttrSkinPaintCtx('ART_artAttrSkinPaint', edit = True, value = value)
+	    cmds.floatSliderGrp("paint_attrValueSlider", edit = True, value = value)
+	    
+    # # # # # # # # # # # # # # # # # # # # # # # # # # # # # # # # # # # # # # # # # # # # # # # # # # # # # # # # # # # # # # # 
+    # # # # # # # # # # # # # # # # # # # # # # # # # # # # # # # # # # # # # # # # # # # # # # # # # # # # # # # # # # # # # # # 
+    # # # # # # # # # # # # # # # # # # # # # # # # # # # # # # # # # # # # # # # # # # # # # # # # # # # # # # # # # # # # # # # 
+    # # # # # # # # # # # # # # # # # # # # # # # # # # # # # # # # # # # # # # # # # # # # # # # # # # # # # # # # # # # # # # # 
+    # # # # # # # # # # # # # # # # # # # # # # # # # # # # # # # # # # # # # # # # # # # # # # # # # # # # # # # # # # # # # # # 
+    def showInfluencedVerts(self, *args):
+	try:
+	    selectItemInTree = cmds.treeView(self.widgets["paint_influencesList_list"], q = True, si = True)[0]
+	    cmds.skinCluster(self.skinCluster, edit = True, siv = selectItemInTree)
+	    
+	except:
+	    cmds.confirmDialog(icon = "warning", title = "Show Influenced Verts", message = "Please select an influence from the influence list.")
+
+	    
+	    
+    # # # # # # # # # # # # # # # # # # # # # # # # # # # # # # # # # # # # # # # # # # # # # # # # # # # # # # # # # # # # # # # 
+    # # # # # # # # # # # # # # # # # # # # # # # # # # # # # # # # # # # # # # # # # # # # # # # # # # # # # # # # # # # # # # # 
+    # # # # # # # # # # # # # # # # # # # # # # # # # # # # # # # # # # # # # # # # # # # # # # # # # # # # # # # # # # # # # # # 
+    # # # # # # # # # # # # # # # # # # # # # # # # # # # # # # # # # # # # # # # # # # # # # # # # # # # # # # # # # # # # # # # 
+    # # # # # # # # # # # # # # # # # # # # # # # # # # # # # # # # # # # # # # # # # # # # # # # # # # # # # # # # # # # # # # # 
+    def paintWeightsAddInfluence(self, *args):
+	
+	#bring up a UI that has a list of all the currently invisible items. The user can then choose from that list which to add to the skinCluster
+	#at this point, we simply unhide that item and add the influence
+	if cmds.window("paintWeightsAddInf", exists = True):
+	    cmds.deleteUI("paintWeightsAddInf")
+	    
+	window = cmds.window("paintWeightsAddInf", t = "Add Influence", tbm = False, w = 300, h = 300, sizeable = False)
+	
+	mainLayout = cmds.columnLayout(w = 300, h = 300)
+	widgetLayout = cmds.rowColumnLayout(w = 300, h = 300, nc = 2, cw = [(1, 200), (2, 100)])
+	
+	#left side
+	cmds.textScrollList("paintWeightsAddInfTSL", allowMultiSelection = True, w = 200, h = 300, parent = widgetLayout)
+	
+	#right side
+	formLayout = cmds.formLayout(w= 100, h = 300, parent = widgetLayout)
+	
+	addBtn = cmds.button(w = 80, h = 40, label = "Add", c = self.paintWeightsAddInfluence_Add)
+	cancelBtn = cmds.button(w = 80, h = 40, label = "Cancel", c = self.paintWeightsAddInfluence_Cancel)
+	
+	cmds.formLayout(formLayout, edit = True, af = [(addBtn, "left", 8), (addBtn, "bottom", 60)])
+	cmds.formLayout(formLayout, edit = True, af = [(cancelBtn, "left", 8), (cancelBtn, "bottom", 10)])
+	
+	#populate the list comparing the current visible items to the allJoints list
+	visibleItems = cmds.treeView(self.widgets["paint_influencesList_list"], q = True, itemVisible = True)
+	
+	for joint in self.allJoints:
+	    if joint not in visibleItems:
+		cmds.textScrollList("paintWeightsAddInfTSL", edit = True, append = joint)
+
+	#show window
+	cmds.showWindow(window)
+	
+	
+    # # # # # # # # # # # # # # # # # # # # # # # # # # # # # # # # # # # # # # # # # # # # # # # # # # # # # # # # # # # # # # # 
+    # # # # # # # # # # # # # # # # # # # # # # # # # # # # # # # # # # # # # # # # # # # # # # # # # # # # # # # # # # # # # # # 
+    # # # # # # # # # # # # # # # # # # # # # # # # # # # # # # # # # # # # # # # # # # # # # # # # # # # # # # # # # # # # # # # 
+    # # # # # # # # # # # # # # # # # # # # # # # # # # # # # # # # # # # # # # # # # # # # # # # # # # # # # # # # # # # # # # # 
+    # # # # # # # # # # # # # # # # # # # # # # # # # # # # # # # # # # # # # # # # # # # # # # # # # # # # # # # # # # # # # # # 
+    def paintWeightsAddInfluence_Cancel(self, *args):
+	cmds.deleteUI("paintWeightsAddInf")
+    
+    
+    # # # # # # # # # # # # # # # # # # # # # # # # # # # # # # # # # # # # # # # # # # # # # # # # # # # # # # # # # # # # # # # 
+    # # # # # # # # # # # # # # # # # # # # # # # # # # # # # # # # # # # # # # # # # # # # # # # # # # # # # # # # # # # # # # # 
+    # # # # # # # # # # # # # # # # # # # # # # # # # # # # # # # # # # # # # # # # # # # # # # # # # # # # # # # # # # # # # # # 
+    # # # # # # # # # # # # # # # # # # # # # # # # # # # # # # # # # # # # # # # # # # # # # # # # # # # # # # # # # # # # # # # 
+    # # # # # # # # # # # # # # # # # # # # # # # # # # # # # # # # # # # # # # # # # # # # # # # # # # # # # # # # # # # # # # # 
+    def paintWeightsAddInfluence_Add(self, *args):
+	geo = cmds.ls(sl = True)[0]
+	selectItemInTree = cmds.treeView(self.widgets["paint_influencesList_list"], q = True, si = True)[0]
+	imagePath = self.mayaToolsDir + "/General/Icons/ART/"
+	selected = cmds.textScrollList("paintWeightsAddInfTSL", q = True, selectItem = True)
+	
+	for item in selected:
+	    self.influences.append(item)
+	    cmds.skinCluster(self.skinCluster, edit = True, ai = item, wt = 0, lw = True)
+	    cmds.treeView(self.widgets["paint_influencesList_list"], edit = True, itemVisible = (item, True))
+	    cmds.treeView(self.widgets["paint_influencesList_list"], edit = True, lbc = [item, .5, .5, .5], cs = True)
+	    cmds.treeView(self.widgets["paint_influencesList_list"], edit = True, image = (item, 1, imagePath + "LockON.bmp"))
+	    
+
+	#delete the UI
+	cmds.deleteUI("paintWeightsAddInf")
+	
+	#select the geo
+	cmds.select(geo)
+	
+	#select the last selected influence
+	cmds.treeView(self.widgets["paint_influencesList_list"], edit = True, selectItem = (selectItemInTree, True))
+	
+
+    # # # # # # # # # # # # # # # # # # # # # # # # # # # # # # # # # # # # # # # # # # # # # # # # # # # # # # # # # # # # # # # 
+    # # # # # # # # # # # # # # # # # # # # # # # # # # # # # # # # # # # # # # # # # # # # # # # # # # # # # # # # # # # # # # # 
+    # # # # # # # # # # # # # # # # # # # # # # # # # # # # # # # # # # # # # # # # # # # # # # # # # # # # # # # # # # # # # # # 
+    # # # # # # # # # # # # # # # # # # # # # # # # # # # # # # # # # # # # # # # # # # # # # # # # # # # # # # # # # # # # # # # 
+    # # # # # # # # # # # # # # # # # # # # # # # # # # # # # # # # # # # # # # # # # # # # # # # # # # # # # # # # # # # # # # # 
+    def paintWeightsRemoveInfluence(self, *args):
+	
+	#bring up a UI that has a list of all the currently invisible items. The user can then choose from that list which to add to the skinCluster
+	#at this point, we simply unhide that item and add the influence
+	if cmds.window("paintWeightsRemoveInf", exists = True):
+	    cmds.deleteUI("paintWeightsRemoveInf")
+	    
+	window = cmds.window("paintWeightsRemoveInf", t = "Remove Influence", tbm = False, w = 300, h = 300, sizeable = False)
+	
+	mainLayout = cmds.columnLayout(w = 300, h = 300)
+	widgetLayout = cmds.rowColumnLayout(w = 300, h = 300, nc = 2, cw = [(1, 200), (2, 100)])
+	
+	#left side
+	cmds.textScrollList("paintWeightsRemoveInfTSL", allowMultiSelection = True, w = 200, h = 300, parent = widgetLayout)
+	
+	#right side
+	formLayout = cmds.formLayout(w= 100, h = 300, parent = widgetLayout)
+	
+	removeBtn = cmds.button(w = 80, h = 40, label = "Remove", c = self.paintWeightsRemoveInfluence_Remove)
+	removeUnusedBtn = cmds.button(w = 80, h = 40, label = "Remove\nUnused", c = self.paintWeightsRemoveInfluence_RemoveUnused )
+	cancelBtn = cmds.button(w = 80, h = 40, label = "Cancel", c = self.paintWeightsRemoveInfluence_Cancel)
+	
+	cmds.formLayout(formLayout, edit = True, af = [(removeBtn, "left", 8), (removeBtn, "bottom", 110)])
+	cmds.formLayout(formLayout, edit = True, af = [(removeUnusedBtn, "left", 8), (removeUnusedBtn, "bottom", 60)])
+	cmds.formLayout(formLayout, edit = True, af = [(cancelBtn, "left", 8), (cancelBtn, "bottom", 10)])
+	
+	#populate the list comparing the current visible items to the allJoints list
+	visibleItems = cmds.treeView(self.widgets["paint_influencesList_list"], q = True, itemVisible = True)
+	
+	for item in visibleItems:
+	    cmds.textScrollList("paintWeightsRemoveInfTSL", edit = True, append = item)
+
+	#show window
+	cmds.showWindow(window)
+	
+	
+    # # # # # # # # # # # # # # # # # # # # # # # # # # # # # # # # # # # # # # # # # # # # # # # # # # # # # # # # # # # # # # # 
+    # # # # # # # # # # # # # # # # # # # # # # # # # # # # # # # # # # # # # # # # # # # # # # # # # # # # # # # # # # # # # # # 
+    # # # # # # # # # # # # # # # # # # # # # # # # # # # # # # # # # # # # # # # # # # # # # # # # # # # # # # # # # # # # # # # 
+    # # # # # # # # # # # # # # # # # # # # # # # # # # # # # # # # # # # # # # # # # # # # # # # # # # # # # # # # # # # # # # # 
+    # # # # # # # # # # # # # # # # # # # # # # # # # # # # # # # # # # # # # # # # # # # # # # # # # # # # # # # # # # # # # # # 
+    def paintWeightsRemoveInfluence_Cancel(self, *args):
+	cmds.deleteUI("paintWeightsRemoveInf")
+	
+	
+    
+    # # # # # # # # # # # # # # # # # # # # # # # # # # # # # # # # # # # # # # # # # # # # # # # # # # # # # # # # # # # # # # # 
+    # # # # # # # # # # # # # # # # # # # # # # # # # # # # # # # # # # # # # # # # # # # # # # # # # # # # # # # # # # # # # # # 
+    # # # # # # # # # # # # # # # # # # # # # # # # # # # # # # # # # # # # # # # # # # # # # # # # # # # # # # # # # # # # # # # 
+    # # # # # # # # # # # # # # # # # # # # # # # # # # # # # # # # # # # # # # # # # # # # # # # # # # # # # # # # # # # # # # # 
+    # # # # # # # # # # # # # # # # # # # # # # # # # # # # # # # # # # # # # # # # # # # # # # # # # # # # # # # # # # # # # # # 
+    def paintWeightsRemoveInfluence_Remove(self, items, *args):
+	geo = cmds.ls(sl = True)[0]
+	selectItemInTree = cmds.treeView(self.widgets["paint_influencesList_list"], q = True, si = True)[0]
+	imagePath = self.mayaToolsDir + "/General/Icons/ART/"
+	
+	if items == False:
+	    items = cmds.textScrollList("paintWeightsRemoveInfTSL", q = True, selectItem = True)
+	
+	for item in items:
+	    self.influences.remove(item)
+	    cmds.skinCluster(self.skinCluster, edit = True, ri = item)
+	    cmds.treeView(self.widgets["paint_influencesList_list"], edit = True, itemVisible = (item, False))
+
+	    
+
+	#delete the UI
+	cmds.deleteUI("paintWeightsRemoveInf")
+	
+	#select the geo
+	cmds.select(geo)
+	
+	#select the last selected influence
+	cmds.treeView(self.widgets["paint_influencesList_list"], edit = True, selectItem = (selectItemInTree, True))
+	
+	
+    # # # # # # # # # # # # # # # # # # # # # # # # # # # # # # # # # # # # # # # # # # # # # # # # # # # # # # # # # # # # # # # 
+    # # # # # # # # # # # # # # # # # # # # # # # # # # # # # # # # # # # # # # # # # # # # # # # # # # # # # # # # # # # # # # # 
+    # # # # # # # # # # # # # # # # # # # # # # # # # # # # # # # # # # # # # # # # # # # # # # # # # # # # # # # # # # # # # # # 
+    # # # # # # # # # # # # # # # # # # # # # # # # # # # # # # # # # # # # # # # # # # # # # # # # # # # # # # # # # # # # # # # 
+    # # # # # # # # # # # # # # # # # # # # # # # # # # # # # # # # # # # # # # # # # # # # # # # # # # # # # # # # # # # # # # # 
+    def paintWeightsRemoveInfluence_RemoveUnused(self, *args):
+	weightedInfs = cmds.skinCluster(self.skinCluster, q = True, weightedInfluence = True)
+	
+	#compare the weighted influences to the visible items
+	visibleItems = cmds.treeView(self.widgets["paint_influencesList_list"], q = True, itemVisible = True)
+	
+	items = []
+	
+	for item in visibleItems:
+	    if item not in weightedInfs:
+		items.append(item)
+		#call on remove influences :)
+		
+	if items:
+	    self.paintWeightsRemoveInfluence_Remove(items)
+		
+	
+	
+
+    # # # # # # # # # # # # # # # # # # # # # # # # # # # # # # # # # # # # # # # # # # # # # # # # # # # # # # # # # # # # # # # 
+    # # # # # # # # # # # # # # # # # # # # # # # # # # # # # # # # # # # # # # # # # # # # # # # # # # # # # # # # # # # # # # # 
+    # # # # # # # # # # # # # # # # # # # # # # # # # # # # # # # # # # # # # # # # # # # # # # # # # # # # # # # # # # # # # # # 
+    # # # # # # # # # # # # # # # # # # # # # # # # # # # # # # # # # # # # # # # # # # # # # # # # # # # # # # # # # # # # # # # 
+    # # # # # # # # # # # # # # # # # # # # # # # # # # # # # # # # # # # # # # # # # # # # # # # # # # # # # # # # # # # # # # # 
+    def exportSkinWeights(self, *args):
+	
+	selection = cmds.ls(sl = True)
+	if selection == []:
+	    cmds.confirmDialog(icon = "warning", title = "Export Skin Weights", message = "Nothing Selected")
+	    return
+	
+	
+	if selection != []:
+	    if len(selection) > 1:
+		cmds.confirmDialog(icon = "warning", title = "Export Skin Weights", message = "Too many objects selected. Please only selected 1 mesh.")
+		return
+		
+	
+	    #find geo by looking at skinclusters
+	    skinClusters = cmds.ls(type = 'skinCluster')
+	
+	    #add a dialog that warns the user that they may lose weighting information if their changes are drastic
+	    result = cmds.promptDialog(title = "Export Weights", message = "FileName:", button = ["Accept", "Cancel"])
+	    if result == "Accept":
+		name = cmds.promptDialog(q = True, text = True)
+		    
+	    path = self.mayaToolsDir + "\General\ART\SkinWeights\\"
+	    if not os.path.exists(path):
+		os.makedirs(path)
+		
+	    path = path + name + ".txt"
+	    
+	    if os.path.exists(path):
+		result = cmds.confirmDialog(icon = "question", title = "Export Weights", message = "File already exists with given name.", button = ["Overwrite", "Cancel"])
+		if result == "Cancel":
+		    return
+			
+			
+	    for skin in skinClusters:
+
+		    #go through each found skin cluster, and if we find a skin cluster whose geometry matches our selection, get influences
+		    bigList = []
+		    
+		    for cluster in skinClusters:
+			geometry = cmds.skinCluster(cluster, q = True, g = True)[0]
+			geoTransform = cmds.listRelatives(geometry, parent = True)[0]
+			
+			
+			if geoTransform == selection[0]:
+			    f = open(path, 'w')
+			    skinCluster = cluster
+
+
+
+			    #find all vertices in the current geometry
+			    verts = cmds.polyEvaluate(selection[0], vertex = True)
+
+			    for i in range(int(verts)):
+				#get weighted transforms
+				transforms = cmds.skinPercent( skinCluster, selection[0] + ".vtx[" + str(i) + "]", ib = .001, query=True, t= None)
+				
+				
+				#get skinPercent value per transform
+				values = cmds.skinPercent( skinCluster, selection[0] + ".vtx[" + str(i) + "]", ib = .001, query=True, value=True)
+				list = []
+				list.append(i)
+				    
+				for x in range(int(len(values))):
+				    list.append([transforms[x], values[x]])
+			
+				#add each vert info's transforms, and values to the overall huge list
+				bigList.append(list)
+				
+			    #write that overall list to file
+			    cPickle.dump(bigList, f)
+			    
+			    #close the file
+			    f.close()
+			    
+			    #operation complete
+			    cmds.headsUpMessage( "Skin Weights Exported!", time=3.0 )
+
+    # # # # # # # # # # # # # # # # # # # # # # # # # # # # # # # # # # # # # # # # # # # # # # # # # # # # # # # # # # # # # # # 
+    # # # # # # # # # # # # # # # # # # # # # # # # # # # # # # # # # # # # # # # # # # # # # # # # # # # # # # # # # # # # # # # 
+    # # # # # # # # # # # # # # # # # # # # # # # # # # # # # # # # # # # # # # # # # # # # # # # # # # # # # # # # # # # # # # # 
+    # # # # # # # # # # # # # # # # # # # # # # # # # # # # # # # # # # # # # # # # # # # # # # # # # # # # # # # # # # # # # # # 
+    # # # # # # # # # # # # # # # # # # # # # # # # # # # # # # # # # # # # # # # # # # # # # # # # # # # # # # # # # # # # # # # 
+    def importSkinWeights(self, *args):
+	
+	selection = cmds.ls(sl = True)
+	if selection == []:
+	    cmds.confirmDialog(icon = "warning", title = "Import Skin Weights", message = "Nothing Selected")
+	    return
+	
+	
+	if selection != []:
+	    if len(selection) > 1:
+		cmds.confirmDialog(icon = "warning", title = "Import Skin Weights", message = "Too many objects selected. Please only selected 1 mesh.")
+		return
+	    
+	    
+	    #find all files in the system dir
+	    path = self.mayaToolsDir + "\General\ART\SkinWeights\\"
+	    
+	    try:
+		fileToImport = cmds.fileDialog2(fileFilter="*.txt", dialogStyle=2, fm = 1, startingDirectory = path)[0]
+		
+	    except:
+		cmds.warning("Operation Cancelled.")
+		return
+	    
+
+
+	    #load the file with cPickle and parse out the information
+	    file = open(fileToImport, 'r')
+	    weightInfoList = cPickle.load(file)
+		
+
+	    #get the number of vertices in the mesh. This will be used to update our progress window
+	    verts = cmds.polyEvaluate(selection[0], vertex = True)
+	    if verts < 20:
+		verts = verts * 20
+		    
+	    increment = int(verts/20)
+	    matchNum = increment
+		
+	    #find the transforms from the weightInfoList. format = [ [vertNum, [jointName, value], [jointName, value] ] ]
+	    transforms = []
+	    newList = weightInfoList[:]
+	    for info in newList:
+		    
+		for i in info:
+		    if i != info[0]:
+			transform = i[0]
+			transforms.append(transform)
+			    
+	    #remove duplicates from the transforms list
+	    transforms = set(transforms)
+		
+	    #clear selection. Add transforms to selection and geo
+	    cmds.select(clear = True)
+	    for t in transforms:
+		if cmds.objExists(t):
+		    cmds.select(t, add = True)
+			
+			
+	    #check if the geometry that is selected is already skinned or not
+	    skinClusters = cmds.ls(type = 'skinCluster')
+	    skinnedGeometry = []
+	    for cluster in skinClusters:
+		geometry = cmds.skinCluster(cluster, q = True, g = True)[0]
+		geoTransform = cmds.listRelatives(geometry, parent = True)[0]
+		skinnedGeometry.append(geoTransform)
+		
+		if geoTransform == selection[0]:
+		    skin = cluster
+		
+	    
+	    #if a skin cluster does not exist already on the mesh:
+	    if selection[0] not in skinnedGeometry:
+			    
+		#put a skin cluster on the geo normalizeWeights = 1, dr=4.5, mi = 4,
+		cmds.select(selection[0], add = True)
+		skin = cmds.skinCluster( tsb = True, skinMethod = 0, name = (selection[0] +"_skinCluster"))[0]
+
+	    
+	    #setup a progress window to track length of progress
+	    cmds.progressWindow(title='Skin Weights', progress = 0, status = "Reading skin weight information from file...")
+
+	    #get the vert number, and the weights for that vertex
+	    for info in weightInfoList:
+		vertNum = info[0]
+		info.pop(0)
+		
+		#progress bar update
+		if vertNum == matchNum:
+		    cmds.progressWindow( edit=True, progress = ((matchNum/increment) * 5), status=  "Importing skin weights for " + selection[0] + ".\n Please be patient")
+		    matchNum += increment
+		
+		#set the weight for each vertex
+		try:
+		    cmds.skinPercent(skin, selection[0] + ".vtx[" + str(vertNum) + "]", transformValue= info)
+		except:
+		    pass
+		
+		
+	    #close the file
+	    file.close()
+	    
+	    #close progress window
+	    cmds.progressWindow(endProgress = 1)
+	    
+		
+    # # # # # # # # # # # # # # # # # # # # # # # # # # # # # # # # # # # # # # # # # # # # # # # # # # # # # # # # # # # # # # # 
+    # # # # # # # # # # # # # # # # # # # # # # # # # # # # # # # # # # # # # # # # # # # # # # # # # # # # # # # # # # # # # # # 
+    # # # # # # # # # # # # # # # # # # # # # # # # # # # # # # # # # # # # # # # # # # # # # # # # # # # # # # # # # # # # # # # 
+    # # # # # # # # # # # # # # # # # # # # # # # # # # # # # # # # # # # # # # # # # # # # # # # # # # # # # # # # # # # # # # # 
+    # # # # # # # # # # # # # # # # # # # # # # # # # # # # # # # # # # # # # # # # # # # # # # # # # # # # # # # # # # # # # # # 
+    def cacheWeights(self):
+	
+	#find geo by looking at skinclusters
+	skinClusters = cmds.ls(type = 'skinCluster')
+	self.characterGeo = []
+	
+	
+	for skin in skinClusters:
+	    weightedJoints = cmds.skinCluster(skin, q = True, weightedInfluence = True)
+	    
+	    for joint in weightedJoints:
+		geometryShape = cmds.skinCluster(skin, q = True, geometry = True)
+		geometry = cmds.listRelatives(geometryShape, parent = True)[0]
+		if geometry.find("skin") == -1:
+		    if geometry not in self.characterGeo:
+			self.characterGeo.append(geometry)
+		    
+	#if there is any geometry that is not proxy geometry, save out weights
+	if len(self.characterGeo) > 0:
+	    #add a dialog that warns the user that they may lose weighting information if their changes are drastic
+	    result = cmds.confirmDialog(icon = "information", title = "Edit", message = "Depending on the magnitude of your edits, you may lose skin weighting information. The tool saves out skin weight information and tries to restore if for you after your edits are completed.", button = ["OK"])
+	    if result == "OK":
+		for  geo in self.characterGeo:
+		    if geo.find(":") != -1:
+			cmds.confirmDialog(title = "Edit Placement", message = "The mesh [" + geo + "] has an invalid name. Please remove the ':' from the mesh name.")
+			return False
+		    
+		
+		    
+		    bigList = []
+		    path = self.mayaToolsDir + "\General\ART\system\\" + geo + ".txt"
+		    
+		    #go through each found skin cluster, and if we find a skin cluster whose geometry matches our selection, get influences
+		    for cluster in skinClusters:
+			geometry = cmds.skinCluster(cluster, q = True, g = True)[0]
+			geoTransform = cmds.listRelatives(geometry, parent = True)[0]
+			
+			
+			if geoTransform == geo:
+			    f = open(path, 'w')
+			    skinCluster = cluster
+
+
+
+			    #find all vertices in the current geometry
+			    verts = cmds.polyEvaluate(geo, vertex = True)
+			    
+			    for i in range(int(verts)):
+				#get weighted transforms
+				transforms = cmds.skinPercent( skinCluster, geo + ".vtx[" + str(i) + "]", ib = .001, query=True, t= None)
+				
+				
+				#get skinPercent value per transform
+				values = cmds.skinPercent( skinCluster, geo + ".vtx[" + str(i) + "]", ib = .001, query=True, value=True)
+				list = []
+				list.append(i)
+				    
+				for x in range(int(len(values))):
+				    list.append([transforms[x], values[x]])
+			
+				#add each vert info's transforms, and values to the overall huge list
+				bigList.append(list)
+				
+		    #write that overall list to file
+		    cPickle.dump(bigList, f)
+		    
+		    #close the file
+		    f.close()
+		return True
+			
+	
+	else:
+	    return True
+			    
+			
+			
+    
+    # # # # # # # # # # # # # # # # # # # # # # # # # # # # # # # # # # # # # # # # # # # # # # # # # # # # # # # # # # # # # # # 
+    # # # # # # # # # # # # # # # # # # # # # # # # # # # # # # # # # # # # # # # # # # # # # # # # # # # # # # # # # # # # # # # 
+    # # # # # # # # # # # # # # # # # # # # # # # # # # # # # # # # # # # # # # # # # # # # # # # # # # # # # # # # # # # # # # # 
+    # # # # # # # # # # # # # # # # # # # # # # # # # # # # # # # # # # # # # # # # # # # # # # # # # # # # # # # # # # # # # # # 
+    # # # # # # # # # # # # # # # # # # # # # # # # # # # # # # # # # # # # # # # # # # # # # # # # # # # # # # # # # # # # # # # 
+    def loadCachedWeights(self):
+    
+	#find all files in the system dir
+	path = self.mayaToolsDir + "\General\ART\system\\"
+	
+	files = os.listdir(path)
+	
+	
+	importAll = False
+	for file in files:
+	    niceName = file.partition(".")[0]
+	    if cmds.objExists(niceName):
+		#ask user if they want to import weights
+		if importAll == False:
+		    result = cmds.confirmDialog(icon = "question", title = "Import Skin Weights", message = "Skin Weight Data was found for " + niceName + ".", button = ["Import", "Import All", "Ignore"], defaultButton = "Ignore", cancelButton = "Cancel", dismissString = "Cancel")
+		
+		else:
+		    result = "Import"
+		    
+		if result == "Import All":
+		    importAll = True
+		    result = "Import"
+		    
+		if result == "Import":
+
+		    #load the file with cPickle and parse out the information
+		    file = open(path + file, 'r')
+		    weightInfoList = cPickle.load(file)
+		    
+    
+		    #get the number of vertices in the mesh. This will be used to update our progress window
+		    verts = cmds.polyEvaluate(niceName, vertex = True)
+		    if verts < 20:
+			verts = verts * 20
+			
+		    increment = int(verts/20)
+		    matchNum = increment
+		    
+		    #find the transforms from the weightInfoList. format = [ [vertNum, [jointName, value], [jointName, value] ] ]
+		    transforms = []
+		    newList = weightInfoList[:]
+		    for info in newList:
+			
+			for i in info:
+			    if i != info[0]:
+				transform = i[0]
+				transforms.append(transform)
+				
+		    #remove duplicates from the transforms list
+		    transforms = set(transforms)
+		    
+		    #clear selection. Add transforms to selection and geo
+		    cmds.select(clear = True)
+		    for t in transforms:
+			if cmds.objExists(t):
+			    cmds.select(t, add = True)
+			    
+    
+		    #put a skin cluster on the geo normalizeWeights = 1, dr=4.5, mi = 4,
+		    cmds.select(niceName, add = True)
+		    skin = cmds.skinCluster( tsb = True, skinMethod = 0, name = (niceName +"_skinCluster"))[0]
+		    
+		    #flood weights to 1 joint
+		    
+		    
+		    #setup a progress window to track length of progress
+		    cmds.progressWindow(title='Skin Weights', progress = 0, status = "Reading skin weight information from file...")
+    
+		    #get the vert number, and the weights for that vertex
+		    for info in weightInfoList:
+			vertNum = info[0]
+			info.pop(0)
+			
+			#progress bar update
+			if vertNum == matchNum:
+			    cmds.progressWindow( edit=True, progress = ((matchNum/increment) * 5), status=  "Importing skin weights for " + niceName + ".\n Please be patient")
+			    matchNum += increment
+			
+			#set the weight for each vertex
+			#cmds.setAttr(transform + ".liw", 0)
+			try:
+			    cmds.skinPercent(skin, niceName + ".vtx[" + str(vertNum) + "]", transformValue= info)
+			except:
+			    pass
+			
+		    #close the file
+		    file.close()
+		    
+		    #close progress window
+		    cmds.progressWindow(endProgress = 1)
+		
+           
+	#delete all weight cache files in the system folder
+	for file in files:
+	    if file != "selectionSets.txt":
+		if file != "Do_Not_Delete.txt":
+		    os.remove(path + file)
+	    
+	    
+	    
+    # # # # # # # # # # # # # # # # # # # # # # # # # # # # # # # # # # # # # # # # # # # # # # # # # # # # # # # # # # # # # # # 
+    # # # # # # # # # # # # # # # # # # # # # # # # # # # # # # # # # # # # # # # # # # # # # # # # # # # # # # # # # # # # # # # 
+    # # # # # # # # # # # # # # # # # # # # # # # # # # # # # # # # # # # # # # # # # # # # # # # # # # # # # # # # # # # # # # # 
+    # # # # # # # # # # # # # # # # # # # # # # # # # # # # # # # # # # # # # # # # # # # # # # # # # # # # # # # # # # # # # # # 
+    # # # # # # # # # # # # # # # # # # # # # # # # # # # # # # # # # # # # # # # # # # # # # # # # # # # # # # # # # # # # # # # 
+    def isolateElementFromSelection(self, *args):
+	#takes the selected vertex, converts to faces, and grows the selection until all faces of that element have been selected.
+	
+	numSelected = 0
+	currentCount = 1
+	
+	try:
+	    cmds.select(cmds.polyListComponentConversion(tf = True))
+    
+	    while numSelected != currentCount:
+		currentCount = cmds.polyEvaluate(faceComponent = True)
+		cmds.polySelectConstraint(type = 0x008, pp = True, m =3)
+		numSelected = cmds.polyEvaluate(faceComponent = True)
+		
+		
+	    isoPnl = cmds.getPanel(wf=True)
+	    isoCrnt = cmds.isolateSelect(isoPnl, q=True, s=True)    
+	    mel.eval('enableIsolateSelect %s %d' % (isoPnl,not isoCrnt) )
+		
+	    
+	    #change the label and command of the button
+	    cmds.button(self.widgets["paint_SelectionTools_isoElemFromSel"], edit = True, label = "Exit Isolation Mode", c = self.exitIsolationMode)
+	    
+	except:
+	    cmds.confirmDialog(icon = "warning", title = "Isolate", message = "Please select at least 1 vertex.")
+	    
+    
+    
+    # # # # # # # # # # # # # # # # # # # # # # # # # # # # # # # # # # # # # # # # # # # # # # # # # # # # # # # # # # # # # # # 
+    # # # # # # # # # # # # # # # # # # # # # # # # # # # # # # # # # # # # # # # # # # # # # # # # # # # # # # # # # # # # # # # 
+    # # # # # # # # # # # # # # # # # # # # # # # # # # # # # # # # # # # # # # # # # # # # # # # # # # # # # # # # # # # # # # # 
+    # # # # # # # # # # # # # # # # # # # # # # # # # # # # # # # # # # # # # # # # # # # # # # # # # # # # # # # # # # # # # # # 
+    # # # # # # # # # # # # # # # # # # # # # # # # # # # # # # # # # # # # # # # # # # # # # # # # # # # # # # # # # # # # # # # 
+    def exitIsolationMode(self, *args):
+	
+	isoPnl = cmds.getPanel(wf=True)
+	isoCrnt = cmds.isolateSelect(isoPnl, q=True, s=True)    
+	mel.eval('enableIsolateSelect %s %d' % (isoPnl,not isoCrnt) )
+	    
+	#change the label and command of the button
+	cmds.button(self.widgets["paint_SelectionTools_isoElemFromSel"], edit = True, label = "Isolate Element from Selection", c = self.isolateElementFromSelection)
+    
+    
+    
+    # # # # # # # # # # # # # # # # # # # # # # # # # # # # # # # # # # # # # # # # # # # # # # # # # # # # # # # # # # # # # # # 
+    # # # # # # # # # # # # # # # # # # # # # # # # # # # # # # # # # # # # # # # # # # # # # # # # # # # # # # # # # # # # # # # 
+    # # # # # # # # # # # # # # # # # # # # # # # # # # # # # # # # # # # # # # # # # # # # # # # # # # # # # # # # # # # # # # # 
+    # # # # # # # # # # # # # # # # # # # # # # # # # # # # # # # # # # # # # # # # # # # # # # # # # # # # # # # # # # # # # # # 
+    # # # # # # # # # # # # # # # # # # # # # # # # # # # # # # # # # # # # # # # # # # # # # # # # # # # # # # # # # # # # # # # 
+    def selectElement(self, *args):
+	numSelected = 0
+	currentCount = 1
+	
+	try:
+	    cmds.select(cmds.polyListComponentConversion(tf = True))
+    
+	    while numSelected != currentCount:
+		currentCount = cmds.polyEvaluate(faceComponent = True)
+		cmds.polySelectConstraint(type = 0x008, pp = True, m =3)
+		numSelected = cmds.polyEvaluate(faceComponent = True)
+		
+	    cmds.polyListComponentConversion(tf = True)
+	    
+	except:
+	    cmds.confirmDialog(icon = "warning", title = "Select Element", message = "Please select at least 1 component.")
+	    
+	    
+    # # # # # # # # # # # # # # # # # # # # # # # # # # # # # # # # # # # # # # # # # # # # # # # # # # # # # # # # # # # # # # # 
+    # # # # # # # # # # # # # # # # # # # # # # # # # # # # # # # # # # # # # # # # # # # # # # # # # # # # # # # # # # # # # # # 
+    # # # # # # # # # # # # # # # # # # # # # # # # # # # # # # # # # # # # # # # # # # # # # # # # # # # # # # # # # # # # # # # 
+    # # # # # # # # # # # # # # # # # # # # # # # # # # # # # # # # # # # # # # # # # # # # # # # # # # # # # # # # # # # # # # # 
+    # # # # # # # # # # # # # # # # # # # # # # # # # # # # # # # # # # # # # # # # # # # # # # # # # # # # # # # # # # # # # # # 
+    def isolateSelection(self, *args):
+	
+	isoPnl = cmds.getPanel(wf=True)
+	isoCrnt = cmds.isolateSelect(isoPnl, q=True, s=True)    
+	mel.eval('enableIsolateSelect %s %d' % (isoPnl,not isoCrnt) )
+	    
+	#change the label and command of the button
+	cmds.button(self.widgets["paint_SelectionTools_isolate"], edit = True, label = "Exit Manual Isolation Mode", c = self.exitManualIsolation)
+	    
+	    
+	    
+    # # # # # # # # # # # # # # # # # # # # # # # # # # # # # # # # # # # # # # # # # # # # # # # # # # # # # # # # # # # # # # # 
+    # # # # # # # # # # # # # # # # # # # # # # # # # # # # # # # # # # # # # # # # # # # # # # # # # # # # # # # # # # # # # # # 
+    # # # # # # # # # # # # # # # # # # # # # # # # # # # # # # # # # # # # # # # # # # # # # # # # # # # # # # # # # # # # # # # 
+    # # # # # # # # # # # # # # # # # # # # # # # # # # # # # # # # # # # # # # # # # # # # # # # # # # # # # # # # # # # # # # # 
+    # # # # # # # # # # # # # # # # # # # # # # # # # # # # # # # # # # # # # # # # # # # # # # # # # # # # # # # # # # # # # # # 
+    def exitManualIsolation(self, *args):
+	
+	isoPnl = cmds.getPanel(wf=True)
+	isoCrnt = cmds.isolateSelect(isoPnl, q=True, s=True)    
+	mel.eval('enableIsolateSelect %s %d' % (isoPnl,not isoCrnt) )
+	    
+	#change the label and command of the button
+	cmds.button(self.widgets["paint_SelectionTools_isolate"], edit = True, label = "Isolate Selection", c = self.isolateSelection)
+	
+	
+    # # # # # # # # # # # # # # # # # # # # # # # # # # # # # # # # # # # # # # # # # # # # # # # # # # # # # # # # # # # # # # # 
+    # # # # # # # # # # # # # # # # # # # # # # # # # # # # # # # # # # # # # # # # # # # # # # # # # # # # # # # # # # # # # # # 
+    # # # # # # # # # # # # # # # # # # # # # # # # # # # # # # # # # # # # # # # # # # # # # # # # # # # # # # # # # # # # # # # 
+    # # # # # # # # # # # # # # # # # # # # # # # # # # # # # # # # # # # # # # # # # # # # # # # # # # # # # # # # # # # # # # # 
+    # # # # # # # # # # # # # # # # # # # # # # # # # # # # # # # # # # # # # # # # # # # # # # # # # # # # # # # # # # # # # # # 
+    def convertSelToVerts(self, *args):
+	
+	try:
+	    cmds.select(cmds.polyListComponentConversion(tv = True))
+	    selection = cmds.ls(sl = True)
+	    
+	    for each in selection:
+		if each.find(".vtx") != -1:
+		    self.saveSelection = each
+		    
+	except:
+	    cmds.confirmDialog(icon = "warning", title = "Convert Selection", message = "Nothing Selected. Please select at least 1 component.")
+		
+		
+    # # # # # # # # # # # # # # # # # # # # # # # # # # # # # # # # # # # # # # # # # # # # # # # # # # # # # # # # # # # # # # # 
+    # # # # # # # # # # # # # # # # # # # # # # # # # # # # # # # # # # # # # # # # # # # # # # # # # # # # # # # # # # # # # # # 
+    # # # # # # # # # # # # # # # # # # # # # # # # # # # # # # # # # # # # # # # # # # # # # # # # # # # # # # # # # # # # # # # 
+    # # # # # # # # # # # # # # # # # # # # # # # # # # # # # # # # # # # # # # # # # # # # # # # # # # # # # # # # # # # # # # # 
+    # # # # # # # # # # # # # # # # # # # # # # # # # # # # # # # # # # # # # # # # # # # # # # # # # # # # # # # # # # # # # # # 
+    def skeletonBuilderHelp(self, *args):
+	
+	cmds.launch(web = "https://docs.unrealengine.com/latest/INT/Engine/Content/Tools/MayaRiggingTool/RigTool_Rigging/index.html")
+
+    # # # # # # # # # # # # # # # # # # # # # # # # # # # # # # # # # # # # # # # # # # # # # # # # # # # # # # # # # # # # # # # 
+    # # # # # # # # # # # # # # # # # # # # # # # # # # # # # # # # # # # # # # # # # # # # # # # # # # # # # # # # # # # # # # # 
+    # # # # # # # # # # # # # # # # # # # # # # # # # # # # # # # # # # # # # # # # # # # # # # # # # # # # # # # # # # # # # # # 
+    # # # # # # # # # # # # # # # # # # # # # # # # # # # # # # # # # # # # # # # # # # # # # # # # # # # # # # # # # # # # # # # 
+    # # # # # # # # # # # # # # # # # # # # # # # # # # # # # # # # # # # # # # # # # # # # # # # # # # # # # # # # # # # # # # # 
+    def jointMover_aimModeOn(self, *args):
+	
+	#create a container for our aim constraints
+	if cmds.objExists("AimModeNodes"):
+	    cmds.lockNode("AimModeNodes", lock = False)
+	    cmds.delete("AimModeNodes")
+		
+		
+	previousSelection = cmds.ls(sl = True)
+	
+        #while aim mode is activated, disable buttons we don't want to user to be able to use
+	cmds.symbolButton(self.widgets["reset"], edit = True, enable = False)
+	cmds.symbolButton(self.widgets["templatesButton"], edit = True, enable = False)
+	cmds.iconTextCheckBox(self.widgets["symmetry"], edit = True, enable = False)
+
+        #create a container to store the newly created nodes for symmetry
+        aimContainer = cmds.container(name = "AimModeNodes")
+	
+	
+	
+	#setup the constraints for the torso
+	if cmds.objExists("pelvis_mover_offset.pinned"):
+	    if cmds.getAttr("pelvis_mover_offset.pinned") == 0.0:
+		constraint = cmds.aimConstraint("spine_01_mover_offset", "pelvis_mover_offset", aimVector = [1, 0, 0], upVector = [1, 0, 0], wut = "scene", wu = [0, 0, 1])
+		cmds.container(aimContainer, edit = True, addNode= constraint)
+	else:
+	    constraint = cmds.aimConstraint("spine_01_mover_offset", "pelvis_mover_offset", aimVector = [1, 0, 0], upVector = [1, 0, 0], wut = "scene", wu = [0, 0, 1])
+	    cmds.container(aimContainer, edit = True, addNode= constraint)
+	    
+	    
+	if cmds.objExists("spine_01_mover_offset.pinned"):
+	    if cmds.getAttr("spine_01_mover_offset.pinned") == 0.0:
+		constraint = cmds.aimConstraint("spine_02_mover_offset", "spine_01_mover_offset", aimVector = [1, 0, 0], upVector = [1, 0, 0], wut = "scene", wu = [0, 0, 1])
+		cmds.container(aimContainer, edit = True, addNode= constraint)
+	else:
+	    constraint = cmds.aimConstraint("spine_02_mover_offset", "spine_01_mover_offset", aimVector = [1, 0, 0], upVector = [1, 0, 0], wut = "scene", wu = [0, 0, 1])
+	    cmds.container(aimContainer, edit = True, addNode= constraint)
+	                   
+		
+	if cmds.objExists("spine_02_mover_offset.pinned"):
+	    if cmds.getAttr("spine_02_mover_offset.pinned") == 0.0:
+		constraint = cmds.aimConstraint("spine_03_mover_offset", "spine_02_mover_offset", aimVector = [1, 0, 0], upVector = [1, 0, 0], wut = "scene", wu = [0, 0, 1])
+		cmds.container(aimContainer, edit = True, addNode= constraint)
+	else:
+	    constraint = cmds.aimConstraint("spine_03_mover_offset", "spine_02_mover_offset", aimVector = [1, 0, 0], upVector = [1, 0, 0], wut = "scene", wu = [0, 0, 1])
+	    cmds.container(aimContainer, edit = True, addNode= constraint)
+	    
+	
+	if cmds.objExists("spine_04_mover_offset"):
+	    if cmds.objExists("spine_03_mover_offset.pinned"):
+		if cmds.getAttr("spine_03_mover_offset.pinned") == 0.0:
+		    constraint = cmds.aimConstraint("spine_04_mover_offset", "spine_03_mover_offset", aimVector = [1, 0, 0], upVector = [1, 0, 0], wut = "scene", wu = [0, 0, 1])
+		    cmds.container(aimContainer, edit = True, addNode= constraint)
+	    else:
+		constraint = cmds.aimConstraint("spine_04_mover_offset", "spine_03_mover_offset", aimVector = [1, 0, 0], upVector = [1, 0, 0], wut = "scene", wu = [0, 0, 1])
+		cmds.container(aimContainer, edit = True, addNode= constraint)		
+
+	    
+	if cmds.objExists("spine_05_mover_offset"):
+	    if cmds.objExists("spine_04_mover_offset.pinned"):
+		if cmds.getAttr("spine_04_mover_offset.pinned") == 0.0:
+		    constraint = cmds.aimConstraint("spine_05_mover_offset", "spine_04_mover_offset", aimVector = [1, 0, 0], upVector = [1, 0, 0], wut = "scene", wu = [0, 0, 1])
+		    cmds.container(aimContainer, edit = True, addNode= constraint)
+	    else:
+		constraint = cmds.aimConstraint("spine_05_mover_offset", "spine_04_mover_offset", aimVector = [1, 0, 0], upVector = [1, 0, 0], wut = "scene", wu = [0, 0, 1])
+		cmds.container(aimContainer, edit = True, addNode= constraint)		
+
+
+
+
+	    
+	    
+	#setup the constraints for the neck
+	if cmds.objExists("neck_03_mover_offset"):
+	    if cmds.objExists("neck_03_mover_offset.pinned"):
+		if cmds.getAttr("neck_03_mover_offset.pinned") == 0.0:
+		    constraint = cmds.aimConstraint("head_mover_offset", "neck_03_mover_offset", aimVector = [1, 0, 0], upVector = [1, 0, 0], wut = "scene", wu = [0, 0, 1])
+		    cmds.container(aimContainer, edit = True, addNode= constraint)
+	    else:
+		constraint = cmds.aimConstraint("head_mover_offset", "neck_03_mover_offset", aimVector = [1, 0, 0], upVector = [1, 0, 0], wut = "scene", wu = [0, 0, 1])
+		cmds.container(aimContainer, edit = True, addNode= constraint)
+		
+	    if cmds.objExists("neck_02_mover_offset.pinned"):
+		if cmds.getAttr("neck_02_mover_offset.pinned") == 0.0:
+		    constraint = cmds.aimConstraint("neck_03_mover_offset", "neck_02_mover_offset", aimVector = [1, 0, 0], upVector = [1, 0, 0], wut = "scene", wu = [0, 0, 1])
+		    cmds.container(aimContainer, edit = True, addNode= constraint)
+	    else:
+		constraint = cmds.aimConstraint("neck_03_mover_offset", "neck_02_mover_offset", aimVector = [1, 0, 0], upVector = [1, 0, 0], wut = "scene", wu = [0, 0, 1])
+		cmds.container(aimContainer, edit = True, addNode= constraint)		
+	    
+	    if cmds.objExists("neck_01_mover_offset.pinned"):
+		if cmds.getAttr("neck_01_mover_offset.pinned") == 0.0:
+		    constraint = cmds.aimConstraint("neck_02_mover_offset", "neck_01_mover_offset", aimVector = [1, 0, 0], upVector = [1, 0, 0], wut = "scene", wu = [0, 0, 1])
+		    cmds.container(aimContainer, edit = True, addNode= constraint)
+	    else:
+		constraint = cmds.aimConstraint("neck_02_mover_offset", "neck_01_mover_offset", aimVector = [1, 0, 0], upVector = [1, 0, 0], wut = "scene", wu = [0, 0, 1])
+		cmds.container(aimContainer, edit = True, addNode= constraint)		
+
+
+	    
+
+	if cmds.objExists("neck_02_mover_offset") and not cmds.objExists("neck_03_mover_offset"):
+	    if cmds.objExists("neck_02_mover_offset.pinned"):
+		if cmds.getAttr("neck_02_mover_offset.pinned") == 0.0:
+		    constraint = cmds.aimConstraint("head_mover_offset", "neck_02_mover_offset", aimVector = [1, 0, 0], upVector = [1, 0, 0], wut = "scene", wu = [0, 0, 1])
+		    cmds.container(aimContainer, edit = True, addNode= constraint)
+	    else:
+		constraint = cmds.aimConstraint("head_mover_offset", "neck_02_mover_offset", aimVector = [1, 0, 0], upVector = [1, 0, 0], wut = "scene", wu = [0, 0, 1])
+		cmds.container(aimContainer, edit = True, addNode= constraint)		
+	    
+	    if cmds.objExists("neck_01_mover_offset.pinned"):
+		if cmds.getAttr("neck_01_mover_offset.pinned") == 0.0:
+		    constraint = cmds.aimConstraint("neck_02_mover_offset", "neck_01_mover_offset", aimVector = [1, 0, 0], upVector = [1, 0, 0], wut = "scene", wu = [0, 0, 1])
+		    cmds.container(aimContainer, edit = True, addNode= constraint)
+		else:
+		    constraint = cmds.aimConstraint("neck_02_mover_offset", "neck_01_mover_offset", aimVector = [1, 0, 0], upVector = [1, 0, 0], wut = "scene", wu = [0, 0, 1])
+		    cmds.container(aimContainer, edit = True, addNode= constraint)		    
+	    
+	    
+	    
+	    
+	if not cmds.objExists("neck_02_mover_offset"):
+	    if cmds.objExists("neck_01_mover_offset.pinned"):
+		if cmds.getAttr("neck_01_mover_offset.pinned") == 0.0:
+		    constraint = cmds.aimConstraint("head_mover_offset", "neck_01_mover_offset", aimVector = [1, 0, 0], upVector = [1, 0, 0], wut = "scene", wu = [0, 0, 1])
+		    cmds.container(aimContainer, edit = True, addNode= constraint)
+	    else:
+		constraint = cmds.aimConstraint("head_mover_offset", "neck_01_mover_offset", aimVector = [1, 0, 0], upVector = [1, 0, 0], wut = "scene", wu = [0, 0, 1])
+		cmds.container(aimContainer, edit = True, addNode= constraint)		
+
+
+	    
+	    
+	#setup the constraints for the left arm
+	if cmds.objExists("clavicle_mover_offset_l.pinned"):
+	    if cmds.getAttr("clavicle_mover_offset_l.pinned") == 0.0:
+		constraint = cmds.aimConstraint("upperarm_mover_offset_l", "clavicle_mover_offset_l", aimVector = [1, 0, 0], upVector = [0, 0, 1], wut = "scene", wu = [0, 0, 1])
+		cmds.container(aimContainer, edit = True, addNode= constraint)
+	else:
+	    constraint = cmds.aimConstraint("upperarm_mover_offset_l", "clavicle_mover_offset_l", aimVector = [1, 0, 0], upVector = [0, 0, 1], wut = "scene", wu = [0, 0, 1])
+	    cmds.container(aimContainer, edit = True, addNode= constraint)	    
+
+	    
+	if cmds.objExists("upperarm_mover_offset_l.pinned"):
+	    if cmds.getAttr("upperarm_mover_offset_l.pinned") == 0.0:
+		constraint = cmds.aimConstraint("lowerarm_mover_offset_l", "upperarm_mover_offset_l", aimVector = [1, 0, 0], upVector = [0, 0, 1], wut = "scene", wu = [0, 0, 1])
+		cmds.container(aimContainer, edit = True, addNode= constraint)
+	else:
+	    constraint = cmds.aimConstraint("lowerarm_mover_offset_l", "upperarm_mover_offset_l", aimVector = [1, 0, 0], upVector = [0, 0, 1], wut = "scene", wu = [0, 0, 1])
+	    cmds.container(aimContainer, edit = True, addNode= constraint)	    
+	
+	
+	if cmds.objExists("lowerarm_mover_offset_l.pinned"):
+	    if cmds.getAttr("lowerarm_mover_offset_l.pinned") == 0.0:
+		constraint = cmds.aimConstraint("hand_mover_offset_l", "lowerarm_mover_offset_l", aimVector = [1, 0, 0], upVector = [0, 0, 1], wut = "scene", wu = [0, 0, 1])
+		cmds.container(aimContainer, edit = True, addNode= constraint)
+	else:
+	    constraint = cmds.aimConstraint("hand_mover_offset_l", "lowerarm_mover_offset_l", aimVector = [1, 0, 0], upVector = [0, 0, 1], wut = "scene", wu = [0, 0, 1])
+	    cmds.container(aimContainer, edit = True, addNode= constraint)	    
+	
+	
+	
+	
+	#setup the constraints for the right arm
+	if cmds.objExists("clavicle_mover_offset_r.pinned"):
+	    if cmds.getAttr("clavicle_mover_offset_r.pinned") == 0.0:
+		constraint = cmds.aimConstraint("upperarm_mover_offset_r", "clavicle_mover_offset_r", aimVector = [-1, 0, 0], upVector = [0, 0, -1], wut = "scene", wu = [0, 0, 1])
+		cmds.container(aimContainer, edit = True, addNode= constraint)
+	else:
+	    constraint = cmds.aimConstraint("upperarm_mover_offset_r", "clavicle_mover_offset_r", aimVector = [-1, 0, 0], upVector = [0, 0, -1], wut = "scene", wu = [0, 0, 1])
+	    cmds.container(aimContainer, edit = True, addNode= constraint)	    
+	
+	if cmds.objExists("upperarm_mover_offset_r.pinned"):
+	    if cmds.getAttr("upperarm_mover_offset_r.pinned") == 0.0:
+		constraint = cmds.aimConstraint("lowerarm_mover_offset_r", "upperarm_mover_offset_r", aimVector = [-1, 0, 0], upVector = [0, 0, -1], wut = "scene", wu = [0, 0, 1])
+		cmds.container(aimContainer, edit = True, addNode= constraint)
+	else:
+	    constraint = cmds.aimConstraint("lowerarm_mover_offset_r", "upperarm_mover_offset_r", aimVector = [-1, 0, 0], upVector = [0, 0, -1], wut = "scene", wu = [0, 0, 1])
+	    cmds.container(aimContainer, edit = True, addNode= constraint)	    
+	
+	if cmds.objExists("lowerarm_mover_offset_r.pinned"):
+	    if cmds.getAttr("lowerarm_mover_offset_r.pinned") == 0.0:
+		constraint = cmds.aimConstraint("hand_mover_offset_r", "lowerarm_mover_offset_r", aimVector = [-1, 0, 0], upVector = [0, 0, -1], wut = "scene", wu = [0, 0, 1])
+		cmds.container(aimContainer, edit = True, addNode= constraint)
+	else:
+	    constraint = cmds.aimConstraint("hand_mover_offset_r", "lowerarm_mover_offset_r", aimVector = [-1, 0, 0], upVector = [0, 0, -1], wut = "scene", wu = [0, 0, 1])
+	    cmds.container(aimContainer, edit = True, addNode= constraint)	    
+
+
+	
+	
+	#setup the constraints for the left leg
+	if cmds.objExists("thigh_mover_offset_l.pinned"):
+	    if cmds.getAttr("thigh_mover_offset_l.pinned") == 0.0:
+		constraint = cmds.aimConstraint("calf_mover_offset_l", "thigh_mover_offset_l", aimVector = [-1, 0, 0], upVector = [1, 0, 0], wut = "scene", wu = [0, 0, 1])
+		cmds.container(aimContainer, edit = True, addNode= constraint)
+	else:
+	    constraint = cmds.aimConstraint("calf_mover_offset_l", "thigh_mover_offset_l", aimVector = [-1, 0, 0], upVector = [1, 0, 0], wut = "scene", wu = [0, 0, 1])
+	    cmds.container(aimContainer, edit = True, addNode= constraint)	    
+	
+	
+	if cmds.objExists("calf_mover_offset_l.pinned"):
+	    if cmds.getAttr("calf_mover_offset_l.pinned") == 0.0:
+		constraint = cmds.aimConstraint("foot_mover_offset_l", "calf_mover_offset_l", aimVector = [-1, 0, 0], upVector = [1, 0, 0], wut = "scene", wu = [0, 0, 1])
+		cmds.container(aimContainer, edit = True, addNode= constraint)	
+	else:
+	    constraint = cmds.aimConstraint("foot_mover_offset_l", "calf_mover_offset_l", aimVector = [-1, 0, 0], upVector = [1, 0, 0], wut = "scene", wu = [0, 0, 1])
+	    cmds.container(aimContainer, edit = True, addNode= constraint)		    
+
+	
+	
+	#setup the constraints for the right leg
+	if cmds.objExists("thigh_mover_offset_r.pinned"):
+	    if cmds.getAttr("thigh_mover_offset_r.pinned") == 0.0:
+		constraint = cmds.aimConstraint("calf_mover_offset_r", "thigh_mover_offset_r", aimVector = [1, 0, 0], upVector = [-1, 0, 0], wut = "scene", wu = [0, 0, 1])
+		cmds.container(aimContainer, edit = True, addNode= constraint)
+	else:
+	    constraint = cmds.aimConstraint("calf_mover_offset_r", "thigh_mover_offset_r", aimVector = [1, 0, 0], upVector = [-1, 0, 0], wut = "scene", wu = [0, 0, 1])
+	    cmds.container(aimContainer, edit = True, addNode= constraint)	    
+
+	if cmds.objExists("calf_mover_offset_r.pinned"):
+	    if cmds.getAttr("calf_mover_offset_r.pinned") == 0.0:
+		constraint = cmds.aimConstraint("foot_mover_offset_r", "calf_mover_offset_r", aimVector = [1, 0, 0], upVector = [-1, 0, 0], wut = "scene", wu = [0, 0, 1])
+		cmds.container(aimContainer, edit = True, addNode= constraint)
+	else:
+	    constraint = cmds.aimConstraint("foot_mover_offset_r", "calf_mover_offset_r", aimVector = [1, 0, 0], upVector = [-1, 0, 0], wut = "scene", wu = [0, 0, 1])
+	    cmds.container(aimContainer, edit = True, addNode= constraint)	    
+	
+	
+	
+	#setup the constraints for the fingers
+	for side in ["l", "r"]:
+	    for finger in ["index", "middle", "ring", "pinky", "thumb"]:
+		if side == "l":
+		    aimVector = [1, 0, 0]
+		    upVector = [0, 0, 0]
+		    
+		else:
+		    aimVector = [-1, 0, 0]
+		    upVector = [0, 0, 0]
+		    
+		if cmds.objExists(finger + "_metacarpal_mover_offset_" + side):
+		    if cmds.objExists(finger + "_01_mover_" + side):
+			if cmds.objExists(finger + "_metacarpal_mover_offset_" + side + ".pinned"):
+			    if cmds.getAttr(finger + "_metacarpal_mover_offset_" + side + ".pinned") == 0.0:
+				constraint = cmds.aimConstraint(finger + "_01_mover_offset_" + side, finger + "_metacarpal_mover_offset_" + side, aimVector = aimVector, upVector = upVector, wut = "scene", wu = [0, 0, 1])
+				cmds.container(aimContainer, edit = True, addNode= constraint)
+			else:
+			    constraint = cmds.aimConstraint(finger + "_01_mover_offset_" + side, finger + "_metacarpal_mover_offset_" + side, aimVector = aimVector, upVector = upVector, wut = "scene", wu = [0, 0, 1])
+			    cmds.container(aimContainer, edit = True, addNode= constraint)			    
+			
+			
+			
+		if cmds.objExists(finger + "_02_mover_offset_" + side):
+		    if cmds.objExists(finger + "_01_mover_offset_" + side + ".pinned"):
+			if cmds.getAttr(finger + "_01_mover_offset_" + side + ".pinned") == 0.0:
+			    constraint = cmds.aimConstraint(finger + "_02_mover_offset_" + side, finger + "_01_mover_offset_" + side, aimVector = aimVector, upVector = upVector, wut = "scene", wu = [0, 0, 1])
+			    cmds.container(aimContainer, edit = True, addNode= constraint)
+		    else:
+			constraint = cmds.aimConstraint(finger + "_02_mover_offset_" + side, finger + "_01_mover_offset_" + side, aimVector = aimVector, upVector = upVector, wut = "scene", wu = [0, 0, 1])
+			cmds.container(aimContainer, edit = True, addNode= constraint)			
+
+
+		    
+		if cmds.objExists(finger + "_03_mover_offset_" + side):
+		    if cmds.objExists(finger + "_02_mover_offset_" + side + ".pinned"):
+			if cmds.getAttr(finger + "_02_mover_offset_" + side + ".pinned") == 0.0:
+			    constraint = cmds.aimConstraint(finger + "_03_mover_offset_" + side, finger + "_02_mover_offset_" + side, aimVector = aimVector, upVector = upVector, wut = "scene", wu = [0, 0, 1])
+			    cmds.container(aimContainer, edit = True, addNode= constraint)
+		    else:
+			constraint = cmds.aimConstraint(finger + "_03_mover_offset_" + side, finger + "_02_mover_offset_" + side, aimVector = aimVector, upVector = upVector, wut = "scene", wu = [0, 0, 1])
+			cmds.container(aimContainer, edit = True, addNode= constraint)			
+			
+			
+			
+	#custom joint chains
+	attrs = cmds.listAttr("Skeleton_Settings")
+	
+	for attr in attrs:
+	    if attr.find("extraJoint") == 0:
+		attribute = cmds.getAttr("Skeleton_Settings." + attr, asString = True)
+		parent = attribute.partition("/")[0]
+		jointType = attribute.partition("/")[2].partition("/")[0]
+		name = attribute.rpartition("/")[2]
+		
+		if jointType == "chain" or jointType == "dynamic":
+		    numJoints = name.partition(" (")[2].partition(")")[0]
+		    name = name.partition(" (")[0]
+		    
+		    rx = cmds.getAttr(name + "_01_mover_grp.rx")
+		    if rx < -100:
+			aimVector = [0, -1, 0]
+			upVector = [0, 0, -1]
+			
+		    else:
+			aimVector = [0, 1, 0]
+			upVector = [0, 0, 1]
+			
+			
+		    for i in range(int(numJoints) -1):
+			
+			if i == 0:
+			    aimTarget = name + "_0" + str(int(numJoints)) + "_mover_offset"
+			    control = name + "_0" + str(int(numJoints) - 1) + "_mover_offset"
+			    
+			    aimTo = int(numJoints) - 1
+			    aimFrom =  aimTo - 1
+			    
+			    if cmds.objExists(control+ ".pinned"):
+				if cmds.getAttr(control + ".pinned") == 0.0:
+				    constraint = cmds.aimConstraint(aimTarget, control, aimVector = aimVector, upVector = upVector, wut = "scene", wu = [0, 0, 1])
+				    cmds.container(aimContainer, edit = True, addNode= constraint)	
+			    else:
+				constraint = cmds.aimConstraint(aimTarget, control, aimVector = aimVector, upVector = upVector, wut = "scene", wu = [0, 0, 1])
+				cmds.container(aimContainer, edit = True, addNode= constraint)				
+			    
+			else:
+			    
+			    if cmds.objExists(name + "_0" + str(aimFrom) + "_mover_offset" + ".pinned"):
+				if cmds.getAttr(name + "_0" + str(aimFrom) + "_mover_offset" + ".pinned") == 0.0:
+				    constraint = cmds.aimConstraint(name + "_0" + str(aimTo) + "_mover_offset", name + "_0" + str(aimFrom) + "_mover_offset", aimVector = aimVector, upVector = upVector, wut = "scene", wu = [0, 0, 1])
+				    cmds.container(aimContainer, edit = True, addNode= constraint)
+			    else:
+				constraint = cmds.aimConstraint(name + "_0" + str(aimTo) + "_mover_offset", name + "_0" + str(aimFrom) + "_mover_offset", aimVector = aimVector, upVector = upVector, wut = "scene", wu = [0, 0, 1])
+				cmds.container(aimContainer, edit = True, addNode= constraint)				
+				
+			    
+			    aimTo = aimFrom
+			    aimFrom = aimTo - 1
+			    
+			
+			
+			
+	#lock the container
+	cmds.lockNode("AimModeNodes", lock = True)
+	
+	
+    # # # # # # # # # # # # # # # # # # # # # # # # # # # # # # # # # # # # # # # # # # # # # # # # # # # # # # # # # # # # # # # 
+    # # # # # # # # # # # # # # # # # # # # # # # # # # # # # # # # # # # # # # # # # # # # # # # # # # # # # # # # # # # # # # # 
+    # # # # # # # # # # # # # # # # # # # # # # # # # # # # # # # # # # # # # # # # # # # # # # # # # # # # # # # # # # # # # # # 
+    # # # # # # # # # # # # # # # # # # # # # # # # # # # # # # # # # # # # # # # # # # # # # # # # # # # # # # # # # # # # # # # 
+    # # # # # # # # # # # # # # # # # # # # # # # # # # # # # # # # # # # # # # # # # # # # # # # # # # # # # # # # # # # # # # # 
+    def jointMover_aimModeOff(self, *args):
+	
+	if cmds.objExists("AimModeNodes"):
+	    cmds.lockNode("AimModeNodes", lock = False)
+	    cmds.delete("AimModeNodes")
+	    
+	    
+	    #re-enable disabled buttons
+	    cmds.symbolButton(self.widgets["reset"], edit = True, enable = True)
+	    cmds.symbolButton(self.widgets["templatesButton"], edit = True, enable = True)
+	    cmds.iconTextCheckBox(self.widgets["symmetry"], edit = True, enable = True)
+	    
+	    #bake to global
+	    self.transferToGlobal(False)
+	    
+	    #zero out twist controls
+	    for ctrl in ["upperarm_twist_01_mover_l", "upperarm_twist_02_mover_l", "upperarm_twist_03_mover_l", "upperarm_twist_01_mover_r", "upperarm_twist_02_mover_r", "upperarm_twist_03_mover_r", 
+	                 "lowerarm_twist_01_mover_l", "lowerarm_twist_02_mover_l", "lowerarm_twist_03_mover_l", "lowerarm_twist_01_mover_r", "lowerarm_twist_02_mover_r", "lowerarm_twist_03_mover_r", 
+	                 "upperleg_twist_01_mover_l", "upperleg_twist_02_mover_l", "upperleg_twist_03_mover_l", "upperleg_twist_01_mover_r", "upperleg_twist_02_mover_r", "upperleg_twist_03_mover_r",
+			 "lowerleg_twist_01_mover_l", "lowerleg_twist_02_mover_l", "lowerleg_twist_03_mover_l", "lowerleg_twist_01_mover_r", "lowerleg_twist_02_mover_r", "lowerleg_twist_03_mover_r"]:
+		
+		if cmds.objExists(ctrl):
+		    for attr in [".tx", ".ty", ".tz", ".rx", ".ry", ".rz"]:
+			cmds.setAttr(ctrl + attr, 0)
+			
+
+    # # # # # # # # # # # # # # # # # # # # # # # # # # # # # # # # # # # # # # # # # # # # # # # # # # # # # # # # # # # # # # # 
+    # # # # # # # # # # # # # # # # # # # # # # # # # # # # # # # # # # # # # # # # # # # # # # # # # # # # # # # # # # # # # # # 
+    # # # # # # # # # # # # # # # # # # # # # # # # # # # # # # # # # # # # # # # # # # # # # # # # # # # # # # # # # # # # # # # 
+    # # # # # # # # # # # # # # # # # # # # # # # # # # # # # # # # # # # # # # # # # # # # # # # # # # # # # # # # # # # # # # # 
+    # # # # # # # # # # # # # # # # # # # # # # # # # # # # # # # # # # # # # # # # # # # # # # # # # # # # # # # # # # # # # # # 
+    def createListView(self,  *args):
+        
+        #setup button color variables
+        self.blue = [.09, .75, .96]
+        self.white = [1, 1, 1]
+        self.orange = [1, .68, 0]
+        self.purple = [.5, .09, .96]
+        self.green = [0, 1, .16]
+	
+	
+        #create a scrollLayout that is a child of the passed in layout
+	self.widgets["listViewMainLayout"] = cmds.scrollLayout(w = 400, h = 700, hst = 0, parent = self.widgets["listViewTab"])
+        
+        
+        #create the tree view widget
+        self.widgets["treeViewWidget"] = cmds.treeView(parent = self.widgets["listViewMainLayout"], numberOfButtons = 2, abr = True, w = 380, h = 3000, selectionChangedCommand = partial(self.listViewSelectItem, None), pressCommand = [[1, partial(self.listViewSelectItem, "Offset")], [2, partial(self.listViewSelectItem, "Geo")]], enk = True)
+        
+        #create right click menu
+	self.widgets["treeWidget_ContextMenu"] = cmds.popupMenu(parent = self.widgets["treeViewWidget"], b = 3)
+	cmds.menuItem(label = "Lock Selected", parent = self.widgets["treeWidget_ContextMenu"], c = self.listView_LockSelected)
+	cmds.menuItem(label = "Unlock Selected", parent = self.widgets["treeWidget_ContextMenu"], c = self.listView_unLockSelected)
+	cmds.menuItem(label = "Unlock All", parent = self.widgets["treeWidget_ContextMenu"], c = self.listView_unLockAll)
+	
+        
+        #TOP LEVEL CONTROLS
+        cmds.treeView(self.widgets["treeViewWidget"], e=True, addItem = ("JOINT MOVER", ""), hb = True)
+        cmds.treeView(self.widgets["treeViewWidget"], e=True, ff = ["JOINT MOVER", 1], expandItem = ["JOINT MOVER", False],tc = ["JOINT MOVER", self.orange[0], self.orange[1], self.orange[2]])
+        cmds.treeView(self.widgets["treeViewWidget"], e=True, addItem = ("root_mover", "JOINT MOVER"), hb = True)
+	
+	cmds.treeView(self.widgets["treeViewWidget"], e=True, addItem = ("TORSO CONTROLS", "JOINT MOVER"), hb = True)
+	cmds.treeView(self.widgets["treeViewWidget"], e=True, ff = ["TORSO CONTROLS", 1], expandItem = ["TORSO CONTROLS", False],tc = ["TORSO CONTROLS", self.orange[0], self.orange[1], self.orange[2]], labelBackgroundColor = ["TORSO CONTROLS", .3, .3, .3])	
+	
+	#TORSO
+	cmds.treeView(self.widgets["treeViewWidget"], e=True, addItem = ("pelvis_mover", "TORSO CONTROLS"), hb = False)
+	cmds.treeView(self.widgets["treeViewWidget"], e=True, expandItem = ["pelvis_mover", False], bto = [["pelvis_mover", 1, 1], ["pelvis_mover", 2, 1]], btc = [["pelvis_mover", 1, .34, .65, .75], ["pelvis_mover", 2, .80, .57, .57]])
+	
+	for obj in ["spine_01", "spine_02", "spine_03", "spine_04", "spine_05"]:
+	    if cmds.objExists(obj + "_mover"):
+		cmds.treeView(self.widgets["treeViewWidget"], e=True, addItem = (obj + "_mover", "TORSO CONTROLS"), hb = False)
+		cmds.treeView(self.widgets["treeViewWidget"], e=True, expandItem = [obj + "_mover", False], bto = [[obj + "_mover", 1, 1], [obj + "_mover", 2, 1]], btc = [[obj + "_mover", 1, .34, .65, .75], [obj + "_mover", 2, .80, .57, .57]])
+		
+	#make a list of the movers that are actually being used in the build
+	visible = []
+	for mover in self.globalMovers:
+	    if cmds.getAttr(mover + ".v", lock = True) == 0:	
+		visible.append(mover.partition("|")[0])
+
+
+
+
+	#NECK/HEAD
+	cmds.treeView(self.widgets["treeViewWidget"], e=True, addItem = ("NECK/HEAD CONTROLS", "JOINT MOVER"), hb = True)
+	cmds.treeView(self.widgets["treeViewWidget"], e=True, ff = ["NECK/HEAD CONTROLS", 1], expandItem = ["NECK/HEAD CONTROLS", False],tc = ["NECK/HEAD CONTROLS", self.orange[0], self.orange[1], self.orange[2]], labelBackgroundColor = ["NECK/HEAD CONTROLS", .3, .3, .3])	
+
+	headControls = ["neck_01", "neck_02", "neck_03", "head"]
+
+	for ctrl in headControls:
+	    mover = ctrl + "_mover"
+	    if mover in visible:
+		cmds.treeView(self.widgets["treeViewWidget"], e=True, addItem = (mover, "NECK/HEAD CONTROLS"), hb = False)
+		cmds.treeView(self.widgets["treeViewWidget"], e=True, expandItem = [mover, False], bto = [[mover, 1, 1], [mover, 2, 1]], btc = [[mover, 1, .34, .65, .75], [mover, 2, .80, .57, .57]])
+
+
+
+	#ARMS
+	cmds.treeView(self.widgets["treeViewWidget"], e=True, addItem = ("ARM CONTROLS", "JOINT MOVER"), hb = True)
+	cmds.treeView(self.widgets["treeViewWidget"], e=True, ff = ["ARM CONTROLS", 1], expandItem = ["ARM CONTROLS", False],tc = ["ARM CONTROLS", self.orange[0], self.orange[1], self.orange[2]], labelBackgroundColor = ["ARM CONTROLS", .3, .3, .3])	
+
+
+
+	#Left Arm
+	cmds.treeView(self.widgets["treeViewWidget"], e=True, addItem = ("LEFT ARM", "ARM CONTROLS"), hb = True)
+	cmds.treeView(self.widgets["treeViewWidget"], e=True, expandItem = ["LEFT ARM", False], ff = ["LEFT ARM", 2], tc = ["LEFT ARM", self.blue[0], self.blue[1], self.blue[2]])
+
+	leftArmControls = ["clavicle", "upperarm", "lowerarm", "hand", "upperarm_twist_01", "upperarm_twist_02", "upperarm_twist_03", "lowerarm_twist_01", "lowerarm_twist_02", "lowerarm_twist_03"]
+	
+
+	for control in leftArmControls:
+	    mover = control + "_mover_l"
+	    if mover in visible:
+		cmds.treeView(self.widgets["treeViewWidget"], e=True, addItem = (mover, "LEFT ARM"), hb = False)
+		cmds.treeView(self.widgets["treeViewWidget"], e=True, expandItem = [mover, False], bto = [[mover, 1, 1], [mover, 2, 1]], btc = [[mover, 1, .34, .65, .75], [mover, 2, .80, .57, .57]])
+
+
+
+
+
+	#Left Fingers
+	cmds.treeView(self.widgets["treeViewWidget"], e=True, addItem = ("LEFT FINGERS", "ARM CONTROLS"), hb = True)
+	cmds.treeView(self.widgets["treeViewWidget"], e=True, expandItem = ["LEFT FINGERS", False], ff = ["LEFT FINGERS", 2], tc = ["LEFT FINGERS", self.blue[0], self.blue[1], self.blue[2]])
+	
+	leftFingerControls = ["index_metacarpal", "middle_metacarpal", "ring_metacarpal", "pinky_metacarpal", "index_01", "index_02", "index_03", "middle_01", "middle_02", "middle_03", "ring_01", "ring_02","ring_03", "pinky_01", "pinky_02", "pinky_03", "thumb_01", "thumb_02", "thumb_03"]
+	
+
+
+	for ctrl in leftFingerControls:
+	    mover = ctrl + "_mover_l"
+	    if mover in visible:
+		cmds.treeView(self.widgets["treeViewWidget"], e=True, addItem = (mover, "LEFT FINGERS"), hb = False)
+		cmds.treeView(self.widgets["treeViewWidget"], e=True, expandItem = [mover, False], bto = [[mover, 1, 1], [mover, 2, 1]], btc = [[mover, 1, .34, .65, .75], [mover, 2, .80, .57, .57]])
+
+
+	#Left Pivot Controls
+	pivotControls = ["hand_pinky_pivot", "hand_mid_pivot", "hand_thumb_pivot", "hand_tip_pivot"]
+	cmds.treeView(self.widgets["treeViewWidget"], e=True, addItem = ("LEFT HAND PIVOT CONTROLS", "LEFT ARM"), hb = True)
+	cmds.treeView(self.widgets["treeViewWidget"], e=True, expandItem = ["LEFT HAND PIVOT CONTROLS", False], ff = ["LEFT HAND PIVOT CONTROLS", 2], tc = ["LEFT HAND PIVOT CONTROLS", self.blue[0], self.blue[1], self.blue[2]])
+
+	for ctrl in pivotControls:
+	    mover = ctrl + "_mover_l"
+	    if mover in visible:
+		cmds.treeView(self.widgets["treeViewWidget"], e=True, addItem = (mover, "LEFT HAND PIVOT CONTROLS"), hb = True)
+		cmds.treeView(self.widgets["treeViewWidget"], e=True, expandItem = [mover, False])	
+
+
+
+
+	#Right Arm
+	cmds.treeView(self.widgets["treeViewWidget"], e=True, addItem = ("RIGHT ARM", "ARM CONTROLS"), hb = True)
+	cmds.treeView(self.widgets["treeViewWidget"], e=True,  expandItem = ["RIGHT ARM", False], ff = ["RIGHT ARM", 2], tc = ["RIGHT ARM", self.blue[0], self.blue[1], self.blue[2]])
+
+	rightArmControls = ["clavicle", "upperarm", "lowerarm", "hand", "upperarm_twist_01", "upperarm_twist_02", "upperarm_twist_03", "lowerarm_twist_01", "lowerarm_twist_02", "lowerarm_twist_03"]
+	
+
+	for control in rightArmControls:
+	    mover = control + "_mover_r"
+	    if mover in visible:
+		cmds.treeView(self.widgets["treeViewWidget"], e=True, addItem = (mover, "RIGHT ARM"), hb = False)
+		cmds.treeView(self.widgets["treeViewWidget"], e=True, expandItem = [mover, False], bto = [[mover, 1, 1], [mover, 2, 1]], btc = [[mover, 1, .34, .65, .75], [mover, 2, .80, .57, .57]])
+
+
+
+
+
+	#RIGHT Fingers
+	cmds.treeView(self.widgets["treeViewWidget"], e=True, addItem = ("RIGHT FINGERS", "ARM CONTROLS"), hb = True)
+	cmds.treeView(self.widgets["treeViewWidget"], e=True, expandItem = ["RIGHT FINGERS", False], ff = ["RIGHT FINGERS", 2], tc = ["RIGHT FINGERS", self.blue[0], self.blue[1], self.blue[2]], bti = [["RIGHT FINGERS", 1, "S"], ["RIGHT FINGERS", 2, "V"]])
+	
+	rightFingerControls = ["index_metacarpal", "middle_metacarpal", "ring_metacarpal", "pinky_metacarpal", "index_01", "index_02", "index_03", "middle_01", "middle_02", "middle_03", "ring_01", "ring_02","ring_03", "pinky_01", "pinky_02", "pinky_03", "thumb_01", "thumb_02", "thumb_03"]
+	
+
+
+	for ctrl in rightFingerControls:
+	    mover = ctrl + "_mover_r"
+	    if mover in visible:
+		cmds.treeView(self.widgets["treeViewWidget"], e=True, addItem = (mover, "RIGHT FINGERS"), hb = False)
+		cmds.treeView(self.widgets["treeViewWidget"], e=True, expandItem = [mover, False], bto = [[mover, 1, 1], [mover, 2, 1]], btc = [[mover, 1, .34, .65, .75], [mover, 2, .80, .57, .57]])
+
+
+	#Right Pivot Controls
+	pivotControls = ["hand_pinky_pivot", "hand_mid_pivot", "hand_thumb_pivot", "hand_tip_pivot"]
+	cmds.treeView(self.widgets["treeViewWidget"], e=True, addItem = ("RIGHT HAND PIVOT CONTROLS", "RIGHT ARM"), hb = True)
+	cmds.treeView(self.widgets["treeViewWidget"], e=True, expandItem = ["RIGHT HAND PIVOT CONTROLS", False], ff = ["RIGHT HAND PIVOT CONTROLS", 2], tc = ["RIGHT HAND PIVOT CONTROLS", self.blue[0], self.blue[1], self.blue[2]])
+
+	for ctrl in pivotControls:
+	    mover = ctrl + "_mover_r"
+	    if mover in visible:
+		cmds.treeView(self.widgets["treeViewWidget"], e=True, addItem = (mover, "RIGHT HAND PIVOT CONTROLS"), hb = True)
+		cmds.treeView(self.widgets["treeViewWidget"], e=True, expandItem = [mover, False])
+
+
+
+
+	#LEGS
+	cmds.treeView(self.widgets["treeViewWidget"], e=True, addItem = ("LEG CONTROLS", "JOINT MOVER"), hb = True)
+	cmds.treeView(self.widgets["treeViewWidget"], e=True, ff = ["LEG CONTROLS", 1], expandItem = ["LEG CONTROLS", False],tc = ["LEG CONTROLS", self.orange[0], self.orange[1], self.orange[2]], labelBackgroundColor = ["LEG CONTROLS", .3, .3, .3])	
+
+
+
+	#Left Leg
+	cmds.treeView(self.widgets["treeViewWidget"], e=True, addItem = ("LEFT LEG", "LEG CONTROLS"), hb = True)
+	cmds.treeView(self.widgets["treeViewWidget"], e=True, expandItem = ["LEFT LEG", False], ff = ["LEFT LEG", 2], tc = ["LEFT LEG", self.blue[0], self.blue[1], self.blue[2]])
+
+	leftLegControls = ["thigh", "calf", "foot", "ball", "thigh_twist_01", "thigh_twist_02", "thigh_twist_03", "calf_twist_01", "calf_twist_02", "calf_twist_03"]
+	
+
+	for control in leftLegControls:
+	    mover = control + "_mover_l"
+	    if mover in visible:
+		cmds.treeView(self.widgets["treeViewWidget"], e=True, addItem = (mover, "LEFT LEG"), hb = False)
+		cmds.treeView(self.widgets["treeViewWidget"], e=True, expandItem = [mover, False], bto = [[mover, 1, 1], [mover, 2, 1]], btc = [[mover, 1, .34, .65, .75], [mover, 2, .80, .57, .57]])
+
+
+
+	#pivot controls
+	pivotControls = ["heel_pivot", "toe_pivot", "inside_pivot", "outside_pivot"]
+	cmds.treeView(self.widgets["treeViewWidget"], e=True, addItem = ("LEFT PIVOT CONTROLS", "LEFT LEG"), hb = True)
+	cmds.treeView(self.widgets["treeViewWidget"], e=True, expandItem = ["LEFT PIVOT CONTROLS", False], ff = ["LEFT PIVOT CONTROLS", 2], tc = ["LEFT PIVOT CONTROLS", self.blue[0], self.blue[1], self.blue[2]])
+
+	for ctrl in pivotControls:
+	    mover = ctrl + "_l_mover"
+	    if mover in visible:
+		cmds.treeView(self.widgets["treeViewWidget"], e=True, addItem = (mover, "LEFT PIVOT CONTROLS"), hb = True)
+		cmds.treeView(self.widgets["treeViewWidget"], e=True, expandItem = [mover, False])
+
+	    
+
+	#Right Leg
+	cmds.treeView(self.widgets["treeViewWidget"], e=True, addItem = ("RIGHT LEG", "LEG CONTROLS"), hb = True)
+	cmds.treeView(self.widgets["treeViewWidget"], e=True, expandItem = ["RIGHT LEG", False], ff = ["RIGHT LEG", 2], tc = ["RIGHT LEG", self.blue[0], self.blue[1], self.blue[2]])
+
+	rightLegControls = ["thigh", "calf", "foot", "ball", "thigh_twist_01", "thigh_twist_02", "thigh_twist_03", "calf_twist_01", "calf_twist_02", "calf_twist_03"]
+	
+
+	for control in rightLegControls:
+	    mover = control + "_mover_r"
+	    if mover in visible:
+		cmds.treeView(self.widgets["treeViewWidget"], e=True, addItem = (mover, "RIGHT LEG"), hb = False)
+		cmds.treeView(self.widgets["treeViewWidget"], e=True, expandItem = [mover, False], bto = [[mover, 1, 1], [mover, 2, 1]], btc = [[mover, 1, .34, .65, .75], [mover, 2, .80, .57, .57]])
+
+
+
+
+	#pivot controls
+	pivotControls = ["heel_pivot", "toe_pivot", "inside_pivot", "outside_pivot"]
+	cmds.treeView(self.widgets["treeViewWidget"], e=True, addItem = ("RIGHT PIVOT CONTROLS", "RIGHT LEG"), hb = True)
+	cmds.treeView(self.widgets["treeViewWidget"], e=True, expandItem = ["RIGHT PIVOT CONTROLS", False], ff = ["RIGHT PIVOT CONTROLS", 2], tc = ["RIGHT PIVOT CONTROLS", self.blue[0], self.blue[1], self.blue[2]])
+
+	for ctrl in pivotControls:
+	    mover = ctrl + "_r_mover"
+	    if mover in visible:
+		cmds.treeView(self.widgets["treeViewWidget"], e=True, addItem = (mover, "RIGHT PIVOT CONTROLS"), hb = True)
+		cmds.treeView(self.widgets["treeViewWidget"], e=True, expandItem = [mover, False])
+
+
+
+
+	
+	#CUSTOM MODULES
+	cmds.treeView(self.widgets["treeViewWidget"], e=True, addItem = ("CUSTOM LEAF JOINTS", "JOINT MOVER"), hb = True)
+	cmds.treeView(self.widgets["treeViewWidget"], e=True, ff = ["CUSTOM LEAF JOINTS", 1], expandItem = ["CUSTOM LEAF JOINTS", False],tc = ["CUSTOM LEAF JOINTS", self.orange[0], self.orange[1], self.orange[2]], labelBackgroundColor = ["CUSTOM LEAF JOINTS", .3, .3, .3])	
+
+	cmds.treeView(self.widgets["treeViewWidget"], e=True, addItem = ("CUSTOM JIGGLE JOINTS", "JOINT MOVER"), hb = True)
+	cmds.treeView(self.widgets["treeViewWidget"], e=True, ff = ["CUSTOM JIGGLE JOINTS", 1], expandItem = ["CUSTOM JIGGLE JOINTS", False],tc = ["CUSTOM JIGGLE JOINTS", self.orange[0], self.orange[1], self.orange[2]], labelBackgroundColor = ["CUSTOM JIGGLE JOINTS", .3, .3, .3])	
+
+	cmds.treeView(self.widgets["treeViewWidget"], e=True, addItem = ("CUSTOM CHAINS", "JOINT MOVER"), hb = True)
+	cmds.treeView(self.widgets["treeViewWidget"], e=True, ff = ["CUSTOM CHAINS", 1], expandItem = ["CUSTOM CHAINS", False],tc = ["CUSTOM CHAINS", self.orange[0], self.orange[1], self.orange[2]], labelBackgroundColor = ["CUSTOM CHAINS", .3, .3, .3])	
+
+	
+	#find extra rig modules
+	attrs = cmds.listAttr("Skeleton_Settings")
+
+	for attr in attrs:
+	    if attr.find("extraJoint") == 0:
+		attribute = cmds.getAttr("Skeleton_Settings." + attr, asString = True)
+		jointType = attribute.partition("/")[2].partition("/")[0]
+		name = attribute.rpartition("/")[2]	
+
+
+			    
+		#LEAF JOINTS
+		if jointType == "leaf":
+		    if name.find("(") != -1:
+			name = name.partition(" (")[0]
+
+		    mover = name + "_mover"
+		    cmds.treeView(self.widgets["treeViewWidget"], e=True, addItem = (mover, "CUSTOM LEAF JOINTS"), hb = False)
+		    cmds.treeView(self.widgets["treeViewWidget"], e=True, expandItem = [mover, False], bto = [[mover, 1, 1], [mover, 2, 1]], btc = [[mover, 1, .34, .65, .75], [mover, 2, .80, .57, .57]])				
+
+					
+		#JIGGLE JOINTS
+		if jointType == "jiggle":
+		    mover = name + "_mover"
+		    cmds.treeView(self.widgets["treeViewWidget"], e=True, addItem = (mover, "CUSTOM JIGGLE JOINTS"), hb = False)
+		    cmds.treeView(self.widgets["treeViewWidget"], e=True, expandItem = [mover, False], bto = [[mover, 1, 1], [mover, 2, 1]], btc = [[mover, 1, .34, .65, .75], [mover, 2, .80, .57, .57]])	
+
+
+		#DYNAMIC CHAINS
+		if jointType == "chain":
+		    numJoints = name.partition(" (")[2].partition(")")[0]
+		    name = name.partition(" (")[0]
+
+
+		    cmds.treeView(self.widgets["treeViewWidget"], e=True, addItem = (name, "CUSTOM CHAINS"), hb = True)
+		    cmds.treeView(self.widgets["treeViewWidget"], e=True, ff = [name, 1], expandItem = [name, False],tc = [name, self.orange[0], self.orange[1], self.orange[2]], labelBackgroundColor = [name, .3, .3, .3])	
+		    
+		    for i in range(int(numJoints)):
+			mover = name + "_0" + str(i + 1) + "_mover"
+			cmds.treeView(self.widgets["treeViewWidget"], e=True, addItem = (mover, name), hb = False)
+			cmds.treeView(self.widgets["treeViewWidget"], e=True, expandItem = [mover, False], bto = [[mover, 1, 1], [mover, 2, 1]], btc = [[mover, 1, .34, .65, .75], [mover, 2, .80, .57, .57]])
+
+
+
+    # # # # # # # # # # # # # # # # # # # # # # # # # # # # # # # # # # # # # # # # # # # # # # # # # # # # # # # # # # # # # # # 
+    # # # # # # # # # # # # # # # # # # # # # # # # # # # # # # # # # # # # # # # # # # # # # # # # # # # # # # # # # # # # # # # 
+    # # # # # # # # # # # # # # # # # # # # # # # # # # # # # # # # # # # # # # # # # # # # # # # # # # # # # # # # # # # # # # # 
+    # # # # # # # # # # # # # # # # # # # # # # # # # # # # # # # # # # # # # # # # # # # # # # # # # # # # # # # # # # # # # # # 
+    # # # # # # # # # # # # # # # # # # # # # # # # # # # # # # # # # # # # # # # # # # # # # # # # # # # # # # # # # # # # # # # 
+    def listViewSelectItem(self, mode, *args):
+        
+        mods = cmds.getModifiers()
+	suffix = None
+
+        selected = cmds.treeView(self.widgets["treeViewWidget"], q = True, selectItem = True)
+	if selected != None:
+	    if mode == None:
+		if (mods & 1) > 0:
+		    for item in selected:
+			if cmds.objExists(item):
+			    cmds.select(item, tgl = True)
+	
+		if (mods & 1) == 0:
+		    cmds.select(clear = True)
+		    
+		    for item in selected:
+			if cmds.objExists(item):
+			    cmds.select(item, add = True)
+    
+	    if mode == "Offset":
+		if (mods & 1) > 0:
+		    for item in selected:
+			if item.partition("_mover")[2].find("_") == 0:
+			    suffix = item.partition("_mover_")[2]
+			
+			if suffix == None:
+			    mover = item.partition("_mover")[0] + "_mover_offset"
+			else:
+			    mover = item.partition("_mover")[0] + "_mover_offset_" + suffix
+    
+			if cmds.objExists(mover):
+			    cmds.select(mover, tgl = True)
+	
+		if (mods & 1) == 0:
+		    cmds.select(clear = True)
+		    
+		    for item in selected:
+			if item.partition("_mover")[2].find("_") == 0:
+			    suffix = item.partition("_mover_")[2]
+			
+			if suffix == None:
+			    mover = item.partition("_mover")[0] + "_mover_offset"
+			else:
+			    mover = item.partition("_mover")[0] + "_mover_offset_" + suffix
+    
+			if cmds.objExists(mover):
+			    cmds.select(mover, add = True)
+    
+    
+	    if mode == "Geo":
+		if (mods & 1) > 0:
+		    for item in selected:
+			if item.partition("_mover")[2].find("_") == 0:
+			    suffix = item.partition("_mover_")[2]
+			
+			if suffix == None:
+			    mover = item.partition("_mover")[0] + "_geo_mover"
+			else:
+			    mover = item.partition("_mover")[0] + "_geo_mover_" + suffix
+    
+			if cmds.objExists(mover):
+			    cmds.select(mover, tgl = True)
+	
+		if (mods & 1) == 0:
+		    cmds.select(clear = True)
+		    
+		    for item in selected:
+			if item.partition("_mover")[2].find("_") == 0:
+			    suffix = item.partition("_mover_")[2]
+			
+			if suffix == None:
+			    mover = item.partition("_mover")[0] + "_geo_mover"
+			else:
+			    mover = item.partition("_mover")[0] + "_geo_mover_" + suffix
+    
+			if cmds.objExists(mover):
+			    cmds.select(mover, add = True)
+
+
+
+
+    # # # # # # # # # # # # # # # # # # # # # # # # # # # # # # # # # # # # # # # # # # # # # # # # # # # # # # # # # # # # # # # 
+    # # # # # # # # # # # # # # # # # # # # # # # # # # # # # # # # # # # # # # # # # # # # # # # # # # # # # # # # # # # # # # # 
+    # # # # # # # # # # # # # # # # # # # # # # # # # # # # # # # # # # # # # # # # # # # # # # # # # # # # # # # # # # # # # # # 
+    # # # # # # # # # # # # # # # # # # # # # # # # # # # # # # # # # # # # # # # # # # # # # # # # # # # # # # # # # # # # # # # 
+    # # # # # # # # # # # # # # # # # # # # # # # # # # # # # # # # # # # # # # # # # # # # # # # # # # # # # # # # # # # # # # # 
+    def listView_ScriptJob(self, *args):
+        
+        cmds.scriptJob(event = ["SelectionChanged", self.listView_scriptJobCommand], parent = self.widgets["window"])
+        
+        
+    # # # # # # # # # # # # # # # # # # # # # # # # # # # # # # # # # # # # # # # # # # # # # # # # # # # # # # # # # # # # # # # 
+    # # # # # # # # # # # # # # # # # # # # # # # # # # # # # # # # # # # # # # # # # # # # # # # # # # # # # # # # # # # # # # # 
+    # # # # # # # # # # # # # # # # # # # # # # # # # # # # # # # # # # # # # # # # # # # # # # # # # # # # # # # # # # # # # # # 
+    # # # # # # # # # # # # # # # # # # # # # # # # # # # # # # # # # # # # # # # # # # # # # # # # # # # # # # # # # # # # # # # 
+    # # # # # # # # # # # # # # # # # # # # # # # # # # # # # # # # # # # # # # # # # # # # # # # # # # # # # # # # # # # # # # # 
+    def listView_scriptJobCommand(self, *args):
+
+
+	selection = cmds.ls(sl = True)
+	selectedItems = cmds.treeView(self.widgets["treeViewWidget"], q = True, selectItem = True)
+
+	
+	if selectedItems != None:
+	    for item in selectedItems:
+		if item not in selection:
+		    cmds.treeView(self.widgets["treeViewWidget"], edit = True, selectItem = [item, False])
+
+
+		
+		
+
+    # # # # # # # # # # # # # # # # # # # # # # # # # # # # # # # # # # # # # # # # # # # # # # # # # # # # # # # # # # # # # # # 
+    # # # # # # # # # # # # # # # # # # # # # # # # # # # # # # # # # # # # # # # # # # # # # # # # # # # # # # # # # # # # # # # 
+    # # # # # # # # # # # # # # # # # # # # # # # # # # # # # # # # # # # # # # # # # # # # # # # # # # # # # # # # # # # # # # # 
+    # # # # # # # # # # # # # # # # # # # # # # # # # # # # # # # # # # # # # # # # # # # # # # # # # # # # # # # # # # # # # # # 
+    # # # # # # # # # # # # # # # # # # # # # # # # # # # # # # # # # # # # # # # # # # # # # # # # # # # # # # # # # # # # # # # 
+    def findFileOnDisk(self, path, *args):
+	drive = path.partition(":")[0] + ":\\"
+	restOfPath = path.partition(":")[2]
+	restOfPath = os.path.normpath(restOfPath)
+	fullPath = drive + restOfPath
+	
+	
+	Popen_arg = "explorer /select, \"" + os.path.normpath(fullPath) + "\""
+	subprocess.Popen(Popen_arg)
+
+    # # # # # # # # # # # # # # # # # # # # # # # # # # # # # # # # # # # # # # # # # # # # # # # # # # # # # # # # # # # # # # # 
+    # # # # # # # # # # # # # # # # # # # # # # # # # # # # # # # # # # # # # # # # # # # # # # # # # # # # # # # # # # # # # # # 
+    # # # # # # # # # # # # # # # # # # # # # # # # # # # # # # # # # # # # # # # # # # # # # # # # # # # # # # # # # # # # # # # 
+    # # # # # # # # # # # # # # # # # # # # # # # # # # # # # # # # # # # # # # # # # # # # # # # # # # # # # # # # # # # # # # # 
+    # # # # # # # # # # # # # # # # # # # # # # # # # # # # # # # # # # # # # # # # # # # # # # # # # # # # # # # # # # # # # # # 
+    def listView_LockSelected(self, *args):
+	
+	previousSelection = cmds.ls(sl = True)
+	selected = cmds.treeView(self.widgets["treeViewWidget"], q = True, si = True)
+	
+	
+	#unlock Joint mover
+	cmds.lockNode("JointMover", lock = False)
+
+
+
+	for each in selected:
+	    suffix = None
+	    
+	    #get names of global, offset and geo mover
+	    if each.partition("_mover")[2].find("_") == 0:
+		suffix = each.partition("_mover_")[2]
+	    
+	    globalMover = each
+	    
+	    if suffix == None:
+		offsetMover = each.partition("_mover")[0] + "_mover_offset"
+		geoMover = each.partition("_mover")[0] + "_geo_mover"
+		
+	    else:
+		offsetMover = each.partition("_mover")[0] + "_mover_offset_" + suffix
+		geoMover = each.partition("_mover")[0] + "_geo_mover_"	+ suffix
+		
+	    
+	    #if the pinned attribute doesn't exist, create it
+	    for mover in [globalMover, offsetMover, geoMover]:
+		if cmds.objExists(mover + ".pinned") == False:		    
+		    cmds.lockNode(mover, lock = False)
+		    cmds.addAttr(mover, ln = "pinned", defaultValue = 0)
+		    cmds.lockNode(mover, lock = True)
+		    
+	    #set the attr to 1
+	    for mover in [globalMover, offsetMover, geoMover]:
+		cmds.setAttr(mover + ".pinned", 1)
+		
+		
+	    #lastly, change the font to italics
+	    cmds.treeView(self.widgets["treeViewWidget"], edit = True, ff = [each, 2], ornament = [each, 1, 1, 7])
+	    
+	    #select previous selection
+	    if previousSelection != []:
+		cmds.select(previousSelection)
+	    
+	    
+	#relock joint mover
+	cmds.lockNode("JointMover", lock = True)
+    
+	    
+	    
+    # # # # # # # # # # # # # # # # # # # # # # # # # # # # # # # # # # # # # # # # # # # # # # # # # # # # # # # # # # # # # # # 
+    # # # # # # # # # # # # # # # # # # # # # # # # # # # # # # # # # # # # # # # # # # # # # # # # # # # # # # # # # # # # # # # 
+    # # # # # # # # # # # # # # # # # # # # # # # # # # # # # # # # # # # # # # # # # # # # # # # # # # # # # # # # # # # # # # # 
+    # # # # # # # # # # # # # # # # # # # # # # # # # # # # # # # # # # # # # # # # # # # # # # # # # # # # # # # # # # # # # # # 
+    # # # # # # # # # # # # # # # # # # # # # # # # # # # # # # # # # # # # # # # # # # # # # # # # # # # # # # # # # # # # # # # 
+    def listView_unLockSelected(self, *args):
+	
+	previousSelection = cmds.ls(sl = True)
+	selected = cmds.treeView(self.widgets["treeViewWidget"], q = True, si = True)
+	suffix = None	
+	
+	for each in selected:
+	    #get names of global, offset and geo mover
+	    if each.partition("_mover")[2].find("_") == 0:
+		suffix = each.partition("_mover_")[2]
+	    
+	    globalMover = each
+	    
+	    if suffix == None:
+		offsetMover = each.partition("_mover")[0] + "_mover_offset"
+		geoMover = each.partition("_mover")[0] + "_geo_mover"
+		
+	    else:
+		offsetMover = each.partition("_mover")[0] + "_mover_offset_" + suffix
+		geoMover = each.partition("_mover")[0] + "_geo_mover_"	+ suffix
+		
+
+	    try:
+		#set the attr to 0
+		for mover in [globalMover, offsetMover, geoMover]:
+		    cmds.setAttr(mover + ".pinned", 0)
+		    
+		    
+		#lastly, change the font to italics
+		cmds.treeView(self.widgets["treeViewWidget"], edit = True, ff = [each, 0], ornament = [each, 0, 1, 7])
+		
+	    except:
+		pass
+	    
+	    
+	    #select previous selection
+	    if previousSelection != []:
+		cmds.select(previousSelection)
+	    
+	    
+    # # # # # # # # # # # # # # # # # # # # # # # # # # # # # # # # # # # # # # # # # # # # # # # # # # # # # # # # # # # # # # # 
+    # # # # # # # # # # # # # # # # # # # # # # # # # # # # # # # # # # # # # # # # # # # # # # # # # # # # # # # # # # # # # # # 
+    # # # # # # # # # # # # # # # # # # # # # # # # # # # # # # # # # # # # # # # # # # # # # # # # # # # # # # # # # # # # # # # 
+    # # # # # # # # # # # # # # # # # # # # # # # # # # # # # # # # # # # # # # # # # # # # # # # # # # # # # # # # # # # # # # # 
+    # # # # # # # # # # # # # # # # # # # # # # # # # # # # # # # # # # # # # # # # # # # # # # # # # # # # # # # # # # # # # # # 
+    def listView_unLockAll(self, *args):
+	
+	previousSelection = cmds.ls(sl = True)
+	selected = cmds.treeView(self.widgets["treeViewWidget"], q = True, ch = "")
+	suffix = None	
+	
+	for each in selected:
+	    #get names of global, offset and geo mover
+	    if each.partition("_mover")[2].find("_") == 0:
+		suffix = each.partition("_mover_")[2]
+	    
+	    globalMover = each
+	    
+	    if suffix == None:
+		offsetMover = each.partition("_mover")[0] + "_mover_offset"
+		geoMover = each.partition("_mover")[0] + "_geo_mover"
+		
+	    else:
+		offsetMover = each.partition("_mover")[0] + "_mover_offset_" + suffix
+		geoMover = each.partition("_mover")[0] + "_geo_mover_"	+ suffix
+		
+
+	    try:
+		#set the attr to 0
+		for mover in [globalMover, offsetMover, geoMover]:
+		    cmds.setAttr(mover + ".pinned", 0)
+		    
+		    
+		#lastly, change the font to italics
+		cmds.treeView(self.widgets["treeViewWidget"], edit = True, ff = [each, 0], ornament = [each, 0, 1, 7])
+		
+	    except:
+		pass
+	    
+	    
+	    #select previous selection
+	    if previousSelection != []:
+		cmds.select(previousSelection)
+		
+		
+		
+    # # # # # # # # # # # # # # # # # # # # # # # # # # # # # # # # # # # # # # # # # # # # # # # # # # # # # # # # # # # # # # # 
+    # # # # # # # # # # # # # # # # # # # # # # # # # # # # # # # # # # # # # # # # # # # # # # # # # # # # # # # # # # # # # # # 
+    # # # # # # # # # # # # # # # # # # # # # # # # # # # # # # # # # # # # # # # # # # # # # # # # # # # # # # # # # # # # # # # 
+    # # # # # # # # # # # # # # # # # # # # # # # # # # # # # # # # # # # # # # # # # # # # # # # # # # # # # # # # # # # # # # # 
+    # # # # # # # # # # # # # # # # # # # # # # # # # # # # # # # # # # # # # # # # # # # # # # # # # # # # # # # # # # # # # # # 
+    def listView_init(self, *args):
+	
+	selected = cmds.treeView(self.widgets["treeViewWidget"], q = True, ch = "")
+	suffix = None	
+	
+	for each in selected:
+	    #get names of global, offset and geo mover
+	    if each.partition("_mover")[2].find("_") == 0:
+		suffix = each.partition("_mover_")[2]
+	    
+	    globalMover = each
+	    
+	    if suffix == None:
+		offsetMover = each.partition("_mover")[0] + "_mover_offset"
+		geoMover = each.partition("_mover")[0] + "_geo_mover"
+		
+	    else:
+		offsetMover = each.partition("_mover")[0] + "_mover_offset_" + suffix
+		geoMover = each.partition("_mover")[0] + "_geo_mover_"	+ suffix
+		
+
+	    #get the pinned value
+	    for mover in [globalMover, offsetMover, geoMover]:
+		try:
+		    pinned = cmds.getAttr(mover + ".pinned")
+		    
+		    if pinned:
+			cmds.treeView(self.widgets["treeViewWidget"], edit = True, ff = [each, 2], ornament = [each, 1, 1, 7])
+			
+		except:
+		    pass
+		
+	
+	
\ No newline at end of file
