Adding ART scripts to repository

diff --git a/Engine/Extras/Maya_AnimationRiggingTools/MayaTools/General/Scripts/ART_spaceSwitcher.py b/Engine/Extras/Maya_AnimationRiggingTools/MayaTools/General/Scripts/ART_spaceSwitcher.py
new file mode 100644
--- /dev/null
+++ b/Engine/Extras/Maya_AnimationRiggingTools/MayaTools/General/Scripts/ART_spaceSwitcher.py
@@ -0,0 +1,984 @@
+import maya.cmds as cmds
+import maya.mel as mel
+from functools import partial
+import os, cPickle
+
+class SpaceSwitcher():
+    
+    def __init__(self, character, inst):
+        
+        #create class variables
+        self.widgets = {}
+	self.character = character
+        self.animUIInst = inst
+	
+        #find out which project we are in
+        references = cmds.ls(type = "reference")
+	for ref in references:
+	    try:
+		self.project = cmds.referenceQuery(ref, filename = True).rpartition("Projects/")[2].partition("/")[0]
+		
+	    except:
+		pass
+        
+        
+        #get access to our maya tools
+        toolsPath = cmds.internalVar(usd = True) + "mayaTools.txt"
+        if os.path.exists(toolsPath):
+            
+            f = open(toolsPath, 'r')
+            self.mayaToolsDir = f.readline()
+            f.close()
+            
+            
+        #check to see if window exists, if so, delete
+        if cmds.window("spaceSwitcherUI", exists = True):
+            cmds.deleteUI("spaceSwitcherUI")
+            
+        
+        #build window
+        self.widgets["window"] = cmds.window("spaceSwitcherUI", w = 450, h = 400, title = character + "_Space Switcher", sizeable = False)
+        
+        #create the main layout
+	self.widgets["mainLayout"] = cmds.columnLayout(w = 430)
+	self.widgets["scroll"] = cmds.scrollLayout(w = 440, h = 500, hst = -1, parent = self.widgets["mainLayout"])
+        self.widgets["topLevelLayout"] = cmds.columnLayout(w = 415, parent = self.widgets["scroll"])
+        
+        
+        #create the toolbar buttons
+        self.widgets["leftColumnLayout"] = cmds.rowColumnLayout(parent = self.widgets["topLevelLayout"], nc = 6, cw = [(1, 170), (2, 60), (3, 35), (4, 40), (5, 35), (6, 50)], cat = [(1, "both", 5),(2, "both", 5),(3, "both", 5),(4, "both", 5),(5, "both", 1), (6, "left", 20 )])
+	self.widgets["createSpaceButton"] = cmds.symbolButton(ann = "Select space object, then control, and create a new space", w = 125, h = 50, parent = self.widgets["leftColumnLayout"], image = self.mayaToolsDir + "/General/Icons/ART/spaceSwitch_create.bmp", c = self.createSpaceSwitcherSpace)
+        self.widgets["syncStatusButton"] = cmds.button("spaceSwitchSyncStatusButton", w = 50, h = 50, label = "!", ann = "Out of Sync!", bgc = [1, 0, 0], visible = False, c = self.relaunch)
+	
+	selectButton = cmds.symbolButton(ann = "Select space switch nodes for selected checkboxes", image = self.mayaToolsDir + "/General/Icons/ART/selectCurve.bmp", w = 30, h = 30, c = self.selectAllSpaceSwitchNodes)
+	bake = cmds.symbolButton(ann = "Bake Selected checkboxes down to default space", image = self.mayaToolsDir + "/General/Icons/ART/bake.bmp", w = 30, h = 30, c = self.bakeSpace)
+	transferButton = cmds.symbolButton(ann = "Push attributes from the space switch node down to the control on selected checkboxes", image = self.mayaToolsDir + "/General/Icons/ART/transfer.bmp", w = 30, h = 30, c = self.transferAttrs)
+	selectAllButton = cmds.symbolCheckBox(ann = "Select All Checkboxes", w = 30, h = 30,   oni = self.mayaToolsDir + "/General/Icons/ART/checkBox_on.bmp", ofi = self.mayaToolsDir + "/General/Icons/ART/checkBox.bmp", onc = partial(self.selectAll, True), ofc = partial(self.selectAll, False))
+	
+	#get the animation UI settings for match and match method and set the space switch UI to match those
+	self.matchVal = cmds.menuItem(self.animUIInst.widgets["spaceSwitch_MatchToggleCB"], q = True, cb = True)
+	self.matchMethodVal = cmds.menuItem(self.animUIInst.widgets["spaceSwitch_MatchMethodCB"], q = True, cb = True)
+	
+	
+	#create right click menu for bake button (bake all)
+	menu = cmds.popupMenu(b = 3, parent = bake)
+	cmds.menuItem(label = "Bake All", parent = menu, c = self.bakeAll)
+	
+	
+        #create frames for controls
+	#need to find all space switch nodes for the current character
+	cmds.select(character + ":*_space_switcher_follow")
+	nodes = cmds.ls(sl = True)
+	spaceSwitchers = []
+	for node in nodes:
+	    if node.find("invis") == -1:
+		spaceSwitchers.append(node)
+		
+	
+        for control in spaceSwitchers:
+            self.createSpaceInfoForControl(control)
+            
+        #show the window
+        cmds.showWindow(self.widgets["window"])
+	cmds.select(clear = True)
+	
+	#create scriptJob
+	self.createScriptJob()
+        
+    # # # # # # # # # # # # # # # # # # # # # # # # # # # # # # # # # # # # # # # # # # # # # # # # # # # # # # # # # # # # # # # 
+    # # # # # # # # # # # # # # # # # # # # # # # # # # # # # # # # # # # # # # # # # # # # # # # # # # # # # # # # # # # # # # # 
+    # # # # # # # # # # # # # # # # # # # # # # # # # # # # # # # # # # # # # # # # # # # # # # # # # # # # # # # # # # # # # # # 
+    # # # # # # # # # # # # # # # # # # # # # # # # # # # # # # # # # # # # # # # # # # # # # # # # # # # # # # # # # # # # # # # 
+    # # # # # # # # # # # # # # # # # # # # # # # # # # # # # # # # # # # # # # # # # # # # # # # # # # # # # # # # # # # # # # # 
+    def selectAll(self, select, *args):
+	
+	#selects all checkboxes
+	#need to find all space switch nodes for the current character
+	cmds.select(self.character + ":*_space_switcher_follow")
+	nodes = cmds.ls(sl = True)
+	spaceSwitchers = []
+	for node in nodes:
+	    if node.find("invis") == -1:
+		control = node.partition("_space_")[0]
+		spaceSwitchers.append(control)
+		
+	
+        for control in spaceSwitchers:
+	    if select:
+		cmds.symbolCheckBox(control + "_checkboxWidget", edit = True, v = True)	
+	    else:
+		cmds.symbolCheckBox(control + "_checkboxWidget", edit = True, v = False)	
+	
+    # # # # # # # # # # # # # # # # # # # # # # # # # # # # # # # # # # # # # # # # # # # # # # # # # # # # # # # # # # # # # # # 
+    # # # # # # # # # # # # # # # # # # # # # # # # # # # # # # # # # # # # # # # # # # # # # # # # # # # # # # # # # # # # # # # 
+    # # # # # # # # # # # # # # # # # # # # # # # # # # # # # # # # # # # # # # # # # # # # # # # # # # # # # # # # # # # # # # # 
+    # # # # # # # # # # # # # # # # # # # # # # # # # # # # # # # # # # # # # # # # # # # # # # # # # # # # # # # # # # # # # # # 
+    # # # # # # # # # # # # # # # # # # # # # # # # # # # # # # # # # # # # # # # # # # # # # # # # # # # # # # # # # # # # # # # 
+    def changeAnimInstMatch(self, *args):
+	
+	
+	#set the animInst to that value
+	cmds.menuItem(self.animUIInst.widgets["spaceSwitch_MatchToggleCB"], edit = True, cb = self.matchVal)
+	
+	
+    # # # # # # # # # # # # # # # # # # # # # # # # # # # # # # # # # # # # # # # # # # # # # # # # # # # # # # # # # # # # # # # 
+    # # # # # # # # # # # # # # # # # # # # # # # # # # # # # # # # # # # # # # # # # # # # # # # # # # # # # # # # # # # # # # # 
+    # # # # # # # # # # # # # # # # # # # # # # # # # # # # # # # # # # # # # # # # # # # # # # # # # # # # # # # # # # # # # # # 
+    # # # # # # # # # # # # # # # # # # # # # # # # # # # # # # # # # # # # # # # # # # # # # # # # # # # # # # # # # # # # # # # 
+    # # # # # # # # # # # # # # # # # # # # # # # # # # # # # # # # # # # # # # # # # # # # # # # # # # # # # # # # # # # # # # # 
+    def changeAnimInstMatchMethod(self, *args):
+	
+	
+	#set the animInst to that value
+	cmds.menuItem(self.animUIInst.widgets["spaceSwitch_MatchMethodCB"], edit = True, cb = self.matchVal)
+	
+	
+    # # # # # # # # # # # # # # # # # # # # # # # # # # # # # # # # # # # # # # # # # # # # # # # # # # # # # # # # # # # # # # # 
+    # # # # # # # # # # # # # # # # # # # # # # # # # # # # # # # # # # # # # # # # # # # # # # # # # # # # # # # # # # # # # # # 
+    # # # # # # # # # # # # # # # # # # # # # # # # # # # # # # # # # # # # # # # # # # # # # # # # # # # # # # # # # # # # # # # 
+    # # # # # # # # # # # # # # # # # # # # # # # # # # # # # # # # # # # # # # # # # # # # # # # # # # # # # # # # # # # # # # # 
+    # # # # # # # # # # # # # # # # # # # # # # # # # # # # # # # # # # # # # # # # # # # # # # # # # # # # # # # # # # # # # # # 
+    def createSpaceSwitcherSpace(self, *args):
+	
+	#launch an interface where the user can specify the control and the space object
+	if cmds.window("createSpaceWindow", exists = True):
+	    cmds.deleteUI("createSpaceWindow")
+	    
+	    
+	self.widgets["createSpaceWindow"] = cmds.window("createSpaceWindow", w = 250, h = 125, titleBarMenu = True, mxb = False, mnb = False, title = "Create Space", sizeable = False)
+	
+	self.widgets["createSpaceForm"] = cmds.formLayout(w = 250, h = 125, parent = self.widgets["createSpaceWindow"])
+	
+	
+	#create a textField with a button x2, and then a create space button
+	controlLabel = cmds.text(label = "         Control: ")
+	self.widgets["controlObjectField"] = cmds.textField(w = 120)
+	button = cmds.button(label = "<<", w = 30, c = partial(self.createSpaceUI_LoadObject, True, False))
+	
+	
+	spaceLabel = cmds.text(label = "Target Space: ")
+	self.widgets["spaceObjectField"] = cmds.textField(w = 120)
+	button2 = cmds.button(label = "<<", w = 30, c = partial(self.createSpaceUI_LoadObject, False, True))
+	
+	
+	createButton = cmds.button(w = 240, h = 40, label = "Create Space", c = self.createSpace)
+	
+	
+	
+	#place the widgets
+	cmds.formLayout(self.widgets["createSpaceForm"], edit = True, af = [(controlLabel, "top", 7), (controlLabel, "left", 5)])
+	cmds.formLayout(self.widgets["createSpaceForm"], edit = True, af = [(self.widgets["controlObjectField"], "top", 5), (self.widgets["controlObjectField"], "right", 40)])
+	cmds.formLayout(self.widgets["createSpaceForm"], edit = True, af = [(button, "top", 5), (button, "right", 5)])
+	
+	cmds.formLayout(self.widgets["createSpaceForm"], edit = True, af = [(spaceLabel, "top", 37), (spaceLabel, "left", 5)])
+	cmds.formLayout(self.widgets["createSpaceForm"], edit = True, af = [(self.widgets["spaceObjectField"], "top", 35), (self.widgets["spaceObjectField"], "right", 40)])
+	cmds.formLayout(self.widgets["createSpaceForm"], edit = True, af = [(button2, "top", 35), (button2, "right", 5)])
+	
+	cmds.formLayout(self.widgets["createSpaceForm"], edit = True, af = [(createButton, "bottom", 5), (createButton, "left", 5)])
+	
+	
+	#show the window
+	cmds.showWindow(self.widgets["createSpaceWindow"])
+	
+	
+    # # # # # # # # # # # # # # # # # # # # # # # # # # # # # # # # # # # # # # # # # # # # # # # # # # # # # # # # # # # # # # # 
+    # # # # # # # # # # # # # # # # # # # # # # # # # # # # # # # # # # # # # # # # # # # # # # # # # # # # # # # # # # # # # # # 
+    # # # # # # # # # # # # # # # # # # # # # # # # # # # # # # # # # # # # # # # # # # # # # # # # # # # # # # # # # # # # # # # 
+    # # # # # # # # # # # # # # # # # # # # # # # # # # # # # # # # # # # # # # # # # # # # # # # # # # # # # # # # # # # # # # # 
+    # # # # # # # # # # # # # # # # # # # # # # # # # # # # # # # # # # # # # # # # # # # # # # # # # # # # # # # # # # # # # # # 
+    def createSpaceUI_LoadObject(self, control, target, *args):
+	
+	if control:
+	    #check to make sure it is a valid control that can use space switching
+	    obj = cmds.ls(sl = True)
+	    if len(obj) > 0:
+		if cmds.objExists(obj[0] + "_space_switcher"):
+		    cmds.textField(self.widgets["controlObjectField"] , edit = True, text = obj[0])
+		    
+		else:
+		    cmds.confirmDialog(message = "The control you have selected is not a valid control for space switching.", title = "Create Space")
+		   
+	    else:
+		cmds.confirmDialog(message = "Nothing Selected!", title = "Create Space")
+		
+		
+	if target:
+	    obj = cmds.ls(sl = True)
+	    if len(obj) > 0:
+		cmds.textField(self.widgets["spaceObjectField"] , edit = True, text = obj[0])
+	       
+	    else:
+		cmds.confirmDialog(message = "Nothing Selected!", title = "Create Space")
+		
+	
+	
+	
+    # # # # # # # # # # # # # # # # # # # # # # # # # # # # # # # # # # # # # # # # # # # # # # # # # # # # # # # # # # # # # # # 
+    # # # # # # # # # # # # # # # # # # # # # # # # # # # # # # # # # # # # # # # # # # # # # # # # # # # # # # # # # # # # # # # 
+    # # # # # # # # # # # # # # # # # # # # # # # # # # # # # # # # # # # # # # # # # # # # # # # # # # # # # # # # # # # # # # # 
+    # # # # # # # # # # # # # # # # # # # # # # # # # # # # # # # # # # # # # # # # # # # # # # # # # # # # # # # # # # # # # # # 
+    # # # # # # # # # # # # # # # # # # # # # # # # # # # # # # # # # # # # # # # # # # # # # # # # # # # # # # # # # # # # # # # 
+    def createSpaceInfoForControl(self, control, *args):
+        
+        character = control.partition(":")[0]
+        niceName = control.partition(":")[2].partition("_space")[0]
+	parentSpace = cmds.listRelatives(control, parent = True)[0]
+	parentSpace = parentSpace.partition(":")[2]
+	    
+        
+        #create the 2 columns
+        self.widgets[control + "_spaceSwitch2Rows"] = cmds.columnLayout( parent = self.widgets["topLevelLayout"])
+        
+        #create the spaces layout
+	self.widgets[control + "_spaceSwitchSpacesLayout"] = cmds.rowColumnLayout(nc = 3,cw = [(1, 125), (2, 215), (3, 70), ], parent = self.widgets[control + "_spaceSwitch2Rows"], cat = [(1, "both", 5), (2, "both", 5), (3, "both", 10)])
+
+
+
+	#create the UI controls for each .space_ attr found on the space switch node
+	animControl = control.partition("_space_")[0]
+	spaceSwitchNode = control.rpartition("_follow")[0]
+	
+	#get space switch attrs from both control and space switch node
+	attrs = []
+	try:
+	    attrs.extend(cmds.listAttr(spaceSwitchNode, string = "space_*"))
+	except:
+	    pass
+	
+	try:
+	    attrs.extend(cmds.listAttr(animControl, string = "space_*"))
+	except:
+	    pass
+	
+	#column 1 is the text
+	cmds.text(label = niceName, parent = self.widgets[control + "_spaceSwitchSpacesLayout"])
+	
+	#column 2 is the option menu
+	self.widgets[control.partition(":")[2] + "_spacesOM"] = cmds.optionMenu(h = 30, w = 210, parent = self.widgets[control + "_spaceSwitchSpacesLayout"], cc = partial(self.switchSpace, control, False))
+	
+	#column 3 is the select button
+	selectButton = cmds.symbolCheckBox(animControl + "_checkboxWidget", w = 30, h = 30,   oni = self.mayaToolsDir + "/General/Icons/ART/checkBox_on.bmp", ofi = self.mayaToolsDir + "/General/Icons/ART/checkBox.bmp", parent = self.widgets[control + "_spaceSwitchSpacesLayout"])
+	
+	
+	#create the menuItem for the default space
+	cmds.menuItem(label = "default [" + parentSpace + "]", parent = self.widgets[control.partition(":")[2] + "_spacesOM"])
+	
+	
+	#create the menuItems for each attr in attrs
+	found = False
+	for attr in attrs:
+	    label = attr.partition("space_")[2]
+	    if cmds.objExists(spaceSwitchNode + "." + attr):
+		value = cmds.getAttr(spaceSwitchNode + "." + attr)
+	    if cmds.objExists(animControl + "." + attr):
+		value = cmds.getAttr(animControl + "." + attr)
+		
+	    if value == True:
+		found = True
+		
+		
+	    cmds.menuItem(ann = spaceSwitchNode, label = label, parent = self.widgets[control.partition(":")[2] + "_spacesOM"])
+	    
+	    if found:
+		cmds.optionMenu(self.widgets[control.partition(":")[2] + "_spacesOM"], edit = True, v = label)
+
+	if found == False:
+	    cmds.optionMenu(self.widgets[control.partition(":")[2] + "_spacesOM"], edit = True, v = "default [" + parentSpace + "]")
+	    
+
+	    
+	    
+
+	    
+
+    # # # # # # # # # # # # # # # # # # # # # # # # # # # # # # # # # # # # # # # # # # # # # # # # # # # # # # # # # # # # # # # 
+    # # # # # # # # # # # # # # # # # # # # # # # # # # # # # # # # # # # # # # # # # # # # # # # # # # # # # # # # # # # # # # # 
+    # # # # # # # # # # # # # # # # # # # # # # # # # # # # # # # # # # # # # # # # # # # # # # # # # # # # # # # # # # # # # # # 
+    # # # # # # # # # # # # # # # # # # # # # # # # # # # # # # # # # # # # # # # # # # # # # # # # # # # # # # # # # # # # # # # 
+    # # # # # # # # # # # # # # # # # # # # # # # # # # # # # # # # # # # # # # # # # # # # # # # # # # # # # # # # # # # # # # # 
+    def switchSpace(self, control, *args):
+	
+	#need to find space switch node
+	spaceSwitchNode = control.rpartition("_follow")[0]
+	menu = self.widgets[control.partition(":")[2] + "_spacesOM"]
+	selected = cmds.optionMenu(self.widgets[control.partition(":")[2] + "_spacesOM"], q = True, value = True)
+	attr = "space_" + selected
+	
+	
+	if selected == "default":
+	    self.switchToLocalSpace(spaceSwitchNode)
+	    
+	else:
+	    #SHOULD ADD OPTION FOR MATCHING OR NO MATCHING
+	    control = spaceSwitchNode.rpartition("_space")[0]
+	    matching = self.matchMethodVal
+	    matchToControl = self.matchVal
+	    
+	    
+	    #If match is turned on
+	    if matching:
+		#get current time, move time to a frame before, key, and key at current time
+		currentTime = cmds.currentTime(q = True)
+		
+		#set pre-frame key
+		if matchToControl == False:
+		    cmds.currentTime(currentTime - 1)
+		    cmds.setKeyframe(spaceSwitchNode)
+		    cmds.setKeyframe(control)
+		    
+		else:
+		    cmds.currentTime(currentTime - 1)
+		    cmds.setKeyframe(spaceSwitchNode)
+		    cmds.setKeyframe(control)
+		    
+		    loc = cmds.spaceLocator()[0]
+		    constraint = cmds.parentConstraint(control, loc)[0]
+		    cmds.delete(constraint)
+		    
+		    constraint = cmds.parentConstraint(loc, control)[0]
+		    cmds.setKeyframe(control, t = (currentTime - 1))		
+		    cmds.delete(constraint)
+		    cmds.delete(loc)
+		    
+
+		#create temp locator to snap the space switch node to
+		tempLoc = cmds.spaceLocator()[0]
+		cmds.currentTime(currentTime)
+		
+		#constrain temp loc
+		if matchToControl == False:
+		    constraint = cmds.parentConstraint(spaceSwitchNode, tempLoc)[0]
+		    
+		else:
+		    constraint = cmds.parentConstraint(control, tempLoc)[0]
+		    
+		cmds.delete(constraint)
+
+
+
+		#match and switch space
+		attrs = []
+		try:
+		    attrs.extend(cmds.listAttr(spaceSwitchNode, string = "space_*"))
+		except:
+		    pass
+		try:
+		    attrs.extend(cmds.listAttr(control, string = "space_"))
+		except:
+		    pass
+		
+		for attribute in attrs:
+		    if cmds.objExists(spaceSwitchNode + "." + attribute):
+			cmds.setAttr(spaceSwitchNode + "." + attribute, 0)
+		    if cmds.objExists(control + "." + attribute):
+			cmds.setAttr(control + "." + attribute, 0)
+		    
+
+		#set the desired space on
+		try:
+		    if cmds.objExists(spaceSwitchNode + "." + attr):
+			cmds.setAttr(spaceSwitchNode + "." + attr, 1)
+			cmds.setKeyframe(spaceSwitchNode, attribute = attr, t = currentTime)
+		    if cmds.objExists(control + "." + attr):
+			cmds.setAttr(control + "." + attr, 1)
+			cmds.setKeyframe(control, attribute = attr, t = currentTime)
+			
+		except:
+		    pass
+		
+	    
+	    
+		#find out the match method
+		if matchToControl == False:
+		    print "not matching to control"
+		    constraint = cmds.parentConstraint(tempLoc, spaceSwitchNode)[0]
+		    cmds.setKeyframe(spaceSwitchNode, t = currentTime)
+		    cmds.delete(constraint)
+		    cmds.delete(tempLoc)
+		    cmds.select(clear = True)
+		    
+		else:
+		    constraint = cmds.parentConstraint(tempLoc, control)[0]
+		    cmds.setKeyframe(control, t = cmds.currentTime(q = True))
+		    cmds.delete(constraint)
+		    cmds.delete(tempLoc)
+		    cmds.select(clear = True)
+		
+			    
+		    
+		    
+	    #not matching
+	    else:
+		
+		#get current time, move time to a frame before, key, and key at current time
+		currentTime = cmds.currentTime(q = True)
+		cmds.setKeyframe(spaceSwitchNode, t = (currentTime - 1))
+		cmds.setKeyframe(control, t = (currentTime - 1))
+		
+		
+		#need to look for space switch attrs on both the space switch node and the control
+		attrs = []
+		try:
+		    attrs.extend(cmds.listAttr(spaceSwitchNode, string = "space_*"))
+		except:
+		    pass
+		
+		try:
+		    attrs.extend(cmds.listAttr(control, string = "space_"))
+		except:
+		    pass
+		
+		#zero out all spaces
+		for attribute in attrs:
+		    if cmds.objExists(spaceSwitchNode + "." + attribute):
+			cmds.setAttr(spaceSwitchNode+ "." + attribute, 0)
+		    if cmds.objExists(control + "." + attribute):
+			cmds.setAttr(control+ "." + attribute, 0)		    
+		    
+		    
+		#set the desired space
+		try:
+		    if cmds.objExists(spaceSwitchNode + "." + attr):
+			cmds.setAttr(spaceSwitchNode + "." + attr, 1)
+			cmds.setKeyframe(spaceSwitchNode, attribute = attr, t = currentTime)
+			
+		    if cmds.objExists(control + "." + attr):
+			cmds.setAttr(control + "." + attr, 1)
+			cmds.setKeyframe(control, attribute = attr, t = currentTime)		    
+		except:
+		    pass
+		
+		
+	#set the current time back
+	cmds.currentTime(currentTime)
+	
+
+    # # # # # # # # # # # # # # # # # # # # # # # # # # # # # # # # # # # # # # # # # # # # # # # # # # # # # # # # # # # # # # # 
+    # # # # # # # # # # # # # # # # # # # # # # # # # # # # # # # # # # # # # # # # # # # # # # # # # # # # # # # # # # # # # # # 
+    # # # # # # # # # # # # # # # # # # # # # # # # # # # # # # # # # # # # # # # # # # # # # # # # # # # # # # # # # # # # # # # 
+    # # # # # # # # # # # # # # # # # # # # # # # # # # # # # # # # # # # # # # # # # # # # # # # # # # # # # # # # # # # # # # # 
+    # # # # # # # # # # # # # # # # # # # # # # # # # # # # # # # # # # # # # # # # # # # # # # # # # # # # # # # # # # # # # # # 
+    def switchToLocalSpace(self, spaceSwitchNode, *args):
+	#SHOULD ADD OPTION FOR MATCHING OR NO MATCHING
+	matching = self.matchMethodVal
+	matchToControl = self.matchVal
+	
+	if matching:
+	    #get current time, move time to a frame before, key, and key at current time
+	    control = spaceSwitchNode.rpartition("_space")[0]
+	    currentTime = cmds.currentTime(q = True)
+	    
+	    #set pre-frame key
+	    if matchToControl == False:
+		cmds.setKeyframe(spaceSwitchNode, t = (currentTime - 1))
+		
+	    else:
+		cmds.setKeyframe(spaceSwitchNode, t = (currentTime - 1))
+		cmds.currentTime(currentTime - 1)
+		loc = cmds.spaceLocator()[0]
+		constraint = cmds.parentConstraint(control, loc)[0]
+		cmds.delete(constraint)
+		
+		constraint = cmds.parentConstraint(loc, control)[0]
+		cmds.setKeyframe(control, t = (cmds.currentTime(q = True)))
+		cmds.delete(constraint)
+		cmds.delete(loc)	    
+
+
+    
+	    #create temp locator to snap the space switch node to
+	    tempLoc = cmds.spaceLocator()
+	    cmds.currentTime(currentTime)
+	    
+	    #constrain temp loc
+	    if matchToControl == False:
+		constraint = cmds.parentConstraint(spaceSwitchNode, tempLoc)[0]
+	    else:
+		constraint = cmds.parentConstraint(control, tempLoc)[0]
+	    cmds.delete(constraint)	    
+
+	    
+	    #get space switch attrs from both control and space switch node
+	    attrs = []
+	    try:
+		attrs.extend(cmds.listAttr(spaceSwitchNode, string = "space_*"))
+	    except:
+		pass
+	    
+	    try:
+		attrs.extend(cmds.listAttr(control, string = "space_"))
+	    except:
+		pass
+	    
+	    #zero out all space switch attrs
+	    for attr in attrs:
+		if cmds.objExists(spaceSwitchNode + "." + attr):
+		    cmds.setAttr(spaceSwitchNode+ "." + attr, 0)
+		if cmds.objExists(control + "." + attr):
+		    cmds.setAttr(control+ "." + attr, 0)		
+	    
+	    #match
+	    if matchToControl == False:
+		constraint = cmds.parentConstraint(tempLoc[0], spaceSwitchNode)[0]
+		cmds.setKeyframe(spaceSwitchNode, t = currentTime)
+		cmds.setKeyframe(control, t = currentTime)
+		cmds.delete(constraint)
+		cmds.delete(tempLoc)
+		cmds.select(clear = True)
+
+	    else:
+		constraint = cmds.parentConstraint(tempLoc[0], control)[0]
+		cmds.setKeyframe(control, t = currentTime)
+		cmds.delete(constraint)
+		cmds.delete(tempLoc)
+		cmds.select(clear = True)	
+
+
+	else:
+	    #get current time, move time to a frame before, key, and key at current time
+	    currentTime = cmds.currentTime(q = True)
+	    cmds.setKeyframe(spaceSwitchNode, t = (currentTime - 1))
+	    
+	    #get space switch attrs from both control and space switch node
+	    attrs = []
+	    try:
+		attrs.extend(cmds.listAttr(spaceSwitchNode, string = "space_*"))
+	    except:
+		pass
+		
+	    try:
+		attrs.extend(cmds.listAttr(control, string = "space_"))
+	    except:
+		pass
+	    
+	    #zero out spaces
+	    for attr in attrs:
+		if cmds.objExists(spaceSwitchNode + "." + attr):
+		    cmds.setAttr(spaceSwitchNode+ "." + attr, 0)
+		if cmds.objExists(control + "." + attr):
+		    cmds.setAttr(control+ "." + attr, 0)
+		
+	    cmds.setKeyframe(spaceSwitchNode, t = currentTime)
+	    cmds.setKeyframe(control, t = currentTime)
+	    
+    # # # # # # # # # # # # # # # # # # # # # # # # # # # # # # # # # # # # # # # # # # # # # # # # # # # # # # # # # # # # # # # 
+    # # # # # # # # # # # # # # # # # # # # # # # # # # # # # # # # # # # # # # # # # # # # # # # # # # # # # # # # # # # # # # # 
+    # # # # # # # # # # # # # # # # # # # # # # # # # # # # # # # # # # # # # # # # # # # # # # # # # # # # # # # # # # # # # # # 
+    # # # # # # # # # # # # # # # # # # # # # # # # # # # # # # # # # # # # # # # # # # # # # # # # # # # # # # # # # # # # # # # 
+    # # # # # # # # # # # # # # # # # # # # # # # # # # # # # # # # # # # # # # # # # # # # # # # # # # # # # # # # # # # # # # # 
+    def createSpace(self, *args):
+	
+
+	control = cmds.textField(self.widgets["controlObjectField"], q = True, text = True)
+	spaceObj = cmds.textField(self.widgets["spaceObjectField"], q = True, text = True)
+	ann = spaceObj
+	niceName = spaceObj
+	
+	if spaceObj.find(":") != -1:
+	    niceName = spaceObj.partition(":")[2]
+
+	character = control.partition(":")[0]
+	spaceSwitchFollow = control + "_space_switcher_follow"
+	spaceSwitchNode = control + "_space_switcher"
+	
+	
+	#add attr to the space switcher node
+	cmds.select(control)
+	cmds.addAttr(ln = "space_" + niceName, minValue = 0, maxValue = 1, dv = 0, keyable = True)
+	
+	#add constraint to the new object on the follow node
+	constraint = cmds.parentConstraint(spaceObj, spaceSwitchFollow, mo = True)[0]
+	
+	#hook up connections
+	targets = cmds.parentConstraint(constraint, q = True, targetList = True)
+	weight = 0
+	for i in range(int(len(targets))):
+	    if targets[i].find(spaceObj) != -1:
+		weight = i
+		
+	cmds.connectAttr(control + ".space_" + niceName, constraint + "." + niceName + "W" + str(weight))
+	
+	#lockNode on space object so it cannot be deleted by the user (if node is not a referenced node)
+	if spaceObj.find(":") == -1:
+	    cmds.lockNode(spaceObj)
+	    
+	
+	#create button
+	try:
+	    cmds.menuItem(ann = spaceSwitchNode, label = niceName, parent = self.widgets[control.partition(":")[2] + "_space_switcher_follow_spacesOM"])
+	except:
+	    pass
+	
+	
+	if cmds.window("createSpaceWindow", exists = True):
+	    cmds.deleteUI("createSpaceWindow")
+	    
+	    
+    # # # # # # # # # # # # # # # # # # # # # # # # # # # # # # # # # # # # # # # # # # # # # # # # # # # # # # # # # # # # # # # 
+    # # # # # # # # # # # # # # # # # # # # # # # # # # # # # # # # # # # # # # # # # # # # # # # # # # # # # # # # # # # # # # # 
+    # # # # # # # # # # # # # # # # # # # # # # # # # # # # # # # # # # # # # # # # # # # # # # # # # # # # # # # # # # # # # # # 
+    # # # # # # # # # # # # # # # # # # # # # # # # # # # # # # # # # # # # # # # # # # # # # # # # # # # # # # # # # # # # # # # 
+    # # # # # # # # # # # # # # # # # # # # # # # # # # # # # # # # # # # # # # # # # # # # # # # # # # # # # # # # # # # # # # # 
+    def getCharacters(self):
+        referenceNodes = []
+        references = cmds.ls(type = "reference")
+        
+        for reference in references:
+            niceName = reference.rpartition("RN")[0]
+            suffix = reference.rpartition("RN")[2]
+            if suffix != "":
+                if cmds.objExists(niceName + suffix + ":" + "Skeleton_Settings"):
+                    referenceNodes.append(niceName + suffix)
+                    
+            else:
+                if cmds.objExists(niceName + ":" + "Skeleton_Settings"):
+                    referenceNodes.append(niceName)
+                    
+        return referenceNodes
+    
+    # # # # # # # # # # # # # # # # # # # # # # # # # # # # # # # # # # # # # # # # # # # # # # # # # # # # # # # # # # # # # # # 
+    # # # # # # # # # # # # # # # # # # # # # # # # # # # # # # # # # # # # # # # # # # # # # # # # # # # # # # # # # # # # # # # 
+    # # # # # # # # # # # # # # # # # # # # # # # # # # # # # # # # # # # # # # # # # # # # # # # # # # # # # # # # # # # # # # # 
+    # # # # # # # # # # # # # # # # # # # # # # # # # # # # # # # # # # # # # # # # # # # # # # # # # # # # # # # # # # # # # # # 
+    # # # # # # # # # # # # # # # # # # # # # # # # # # # # # # # # # # # # # # # # # # # # # # # # # # # # # # # # # # # # # # # 
+    def bakeAll(self, *args):
+	
+	cmds.select(self.character + ":*_space_switcher_follow")
+	nodes = cmds.ls(sl = True)
+	spaceSwitchers = []
+	bakeNodes = []
+	
+	for node in nodes:
+	    if node.find("invis") == -1:
+		if node.find(self.character + ":spine_01_space_switcher") == -1:
+		    spaceSwitchers.append(node)
+		    
+	cmds.select(clear = True)
+	
+	
+	#get the value of the bake toggle
+	for control in spaceSwitchers:
+	    cmds.iconTextCheckBox(self.widgets[control + "_bakeToggle"], edit = True, v = True)
+	    
+	    
+	self.bakeSpace()
+
+		
+		
+	
+	
+    # # # # # # # # # # # # # # # # # # # # # # # # # # # # # # # # # # # # # # # # # # # # # # # # # # # # # # # # # # # # # # # 
+    # # # # # # # # # # # # # # # # # # # # # # # # # # # # # # # # # # # # # # # # # # # # # # # # # # # # # # # # # # # # # # # 
+    # # # # # # # # # # # # # # # # # # # # # # # # # # # # # # # # # # # # # # # # # # # # # # # # # # # # # # # # # # # # # # # 
+    # # # # # # # # # # # # # # # # # # # # # # # # # # # # # # # # # # # # # # # # # # # # # # # # # # # # # # # # # # # # # # # 
+    # # # # # # # # # # # # # # # # # # # # # # # # # # # # # # # # # # # # # # # # # # # # # # # # # # # # # # # # # # # # # # # 
+    def bakeSpace(self, *args):
+	
+	cmds.select(self.character + ":*_space_switcher_follow")
+	nodes = cmds.ls(sl = True)
+	spaceSwitchers = []
+	bakeNodes = []
+	
+	for node in nodes:
+	    if node.find("invis") == -1:
+		if node.find(self.character + ":spine_01_space_switcher") == -1:
+		    spaceSwitchers.append(node)
+		    
+	cmds.select(clear = True)
+	
+	
+	#get the value of the bake toggle
+	for control in spaceSwitchers:
+	    animControl = control.partition("_space_")[0]
+	    value = cmds.symbolCheckBox(animControl + "_checkboxWidget", q = True, v = True)
+	    if value:
+		bakeNodes.append(control)
+		
+	    
+	if bakeNodes == []:
+	    cmds.confirmDialog(title = "Bake Selected", icon = "warning", message = "No controls were selected for bake! Please use the bake icon next to the controls you want to bake down to toggle them for baking.")
+	    return
+	
+	#now we have a list of nodes to bake
+	locs = []
+	constraints = []
+	for node in bakeNodes:
+	    #create a locator constrained to the space switch node
+	    spaceSwitchNode = node.rpartition("_follow")[0]
+	    control = spaceSwitchNode.partition("_space_switcher")[0]
+	    loc = cmds.spaceLocator(name = spaceSwitchNode + "_loc")[0]
+	    locs.append(loc)
+	    constraint = cmds.parentConstraint(control, loc)[0]
+	    constraints.append(constraint)
+	    
+	#bake all of the locators
+	for loc in locs:
+	    cmds.select(loc, add = True)
+	    
+	start = cmds.playbackOptions(q = True, min = True)
+	end = cmds.playbackOptions(q = True, max = True)
+	cmds.bakeResults(simulation = True, t = (start, end))
+	cmds.delete(constraints)
+	cmds.select(clear = True)
+	
+	
+	#delete keys on all space switch attrs 
+	for node in bakeNodes:
+	    spaceSwitchNode = node.rpartition("_follow")[0]
+	    
+	    attrs = []
+	    try:
+		attrs.extend(cmds.listAttr(spaceSwitchNode, string = "space_*"))
+	    except:
+		pass
+	    
+	    try:
+		attrs.extend(cmds.listAttr(control, string = "space_*"))
+	    except:
+		pass
+	    
+	    
+	    
+	    
+	    for attr in attrs:
+		#zero out all attrs on the space switcher
+		if cmds.objExists(spaceSwitchNode + "." + attr):
+		    cmds.cutKey(spaceSwitchNode, at = attr)
+		    cmds.setAttr(spaceSwitchNode + "." + attr, 0)
+	    
+		if cmds.objExists(control + "." + attr):
+		    cmds.cutKey(control, at = attr)
+		    cmds.setAttr(control + "." + attr, 0)	
+		    
+		
+	#reverse the constraint so now the spaceSwitchNode will follow the loc
+	for node in bakeNodes:
+	    spaceSwitchNode = node.rpartition("_follow")[0]
+	    control = node.rpartition("_space_switcher")[0]
+	    loc = spaceSwitchNode + "_loc"
+	    
+	    
+	    if control.find(self.character + ":ik_elbow") == 0:
+		constraint = cmds.pointConstraint(loc, control)[0]
+		
+	    else:
+		constraint = cmds.parentConstraint(loc, control)[0]
+		
+	    cmds.select(control, add = True)
+	    
+	    for attr in [".tx", ".ty", ".tz", ".rx", ".ry", ".rz"]:
+		cmds.setAttr(spaceSwitchNode + attr, 0)
+		
+	    
+	
+	cmds.bakeResults(simulation = True, t = (start, end))
+	cmds.delete(locs)
+	
+	#delete blendParent attr if exists
+	if cmds.objExists(control + ".blendParent1"):
+	    try:
+		cmds.deleteAttr(n = control, at = "blendParent1")
+	    except:
+		pass
+	    
+	    
+	
+    # # # # # # # # # # # # # # # # # # # # # # # # # # # # # # # # # # # # # # # # # # # # # # # # # # # # # # # # # # # # # # # 
+    # # # # # # # # # # # # # # # # # # # # # # # # # # # # # # # # # # # # # # # # # # # # # # # # # # # # # # # # # # # # # # # 
+    # # # # # # # # # # # # # # # # # # # # # # # # # # # # # # # # # # # # # # # # # # # # # # # # # # # # # # # # # # # # # # # 
+    # # # # # # # # # # # # # # # # # # # # # # # # # # # # # # # # # # # # # # # # # # # # # # # # # # # # # # # # # # # # # # # 
+    # # # # # # # # # # # # # # # # # # # # # # # # # # # # # # # # # # # # # # # # # # # # # # # # # # # # # # # # # # # # # # # 
+    def transferAttrs(self, *args):
+	
+	#get selected checkboxes
+	cmds.select(self.character + ":*_space_switcher_follow")
+	nodes = cmds.ls(sl = True)
+	spaceSwitchers = []
+	for node in nodes:
+	    if node.find("invis") == -1:
+		spaceSwitchers.append(node)
+		
+	
+	selectNodes = []
+        for node in spaceSwitchers:
+	    control = node.partition("_space")[0]
+	    spaceSwitchNode = node.partition("_follow")[0]
+	    
+            value = cmds.symbolCheckBox(control + "_checkboxWidget", q = True, value = True)
+	    if value:
+		selectNodes.append(spaceSwitchNode)
+		
+			
+	node = ""
+	for node in selectNodes:
+	    control = node.partition("_space")[0]
+	    
+	    #need to add attr that is on spaceSwitchNode to control
+	    attrs = []
+	    try:
+		attrs.extend((cmds.listAttr(node, string = "space_*")))
+	    except:
+		pass
+	    
+	    #get the target list from the constraint
+	    constraint = cmds.listConnections(node, type = "parentConstraint")
+	    if constraint != None:
+		targets = cmds.parentConstraint(constraint[0], q = True, weightAliasList = True)	    
+
+	    #we now have a list of attrs and targets that should match
+
+	    if attrs != None:
+		for attr in reversed(attrs):
+		    cmds.addAttr(control, ln = attr, dv = 0, min = 0, max = 1, keyable = True)
+		    
+		    #setup the new constraint connection
+		    targetAttr = None
+		    for target in targets:
+			try:
+			    cmds.disconnectAttr(node + "." + attr, constraint[0] + "." + target)
+			    targetAttr = target
+			    
+			except:
+			    pass
+			
+		    try:
+			cmds.connectAttr(control + "." + attr, constraint[0] + "." + targetAttr)
+		    except:
+			cmds.warning("Could not setup connections on " + control)
+		    
+		    #copy keys from original attr to new attr
+		    start = cmds.findKeyframe(node, at = attr, which = "first")
+		    end = cmds.findKeyframe(node, at = attr, which = "last")
+		    keys = cmds.copyKey(node, time=(start, end), attribute = attr, option = "curve")
+		    
+		    try:
+			cmds.pasteKey(control,  attribute = attr)
+		    except:
+			pass
+		    
+		    #delete spaceSwitchNode attr
+		    cmds.deleteAttr(node, at = attr)
+		    
+	
+	
+    # # # # # # # # # # # # # # # # # # # # # # # # # # # # # # # # # # # # # # # # # # # # # # # # # # # # # # # # # # # # # # # 
+    # # # # # # # # # # # # # # # # # # # # # # # # # # # # # # # # # # # # # # # # # # # # # # # # # # # # # # # # # # # # # # # 
+    # # # # # # # # # # # # # # # # # # # # # # # # # # # # # # # # # # # # # # # # # # # # # # # # # # # # # # # # # # # # # # # 
+    # # # # # # # # # # # # # # # # # # # # # # # # # # # # # # # # # # # # # # # # # # # # # # # # # # # # # # # # # # # # # # # 
+    # # # # # # # # # # # # # # # # # # # # # # # # # # # # # # # # # # # # # # # # # # # # # # # # # # # # # # # # # # # # # # # 
+    def selectAllSpaceSwitchNodes(self, *args):
+	
+    
+	#need to find all space switch nodes for the current character
+	cmds.select(self.character + ":*_space_switcher_follow")
+	nodes = cmds.ls(sl = True)
+	spaceSwitchers = []
+	for node in nodes:
+	    if node.find("invis") == -1:
+		spaceSwitchers.append(node)
+		
+	
+	selectNodes = []
+        for node in spaceSwitchers:
+	    control = node.partition("_space")[0]
+	    spaceSwitchNode = node.partition("_follow")[0]
+	    
+            value = cmds.symbolCheckBox(control + "_checkboxWidget", q = True, value = True)
+	    if value:
+		selectNodes.append(spaceSwitchNode)
+	    
+	cmds.select(clear = True)
+	for node in selectNodes:
+	    cmds.select(node, add = True)
+	    
+	    
+
+
+    # # # # # # # # # # # # # # # # # # # # # # # # # # # # # # # # # # # # # # # # # # # # # # # # # # # # # # # # # # # # # # # 
+    # # # # # # # # # # # # # # # # # # # # # # # # # # # # # # # # # # # # # # # # # # # # # # # # # # # # # # # # # # # # # # # 
+    # # # # # # # # # # # # # # # # # # # # # # # # # # # # # # # # # # # # # # # # # # # # # # # # # # # # # # # # # # # # # # # 
+    # # # # # # # # # # # # # # # # # # # # # # # # # # # # # # # # # # # # # # # # # # # # # # # # # # # # # # # # # # # # # # # 
+    # # # # # # # # # # # # # # # # # # # # # # # # # # # # # # # # # # # # # # # # # # # # # # # # # # # # # # # # # # # # # # # 
+    def relaunch(self,*args):
+	import ART_animationUI
+	reload(ART_animationUI)
+	self.animUIInst.spaceSwitcher()
+	
+	
+	
+    # # # # # # # # # # # # # # # # # # # # # # # # # # # # # # # # # # # # # # # # # # # # # # # # # # # # # # # # # # # # # # # 
+    # # # # # # # # # # # # # # # # # # # # # # # # # # # # # # # # # # # # # # # # # # # # # # # # # # # # # # # # # # # # # # # 
+    # # # # # # # # # # # # # # # # # # # # # # # # # # # # # # # # # # # # # # # # # # # # # # # # # # # # # # # # # # # # # # # 
+    # # # # # # # # # # # # # # # # # # # # # # # # # # # # # # # # # # # # # # # # # # # # # # # # # # # # # # # # # # # # # # # 
+    # # # # # # # # # # # # # # # # # # # # # # # # # # # # # # # # # # # # # # # # # # # # # # # # # # # # # # # # # # # # # # # 
+    def createScriptJob(self,*args):
+	
+	#grab space switch nodes to pass into script job
+	cmds.select(self.character + ":*_space_switcher_follow")
+	nodes = cmds.ls(sl = True)
+	spaceSwitchers = []
+	for node in nodes:
+	    if node.find("invis") == -1:
+		spaceSwitchers.append(node)
+		
+		
+	#create a scriptJob that updates the space switch UI when time is changed
+	self.scriptJob = cmds.scriptJob(event = ["timeChanged", partial(self.updateUI, spaceSwitchers)], parent = self.widgets["window"])
+	
+	cmds.select(clear = True)
+    
+    # # # # # # # # # # # # # # # # # # # # # # # # # # # # # # # # # # # # # # # # # # # # # # # # # # # # # # # # # # # # # # # 
+    # # # # # # # # # # # # # # # # # # # # # # # # # # # # # # # # # # # # # # # # # # # # # # # # # # # # # # # # # # # # # # # 
+    # # # # # # # # # # # # # # # # # # # # # # # # # # # # # # # # # # # # # # # # # # # # # # # # # # # # # # # # # # # # # # # 
+    # # # # # # # # # # # # # # # # # # # # # # # # # # # # # # # # # # # # # # # # # # # # # # # # # # # # # # # # # # # # # # # 
+    # # # # # # # # # # # # # # # # # # # # # # # # # # # # # # # # # # # # # # # # # # # # # # # # # # # # # # # # # # # # # # # 
+    def updateUI(self, spaceSwitchers, *args):
+	
+	#update the UI when time slider changes
+	for node in spaceSwitchers:
+	    spaceSwitchNode = node.rpartition("_follow")[0]
+	    control = spaceSwitchNode.partition("_space")[0]
+	    attrs = []
+	    
+	    try:
+		attrs.extend(cmds.listAttr(spaceSwitchNode, string = "space_*"))
+	    except:
+		pass
+	    
+	    try:
+		attrs.extend(cmds.listAttr(control, string = "space_*"))
+	    except:
+		pass
+	    
+	    
+	    #find the children of the radio collection for each node
+	    options = cmds.optionMenu(self.widgets[node.partition(":")[2] + "_spacesOM"], q = True, itemListLong = True)
+	    
+	    #find the annotations of each button, which will give us the space obj
+	    labels = []
+	    for option in options:
+		ann = cmds.menuItem(option, q = True, label = True)
+		labels.append(ann)
+		
+		
+	    
+	    #find the attr that is currently selected
+	    selected = ""
+	    for attr in attrs:
+		label = attr.partition("space_")[2]
+		
+		if cmds.objExists(spaceSwitchNode + "." + attr):
+		    value = cmds.getAttr(spaceSwitchNode + "." + attr)
+		    
+		if cmds.objExists(control + "." + attr):
+		    value = cmds.getAttr(control + "." + attr)
+		    
+		if value == True:
+		    selected = label
+		    
+		if selected == "":
+		    selected = "default"
+
+	    for label in labels:
+		if label.find(selected) == 0:
+		    cmds.optionMenu(self.widgets[node.partition(":")[2] + "_spacesOM"], edit= True, value = label)
