Adding ART scripts to repository

diff --git a/Engine/Extras/Maya_AnimationRiggingTools/MayaTools/General/Scripts/Modules/ART_Arm.py b/Engine/Extras/Maya_AnimationRiggingTools/MayaTools/General/Scripts/Modules/ART_Arm.py
new file mode 100644
--- /dev/null
+++ b/Engine/Extras/Maya_AnimationRiggingTools/MayaTools/General/Scripts/Modules/ART_Arm.py
@@ -0,0 +1,1624 @@
+import maya.cmds as cmds
+import maya.mel as mel
+import os
+import ART_rigUtils as utils
+reload(utils)
+
+#The arm module, added as an extra rig module, needs a name, a prefix or suffix (but not both), and a parent to get started
+
+class Arm():
+    
+    #pass in items we need to know about (do we build a clavicle? What is the root name of the module?(If set to None, just use upperarm, lowerarm, etc), 
+    #what is the prefix/suffix?, what is the parent on the arm?
+    def __init__(self, includeClavicle, name, prefix, suffix, armGrpParent, color, createBodyOrient):
+        
+        self.name = name
+	
+        if name != "":
+            self.clavicle = "clavicle_" + name
+            self.upperArm = "upperarm_" + name
+            self.lowerArm = "lowerarm_" + name
+            self.hand = "hand_" + name
+            
+        if name == "":
+	    self.clavicle = "clavicle"
+            self.upperArm = "upperarm"
+            self.lowerArm = "lowerarm"
+            self.hand = "hand"
+            
+        if prefix != None:
+            self.prefix = prefix + "_"
+        
+        if prefix == None:
+            self.prefix = ""
+            
+        if suffix != None:
+            self.suffix = "_" + suffix
+        
+        if suffix == None:
+            self.suffix == ""
+            
+        """
+        Note!
+        Arms will always be created with upperarm, lowerarm, hand, etc. Naming will look like "fk_" + self.prefix + self.upperArm + self.suffix ("fk_upperarm_l") self.prefix was ""
+        """
+	#Create some groups for the arm rigs to live in
+	if not cmds.objExists("arm_sys_grp"):
+	    armSysGrp = cmds.group(empty = True, name = "arm_sys_grp")
+	    cmds.parent(armSysGrp, "ctrl_rig")
+	
+      
+
+        #Build the FK Arm Rig
+	fkArmNodes = self.fkArmRig(color, armGrpParent, createBodyOrient)
+	
+        #Build the IK Arm Rig
+	ikArmNodes = self.ikArmRig(fkArmNodes[0], fkArmNodes[1], fkArmNodes[2], color)
+        
+        #If includeClavicle == True, build the FK and IK clavicle rigs
+	if includeClavicle:
+	    #setup the clavicle rigs
+	    clavicleJoint = self.clavicleRig(fkArmNodes[0], ikArmNodes[0], fkArmNodes[1], fkArmNodes[4], armGrpParent, color)
+	    
+        #Build Upper Arm twist rig (if applicable)
+	if cmds.objExists("driver_" + self.prefix + self.upperArm + "_twist_01" + self.suffix):
+	    if includeClavicle:
+		self.upperArmTwistRig(color, clavicleJoint, fkArmNodes[0], ikArmNodes[0])
+	    else:
+		self.upperArmTwistRig(color, armGrpParent, fkArmNodes[0], ikArmNodes[0])
+		
+	    
+        #Build Lower Arm twist rig (if applicable)
+	if cmds.objExists("driver_" + self.prefix + self.lowerArm + "_twist_01" + self.suffix):
+	    self.lowerArmTwistRig(color, armGrpParent, fkArmNodes[0], ikArmNodes[0])
+	    
+	    
+    # # # # # # # # # # # # # # # # # # # # # # # # # # # # # # # # # # # # # # # # # # # # # # # # # # # # # # # # # # # # # # # 
+    # # # # # # # # # # # # # # # # # # # # # # # # # # # # # # # # # # # # # # # # # # # # # # # # # # # # # # # # # # # # # # # 
+    # # # # # # # # # # # # # # # # # # # # # # # # # # # # # # # # # # # # # # # # # # # # # # # # # # # # # # # # # # # # # # # 
+    # # # # # # # # # # # # # # # # # # # # # # # # # # # # # # # # # # # # # # # # # # # # # # # # # # # # # # # # # # # # # # # 
+    # # # # # # # # # # # # # # # # # # # # # # # # # # # # # # # # # # # # # # # # # # # # # # # # # # # # # # # # # # # # # # # 
+    def fkArmRig(self, color, armGrpParent, createBodyOrient):
+	
+	#create a proper name out of prefix/name/suffix
+	name = self.prefix + self.name + self.suffix
+	if name.find("_") == 0:
+	    name = name.partition("_")[2]
+	
+	#check to make sure a driver joint for the arm exists
+	if cmds.objExists("driver_" + self.prefix + self.upperArm + self.suffix):
+	    
+	    #create the fk joint chain
+	    fkArmJoint = cmds.duplicate("driver_" + self.prefix + self.upperArm + self.suffix, po = True, name = "fk_" + self.prefix + self.upperArm + self.suffix)[0]
+	    fkElbowJoint = cmds.duplicate("driver_" + self.prefix + self.lowerArm + self.suffix, po = True, name = "fk_" + self.prefix + self.lowerArm + self.suffix)[0]
+	    fkWristJoint = cmds.duplicate("driver_" + self.prefix + self.hand + self.suffix, po = True, name = "fk_" + self.prefix + self.hand + self.suffix)[0]
+			
+	    #parent the fk upperarm to the world
+	    parent = cmds.listRelatives(fkArmJoint, parent = True)
+	    if parent != None:
+		cmds.parent(fkArmJoint, world = True)
+			    
+	    #recreate the fk arm hierarchy
+	    cmds.parent(fkElbowJoint, fkArmJoint)
+	    cmds.parent(fkWristJoint, fkElbowJoint)
+	    cmds.makeIdentity(fkArmJoint, t = 0, r = 1, s = 0, apply = True)
+			
+	    #set rotation order on fk arm joint
+	    cmds.setAttr(fkArmJoint + ".rotateOrder", 3)
+	    
+	    
+	    #create the upperarm orient locator
+	    armOrient = cmds.spaceLocator(name = "arm_orient_loc_" + name)[0]
+	    armOrientGrp = cmds.group(empty = True, name = "arm_orient_loc_grp_" + name)
+			
+	    constraint = cmds.parentConstraint(fkArmJoint, armOrient)[0]
+	    cmds.delete(constraint)
+	    
+	    constraint = cmds.parentConstraint(fkArmJoint, armOrientGrp)[0]
+	    cmds.delete(constraint)
+	    
+	    cmds.parent(armOrient, armOrientGrp)
+	    
+	
+	    #create the fk control curve 
+	    fkArmCtrl = utils.createControl("circle", 20, "fk_" + self.prefix + "arm" + self.suffix + "_anim")
+	    cmds.setAttr(fkArmCtrl + ".ry", -90)
+	    cmds.makeIdentity(fkArmCtrl, r = 1, apply =True)
+
+	    constraint = cmds.parentConstraint(fkArmJoint, fkArmCtrl)[0]
+	    cmds.delete(constraint)
+	    
+	    fkArmCtrlGrp = cmds.group(empty = True, name = "fk_" + self.prefix + "arm" + self.suffix + "_anim_grp")
+	    constraint = cmds.parentConstraint(fkArmJoint, fkArmCtrlGrp)[0]
+	    cmds.delete(constraint)
+	    cmds.parent(fkArmCtrl, fkArmCtrlGrp)
+	    
+
+	    #position the arm FK control so that it is about halfway down the arm length
+	    dist = (cmds.getAttr(fkElbowJoint + ".tx")) / 2
+	    cmds.setAttr(fkArmCtrl + ".translateX", dist)
+	
+	    #set the pivot of the arm control back to the arm joint
+	    piv = cmds.xform(fkArmJoint, q = True, ws = True, rotatePivot = True)
+	    cmds.xform(fkArmCtrl, ws = True, piv = [piv[0], piv[1], piv[2]])
+	    
+	    #freeze transforms on the control
+	    cmds.makeIdentity(fkArmCtrl, t = 1, r = 1, s = 1, apply = True)
+
+	    #parent the orient arm grp to the fk ctrl
+	    cmds.parent(armOrientGrp, fkArmCtrl)
+	    cmds.setAttr(armOrient + ".v", 0)
+	    
+	    #orient constraint the fk upper arm joint to the arm orient locator
+	    cmds.orientConstraint(armOrient, fkArmJoint)
+
+
+	    #create FK elbow control
+	    fkElbowCtrl = utils.createControl("circle", 18, "fk_" + self.prefix + "elbow" + self.suffix + "_anim")
+	    cmds.setAttr(fkElbowCtrl + ".ry", -90)
+	    cmds.makeIdentity(fkElbowCtrl, r = 1, apply =True)
+	    
+	    constraint = cmds.parentConstraint(fkElbowJoint, fkElbowCtrl)[0]
+	    cmds.delete(constraint)
+	    
+	    fkElbowCtrlGrp = cmds.group(empty = True, name = "fk_" + self.prefix + "elbow" + self.suffix + "_anim_grp")
+	    constraint = cmds.parentConstraint(fkElbowJoint, fkElbowCtrlGrp)[0]
+	    cmds.delete(constraint)
+	    
+	    #Add FK elbow to FK rig hierarchy
+	    cmds.parent(fkElbowCtrl, fkElbowCtrlGrp)
+	    cmds.makeIdentity(fkElbowCtrl, t = 1, r = 1, s = 1, apply = True)
+	    cmds.parent(fkElbowCtrlGrp, fkArmCtrl)
+			
+	    #constrain elbow joint to ctrl
+	    cmds.parentConstraint(fkElbowCtrl, fkElbowJoint)
+			
+			
+	    #create FK wrist control
+	    fkWristCtrl = utils.createControl("circle", 15, "fk_" + self.prefix + "wrist" + self.suffix + "_anim")
+	    cmds.setAttr(fkWristCtrl + ".ry", -90)
+	    cmds.makeIdentity(fkWristCtrl, r = 1, apply =True)
+	    
+	    constraint = cmds.parentConstraint(fkWristJoint, fkWristCtrl)[0]
+	    cmds.delete(constraint)
+	    
+	    fkWristCtrlGrp = cmds.group(empty = True, name = "fk_" + self.prefix + "wrist" + self.suffix + "_anim_grp")
+	    constraint = cmds.parentConstraint(fkWristJoint, fkWristCtrlGrp)[0]
+	    cmds.delete(constraint)
+	    
+	    cmds.parent(fkWristCtrl, fkWristCtrlGrp)
+	    cmds.makeIdentity(fkWristCtrl, t = 1, r = 1, s = 1, apply = True)
+	    cmds.parent(fkWristCtrlGrp, fkElbowCtrl)
+	    
+	    
+	    #constrain wrist joint to ctrl
+	    cmds.parentConstraint(fkWristCtrl, fkWristJoint)
+	    
+
+	    #point constrain the fk arm grp to the fk upper arm joint
+	    cmds.pointConstraint(fkArmJoint, fkArmCtrlGrp)
+			
+
+	    #group up the groups
+	    jointsGrp = cmds.group(empty = True, name = "joints_" + name + "_grp")
+	    cmds.parent(fkArmJoint, jointsGrp)
+	    
+	    masterGrp = cmds.group(empty = True, name = "arm_rig_master_grp_" + name)
+	    constraint = cmds.pointConstraint(fkArmJoint, masterGrp)[0]
+	    cmds.delete(constraint)
+	    
+	    cmds.parent(jointsGrp, masterGrp)
+	    cmds.parent(masterGrp, "arm_sys_grp")
+	    
+			
+	    #add fk orientation options(normal, body, world)
+	    fkOrient = cmds.spaceLocator(name = "fk_orient_master_loc_" + name)[0]
+	    shape = cmds.listRelatives(fkOrient, shapes = True)[0]
+	    cmds.setAttr(shape + ".v", 0)
+	    fkBodyOrient = None
+	    
+	    
+	    constraint = cmds.parentConstraint(fkArmJoint, fkOrient)[0]
+	    cmds.delete(constraint)
+	    
+	    fkNormalOrient = cmds.duplicate(fkOrient, po = True, name = "fk_orient_normal_loc_" + name)[0]
+	    
+	    if createBodyOrient:
+		fkBodyOrient = cmds.duplicate(fkOrient, po = True, name = "fk_orient_body_loc_" + name)[0]
+		
+	    fkWorldOrient = cmds.duplicate(fkOrient, po = True, name = "fk_orient_world_loc_" + name)[0]
+	    
+	    if createBodyOrient:
+		fkOrientConstraint = cmds.orientConstraint([fkNormalOrient, fkBodyOrient, fkWorldOrient], fkOrient)[0]
+		cmds.parent([fkOrient, fkNormalOrient, fkBodyOrient], masterGrp)
+		cmds.parent(fkBodyOrient, "body_anim")
+		
+		
+	    else:
+		fkOrientConstraint = cmds.orientConstraint([fkNormalOrient, fkWorldOrient], fkOrient)[0]
+		cmds.parent([fkOrient, fkNormalOrient], masterGrp)		
+		
+
+	    #parent FK arm ctrl grp to master orient locator
+	    cmds.parent(fkArmCtrlGrp, fkOrient)
+			
+			
+	    #put mode into default fk operation
+	    if createBodyOrient:
+		cmds.setAttr(fkOrientConstraint + "." + fkBodyOrient + "W1", 0)
+		cmds.setAttr(fkOrientConstraint + "." + fkWorldOrient + "W2", 0)
+		
+	    else:
+		cmds.setAttr(fkOrientConstraint + "." + fkWorldOrient + "W1", 0)
+		
+			
+			
+	    #constrain the masterGrp to the armGrpParent
+	    cmds.parentConstraint(armGrpParent, masterGrp, mo = True)
+	    
+	    
+	    #setup FK arm orient attr
+	    cmds.select("Rig_Settings")
+	    if createBodyOrient:
+		cmds.addAttr(longName= name + "FkArmOrient", at = 'enum', en = "fk:body:world:", keyable = True)
+	    
+	    else:
+		cmds.addAttr(longName= name + "FkArmOrient", at = 'enum', en = "fk:world:", keyable = True)
+		
+		
+	    if createBodyOrient:
+		cmds.setAttr("Rig_Settings." + name + "FkArmOrient" + self.name, 0)
+		cmds.setAttr(fkOrientConstraint + "." + fkNormalOrient + "W0", 1)
+		cmds.setAttr(fkOrientConstraint + "." + fkBodyOrient + "W1", 0)
+		cmds.setAttr(fkOrientConstraint + "." + fkWorldOrient + "W2", 0)
+		cmds.setDrivenKeyframe([fkOrientConstraint + "." + fkNormalOrient + "W0", fkOrientConstraint + "." + fkBodyOrient + "W1", fkOrientConstraint + "." + fkWorldOrient + "W2"], cd = "Rig_Settings." + name + "FkArmOrient", itt = "linear", ott = "linear")
+		
+		cmds.setAttr("Rig_Settings." + name + "FkArmOrient", 1)
+		cmds.setAttr(fkOrientConstraint + "." + fkNormalOrient + "W0", 0)
+		cmds.setAttr(fkOrientConstraint + "." + fkBodyOrient + "W1", 1)
+		cmds.setAttr(fkOrientConstraint + "." + fkWorldOrient + "W2", 0)
+		cmds.setDrivenKeyframe([fkOrientConstraint + "." + fkNormalOrient + "W0", fkOrientConstraint + "." + fkBodyOrient + "W1", fkOrientConstraint + "." + fkWorldOrient + "W2"], cd = "Rig_Settings." + name + "FkArmOrient", itt = "linear", ott = "linear")
+		
+		cmds.setAttr("Rig_Settings." + name + "FkArmOrient", 2)
+		cmds.setAttr(fkOrientConstraint + "." + fkNormalOrient + "W0", 0)
+		cmds.setAttr(fkOrientConstraint + "." + fkBodyOrient + "W1", 0)
+		cmds.setAttr(fkOrientConstraint + "." + fkWorldOrient + "W2", 1)
+		cmds.setDrivenKeyframe([fkOrientConstraint + "." + fkNormalOrient + "W0", fkOrientConstraint + "." + fkBodyOrient + "W1", fkOrientConstraint + "." + fkWorldOrient + "W2"], cd = "Rig_Settings." + name + "FkArmOrient", itt = "linear", ott = "linear")
+		
+		cmds.setAttr("Rig_Settings." + name + "FkArmOrient", 0)
+		
+	    
+	    else:
+		cmds.setAttr("Rig_Settings." + name + "FkArmOrient" + self.name, 0)
+		cmds.setAttr(fkOrientConstraint + "." + fkNormalOrient + "W0", 1)
+		cmds.setAttr(fkOrientConstraint + "." + fkWorldOrient + "W1", 0)
+		cmds.setDrivenKeyframe([fkOrientConstraint + "." + fkNormalOrient + "W0", fkOrientConstraint + "." + fkWorldOrient + "W1"], cd = "Rig_Settings." + name + "FkArmOrient", itt = "linear", ott = "linear")
+		
+		cmds.setAttr("Rig_Settings." + name + "FkArmOrient", 1)
+		cmds.setAttr(fkOrientConstraint + "." + fkNormalOrient + "W0", 0)
+		cmds.setAttr(fkOrientConstraint + "." + fkWorldOrient + "W1", 1)
+		cmds.setDrivenKeyframe([fkOrientConstraint + "." + fkNormalOrient + "W0", fkOrientConstraint + "." + fkWorldOrient + "W1"], cd = "Rig_Settings." + name + "FkArmOrient", itt = "linear", ott = "linear")
+		
+		cmds.setAttr("Rig_Settings." + name + "FkArmOrient", 0)		
+			
+
+	    #lock attrs that should not be animated and colorize controls
+	    for control in [fkArmCtrl, fkElbowCtrl, fkWristCtrl]:
+		
+		for attr in [".sx", ".sy", ".sz", ".v"]:
+		    cmds.setAttr(control + attr, lock = True, keyable = False)
+
+		cmds.setAttr(control + ".overrideEnabled", 1)
+		cmds.setAttr(control + ".overrideColor", color)
+			    
+	    #parent fkWorldOrient to master anim
+	    cmds.parent(fkWorldOrient, "master_anim")
+	    
+	    #generate list of nodes to return
+	    returnNodes = [fkArmJoint, fkArmCtrl, masterGrp, fkBodyOrient, fkNormalOrient]
+	    
+	    return returnNodes
+	
+	
+	
+    # # # # # # # # # # # # # # # # # # # # # # # # # # # # # # # # # # # # # # # # # # # # # # # # # # # # # # # # # # # # # # # 
+    # # # # # # # # # # # # # # # # # # # # # # # # # # # # # # # # # # # # # # # # # # # # # # # # # # # # # # # # # # # # # # # 
+    # # # # # # # # # # # # # # # # # # # # # # # # # # # # # # # # # # # # # # # # # # # # # # # # # # # # # # # # # # # # # # # 
+    # # # # # # # # # # # # # # # # # # # # # # # # # # # # # # # # # # # # # # # # # # # # # # # # # # # # # # # # # # # # # # # 
+    # # # # # # # # # # # # # # # # # # # # # # # # # # # # # # # # # # # # # # # # # # # # # # # # # # # # # # # # # # # # # # # 
+    def ikArmRig(self, fkArmJoint, fkArmControl, masterGrp, color):
+	
+	#create proper name
+	name = self.prefix + self.name + self.suffix
+	if name.find("_") == 0:
+	    name = name.partition("_")[2]
+
+	#duplicate the fk arm joints and create our ik arm chain
+	if cmds.objExists(fkArmJoint):
+	    fkArmJoints = cmds.listRelatives(fkArmJoint, allDescendents = True, type = "joint")
+	    fkArmJoints.reverse()
+	    
+	    ikUpArmJoint = cmds.duplicate(fkArmJoint, po = True, name = "ik_" + self.prefix + self.upperArm + self.suffix)[0]
+	    ikLowArmJoint = cmds.duplicate(fkArmJoints[0], po = True, name = "ik_" + self.prefix + self.lowerArm + self.suffix)[0]
+	    ikWristJoint = cmds.duplicate(fkArmJoints[1], po = True, name = "ik_" + self.prefix + self.hand + self.suffix)[0]
+	    ikWristEndJoint = cmds.duplicate(fkArmJoints[1], po = True, name = "fk_" + self.prefix + self.hand + self.suffix + "_end")[0]
+	    
+	    cmds.parent([ikWristEndJoint], ikWristJoint)
+	    cmds.parent(ikWristJoint, ikLowArmJoint)
+	    cmds.parent(ikLowArmJoint, ikUpArmJoint)
+		
+	    #create fk matching joints (ik_upperarm_fk_matcher_l)
+	    fkMatchUpArm = cmds.duplicate(ikUpArmJoint, po = True, name = "ik_" + self.prefix + self.upperArm + "_fk_matcher" + self.suffix)[0]
+	    fkMatchLowArm = cmds.duplicate(ikLowArmJoint, po = True, name = "ik_" + self.prefix + self.lowerArm + "_fk_matcher" + self.suffix)[0]
+	    fkMatchWrist = cmds.duplicate(ikWristJoint, po = True, name = "ik_" + self.prefix + self.hand + "_fk_matcher" + self.suffix)[0]
+	    
+	    cmds.parent(fkMatchWrist, fkMatchLowArm)
+	    cmds.parent(fkMatchLowArm, fkMatchUpArm)
+	    
+
+    
+	    #move wrist end joint out a bit
+	    scaleFactor = utils.getScaleFactor()
+	    direction = cmds.getAttr(ikWristJoint + ".tx")
+	    if direction > 0:
+		cmds.setAttr(ikWristEndJoint + ".tx", 15 * scaleFactor)
+		
+	    else:
+		cmds.setAttr(ikWristEndJoint + ".tx", -15 * scaleFactor)
+	    
+	    cmds.makeIdentity(ikUpArmJoint, t = 0, r = 1, s = 0, apply = True)
+	    cmds.makeIdentity(fkMatchUpArm, t = 0, r = 1, s = 0, apply = True)
+	    
+	    
+
+	    #constrain fk match joints
+	    cmds.parentConstraint(ikUpArmJoint, fkMatchUpArm, mo = True)
+	    cmds.parentConstraint(ikLowArmJoint, fkMatchLowArm, mo = True)
+	    cmds.parentConstraint(ikWristJoint, fkMatchWrist, mo = True)
+	    
+	    
+	    
+	    
+	    #set rotate order on ikUpArm
+	    cmds.setAttr(ikUpArmJoint + ".rotateOrder", 3)
+		
+	    #set preferred angle on arm
+	    cmds.setAttr(ikLowArmJoint + ".preferredAngleZ", -90)
+		
+	    
+	    #create ik control
+	    ikCtrl = utils.createControl("circle", 15, "ik_wrist_" + name + "_anim")
+	    cmds.setAttr(ikCtrl + ".ry", -90)
+	    cmds.makeIdentity(ikCtrl, r = 1, apply =True)
+	    
+	    constraint = cmds.pointConstraint(ikWristJoint, ikCtrl)[0]
+	    cmds.delete(constraint)
+	    
+	    ikCtrlGrp = cmds.group(empty = True, name = "ik_wrist_" + name + "_anim_grp")
+	    constraint = cmds.pointConstraint(ikWristJoint, ikCtrlGrp)[0]
+	    cmds.delete(constraint)
+
+	    
+	    ikCtrlSpaceSwitchFollow = cmds.duplicate(ikCtrlGrp, po = True, n = "ik_wrist_" + name + "_anim_space_switcher_follow")[0]
+	    ikCtrlSpaceSwitch = cmds.duplicate(ikCtrlGrp, po = True, n = "ik_wrist_" + name + "_anim_space_switcher")[0]
+	    
+	    cmds.parent(ikCtrlSpaceSwitchFollow, "arm_sys_grp")
+	    cmds.parent(ikCtrlSpaceSwitch, ikCtrlSpaceSwitchFollow)
+	    cmds.parent(ikCtrlGrp, ikCtrlSpaceSwitch)
+	
+	    cmds.parent(ikCtrl, ikCtrlGrp)
+	    cmds.makeIdentity(ikCtrlGrp, t = 1, r = 1, s = 1, apply = True)
+	    
+	    #lock out vis/scale
+	    for attr in [".sx", ".sy", ".sz", ".v"]:
+		cmds.setAttr(ikCtrl + attr, lock = True, keyable = False)
+		
+
+	    #create RP IK on arm and SC ik from wrist to wrist end
+	    rpIkHandle = cmds.ikHandle(name = "arm_ikHandle_" + name, solver = "ikRPsolver", sj = ikUpArmJoint, ee = ikWristJoint)[0]
+	    scIkHandle = cmds.ikHandle(name = "hand_ikHandle_" + name, solver = "ikSCsolver", sj = ikWristJoint, ee = ikWristEndJoint)[0]
+	    
+	    cmds.parent(scIkHandle, rpIkHandle)
+	    cmds.setAttr(rpIkHandle + ".v", 0)
+	    cmds.setAttr(scIkHandle + ".v", 0)
+		
+		
+	    #parent IK to ik control
+	    cmds.parent(rpIkHandle, ikCtrl)
+		
+		
+	    #create a pole vector control
+	    ikPvCtrl = utils.createControl("sphere", 6, "ik_elbow_" + name + "_anim")
+	    constraint = cmds.pointConstraint(ikLowArmJoint, ikPvCtrl)[0]
+	    cmds.delete(constraint)
+	    cmds.makeIdentity(ikPvCtrl, t = 1, r = 1, s = 1, apply = True)
+	    
+	    #move out a bit
+	    cmds.setAttr(ikPvCtrl + ".ty", (30 * scaleFactor))
+	    cmds.makeIdentity(ikPvCtrl, t = 1, r = 1, s = 1, apply = True)	    
+		
+		
+	    #create group for control
+	    ikPvCtrlGrp = cmds.group(empty = True, name = "ik_elbow_" + name + "_anim_grp")
+	    constraint = cmds.parentConstraint(ikPvCtrl, ikPvCtrlGrp)[0]
+	    cmds.delete(constraint)
+	    
+	    ikPvSpaceSwitchFollow = cmds.duplicate(ikPvCtrlGrp, po = True, name = "ik_elbow_" + name + "_anim_space_switcher_follow")[0]
+	    ikPvSpaceSwitch = cmds.duplicate(ikPvCtrlGrp, po = True, name = "ik_elbow_" + name + "_anim_space_switcher")[0]
+	    
+	    cmds.parent(ikPvSpaceSwitch, ikPvSpaceSwitchFollow)
+	    cmds.parent(ikPvCtrlGrp, ikPvSpaceSwitch)
+	    cmds.parent(ikPvCtrl, ikPvCtrlGrp)
+	    cmds.parent(ikPvSpaceSwitchFollow, "arm_sys_grp")
+	    cmds.makeIdentity(ikPvCtrl, t = 1, r = 1, s = 1, apply = True)
+	    
+	    #lock out vis, rotates, and scale
+	    for attr in [".sx", ".sy", ".sz", ".rx", ".ry", ".rz", ".v"]:
+		cmds.setAttr(ikPvCtrl + attr, lock = True, keyable = False)
+		
+	    #setup pole vector constraint
+	    cmds.poleVectorConstraint(ikPvCtrl, rpIkHandle)
+	    
+	    
+	    #constrain driver joints to both fk and ik chains
+	    upArmConstPoint = cmds.pointConstraint([fkArmControl, ikUpArmJoint], "driver_" + self.prefix + self.upperArm + self.suffix)[0]
+	    upArmConstOrient = cmds.orientConstraint([fkArmJoint, ikUpArmJoint], "driver_" + self.prefix + self.upperArm + self.suffix)[0]
+	    lowArmConst = cmds.parentConstraint([fkArmJoints[0], ikLowArmJoint], "driver_" + self.prefix + self.lowerArm + self.suffix)[0]
+	    handConst = cmds.parentConstraint([fkArmJoints[1], ikWristJoint], "driver_" + self.prefix + self.hand + self.suffix)[0]
+	    
+	    
+	    
+	    #create blend nodes for the scale
+	    scaleBlendColors_UpArm = cmds.shadingNode("blendColors", asUtility = True, name = "up_arm_scale_blend_" + name)
+	    cmds.connectAttr(ikUpArmJoint + ".scale", scaleBlendColors_UpArm + ".color1")
+	    cmds.connectAttr(fkArmControl + ".scale", scaleBlendColors_UpArm + ".color2")
+	    cmds.connectAttr(scaleBlendColors_UpArm + ".output", "driver_" + self.prefix + self.upperArm + self.suffix + ".scale")
+	    
+	    
+	    scaleBlendColors_LoArm = cmds.shadingNode("blendColors", asUtility = True, name = "lo_arm_scale_blend_" + name)
+	    cmds.connectAttr(ikLowArmJoint + ".scale", scaleBlendColors_LoArm + ".color1")
+	    cmds.connectAttr("fk_" + self.prefix + "elbow" + self.suffix + "_anim"  + ".scale", scaleBlendColors_LoArm + ".color2")
+	    cmds.connectAttr(scaleBlendColors_LoArm + ".output", "driver_" + self.prefix + self.lowerArm + self.suffix + ".scale")
+	    
+	    scaleBlendColors_Wrist = cmds.shadingNode("blendColors", asUtility = True, name = "wrist_scale_blend_" + name)
+	    cmds.connectAttr(ikWristJoint + ".scale", scaleBlendColors_Wrist + ".color1")
+	    cmds.connectAttr("fk_" + self.prefix + "wrist" + self.suffix + "_anim" + ".scale", scaleBlendColors_Wrist + ".color2")
+	    cmds.connectAttr(scaleBlendColors_Wrist + ".output", "driver_" + self.prefix + self.hand + self.suffix + ".scale")
+	    
+	    #set limits
+	    cmds.select("driver_" + self.prefix + self.upperArm + self.suffix)
+	    cmds.transformLimits(sy = (.05, 1.25), sz = (.05, 1.25), esy = [False, True], esz = [False, True])
+	    cmds.select("driver_" + self.prefix + self.lowerArm + self.suffix)
+	    cmds.transformLimits(sy = (.05, 1.25), sz = (.05, 1.25), esy = [False, True], esz = [False, True])	
+	    
+	    
+	    
+	    #create the IK/FK switch
+	    # # # # # # # # # # # # # # # # # # # # # # # # # # # # # # # # # # # # # # # # # # # # # 
+	    ###########################################################################################
+	    
+	    ###########################################################################################
+	    
+	    ###########################################################################################
+	    
+	    cmds.select("Rig_Settings")
+	    cmds.addAttr(longName= name + "ArmMode", at = 'enum', en = "fk:ik:", keyable = True)	    
+	    
+	    # # # # # # # # # # # # # # # # # # # # # # # # # # # # # # # # # # # # # # # # # # # # # # 
+	    #FK MODE
+	    cmds.setAttr("Rig_Settings." + name + "ArmMode", 0)
+	    
+	    cmds.setAttr(upArmConstPoint + "." + fkArmControl + "W0", 1)
+	    cmds.setAttr(upArmConstPoint + "." + ikUpArmJoint + "W1", 0)
+	    cmds.setAttr(upArmConstOrient + "." + fkArmJoint + "W0", 1)
+	    cmds.setAttr(upArmConstOrient + "." + ikUpArmJoint + "W1", 0)
+	    cmds.setAttr(scaleBlendColors_UpArm + "." + "blender", 0)	
+	    
+	    cmds.setAttr(lowArmConst + "." + fkArmJoints[0] + "W0", 1)
+	    cmds.setAttr(lowArmConst + "." + ikLowArmJoint + "W1", 0)
+	    cmds.setAttr(scaleBlendColors_LoArm + "." + "blender", 0)
+	    
+	    cmds.setAttr(handConst + "." + fkArmJoints[1] + "W0", 1)
+	    cmds.setAttr(handConst + "." + ikWristJoint + "W1", 0)
+	    cmds.setAttr(scaleBlendColors_Wrist + "." + "blender", 0)
+	    
+	    
+	    cmds.setAttr(fkArmControl + "_grp.v", 1)
+	    cmds.setAttr("ik_wrist_" + name + "_anim_space_switcher.v", 0)
+	    cmds.setAttr("ik_elbow_" + name + "_anim_space_switcher.v", 0)
+	    
+	    #SET KEYS
+	    cmds.setDrivenKeyframe([scaleBlendColors_UpArm + "." + "blender", scaleBlendColors_LoArm + "." + "blender", scaleBlendColors_Wrist + "." + "blender", upArmConstPoint + "." + fkArmControl + "W0", upArmConstPoint + "." + ikUpArmJoint + "W1", ], cd = "Rig_Settings." + name + "ArmMode", itt = "linear", ott = "linear")
+	    cmds.setDrivenKeyframe([upArmConstOrient + "." + fkArmJoint + "W0", upArmConstOrient + "." + ikUpArmJoint + "W1", ], cd = "Rig_Settings." + name + "ArmMode", itt = "linear", ott = "linear")
+	    cmds.setDrivenKeyframe([lowArmConst + "." + fkArmJoints[0] + "W0", lowArmConst + "." + ikLowArmJoint + "W1", ], cd = "Rig_Settings." + name + "ArmMode", itt = "linear", ott = "linear")
+	    cmds.setDrivenKeyframe([handConst + "." + fkArmJoints[1] + "W0", handConst + "." + ikWristJoint + "W1", ], cd = "Rig_Settings." + name + "ArmMode", itt = "linear", ott = "linear")
+	    cmds.setDrivenKeyframe(fkArmControl + "_grp.v", cd = "Rig_Settings." + name + "ArmMode", itt = "linear", ott = "linear")
+	    cmds.setDrivenKeyframe(["ik_wrist_" + name + "_anim_space_switcher.v", "ik_elbow_" + name + "_anim_space_switcher.v"], cd = "Rig_Settings." + name + "ArmMode", itt = "linear", ott = "linear")
+	    
+	    # # # # # # # # # # # # # # # # # # # # # # # # # # # # # # # # # # # # # # # # # # # # # # 
+	    #IK MODE
+	    
+	    cmds.setAttr("Rig_Settings." + name + "ArmMode", 1)
+	    
+	    cmds.setAttr(upArmConstPoint + "." + fkArmControl + "W0", 0)
+	    cmds.setAttr(upArmConstPoint + "." + ikUpArmJoint + "W1", 1)
+	    cmds.setAttr(upArmConstOrient + "." + fkArmJoint + "W0", 0)
+	    cmds.setAttr(upArmConstOrient + "." + ikUpArmJoint + "W1", 1)
+	    cmds.setAttr(scaleBlendColors_UpArm + "." + "blender", 1)	
+	    
+	    cmds.setAttr(lowArmConst + "." + fkArmJoints[0] + "W0", 0)
+	    cmds.setAttr(lowArmConst + "." + ikLowArmJoint + "W1", 1)
+	    cmds.setAttr(scaleBlendColors_LoArm + "." + "blender", 1)
+	    
+	    cmds.setAttr(handConst + "." + fkArmJoints[1] + "W0", 0)
+	    cmds.setAttr(handConst + "." + ikWristJoint + "W1", 1)
+	    cmds.setAttr(scaleBlendColors_Wrist + "." + "blender", 1)
+	    
+	    
+	    cmds.setAttr(fkArmControl + "_grp.v", 0)
+	    cmds.setAttr("ik_wrist_" + name + "_anim_space_switcher.v", 1)
+	    cmds.setAttr("ik_elbow_" + name + "_anim_space_switcher.v", 1)
+	    
+	    #SET KEYS
+	    cmds.setDrivenKeyframe([scaleBlendColors_UpArm + "." + "blender", scaleBlendColors_LoArm + "." + "blender", scaleBlendColors_Wrist + "." + "blender", upArmConstPoint + "." + fkArmControl + "W0", upArmConstPoint + "." + ikUpArmJoint + "W1", ], cd = "Rig_Settings." + name + "ArmMode", itt = "linear", ott = "linear")
+	    cmds.setDrivenKeyframe([upArmConstOrient + "." + fkArmJoint + "W0", upArmConstOrient + "." + ikUpArmJoint + "W1", ], cd = "Rig_Settings." + name + "ArmMode", itt = "linear", ott = "linear")
+	    cmds.setDrivenKeyframe([lowArmConst + "." + fkArmJoints[0] + "W0", lowArmConst + "." + ikLowArmJoint + "W1", ], cd = "Rig_Settings." + name + "ArmMode", itt = "linear", ott = "linear")
+	    cmds.setDrivenKeyframe([handConst + "." + fkArmJoints[1] + "W0", handConst + "." + ikWristJoint + "W1", ], cd = "Rig_Settings." + name + "ArmMode", itt = "linear", ott = "linear")
+	    cmds.setDrivenKeyframe(fkArmControl + "_grp.v", cd = "Rig_Settings." + name + "ArmMode", itt = "linear", ott = "linear")
+	    cmds.setDrivenKeyframe(["ik_wrist_" + name + "_anim_space_switcher.v", "ik_elbow_" + name + "_anim_space_switcher.v"], cd = "Rig_Settings." + name + "ArmMode", itt = "linear", ott = "linear")
+
+	    #reset back to FK mode
+	    cmds.setAttr("Rig_Settings." + name + "ArmMode", 0)
+	    
+	    ###########################################################################################
+	    
+	    ###########################################################################################
+	    
+	    ###########################################################################################
+	    
+	    
+	    # # # # # # # # # # # # # # # # # # # # # # # # # # # # # # # # # # # # # # # # # # # # # 
+	    ###########################################################################################
+	    
+	    ###########################################################################################
+	    
+	    ###########################################################################################
+	    
+	    #setup stretch on IK
+	    cmds.select(ikCtrl)
+	    cmds.addAttr(longName=("stretch"), at = 'double',min = 0, max = 1, dv = 0, keyable = True)
+	    cmds.addAttr(longName=("squash"), at = 'double',min = 0, max = 1, dv = 0, keyable = True)
+	    stretchMultNode = cmds.shadingNode("multiplyDivide", asUtility = True, name = "ikHand_stretchToggleMultNode_" + name)
+
+	    #need to get the total length of the arm chain
+	    totalDist = abs(cmds.getAttr(ikLowArmJoint + ".tx" ) + cmds.getAttr(ikWristJoint + ".tx"))	    
+	    
+	    #create a distanceBetween node
+	    distBetween = cmds.shadingNode("distanceBetween", asUtility = True, name = "ik_arm_distBetween_" + name)
+	    
+	    #get world positions of upper arm and ik
+	    baseGrp = cmds.group(empty = True, name = "ik_arm_base_grp_" + name)
+	    endGrp = cmds.group(empty = True, name = "ik_arm_end_grp_" + name)
+	    cmds.pointConstraint(ikUpArmJoint, baseGrp)
+	    cmds.pointConstraint(ikCtrl, endGrp)
+
+	    #hook in group translates into distanceBetween node inputs
+	    cmds.connectAttr(baseGrp + ".translate", distBetween + ".point1")
+	    cmds.connectAttr(endGrp + ".translate", distBetween + ".point2")
+	    
+	    #create a condition node that will compare original length to current length
+	    #if second term is greater than, or equal to the first term, the chain needs to stretch
+	    ikArmCondition = cmds.shadingNode("condition", asUtility = True, name = "ik_arm_stretch_condition_" + name)
+	    cmds.setAttr(ikArmCondition + ".operation", 3)
+	    cmds.connectAttr(distBetween + ".distance", ikArmCondition + ".secondTerm")
+	    cmds.setAttr(ikArmCondition + ".firstTerm", totalDist)
+	    
+	    #hook up the condition node's return colors
+	    cmds.setAttr(ikArmCondition + ".colorIfTrueR", totalDist)
+	    cmds.connectAttr(distBetween + ".distance", ikArmCondition + ".colorIfFalseR")
+	    
+	    #create the mult/divide node(set to divide) that will take the original creation length as a static value in input2x, and the connected length into 1x.
+	    armDistMultNode = cmds.shadingNode("multiplyDivide", asUtility = True, name = "arm_dist_multNode_" + name)
+	    cmds.setAttr(armDistMultNode + ".operation", 2) #divide
+	    cmds.setAttr(armDistMultNode + ".input2X", totalDist)
+	    cmds.connectAttr(ikArmCondition + ".outColorR", armDistMultNode + ".input1X")
+	    
+	    
+	    #create a stretch toggle mult node that multiplies the stretch factor by the bool of the stretch attr. (0 or 1), this way our condition reads
+	    #if this result is greater than the original length(impossible if stretch bool is off, since result will be 0), than take this result and plug it
+	    #into the scale of our IK arm joints
+	    stretchToggleCondition = cmds.shadingNode("condition", asUtility = True, name = "arm_stretch_toggle_condition_" + name)
+	    cmds.setAttr(stretchToggleCondition + ".operation", 0)
+	    cmds.connectAttr(ikCtrl + ".stretch", stretchToggleCondition + ".firstTerm")
+	    cmds.setAttr(stretchToggleCondition + ".secondTerm", 1)
+	    cmds.connectAttr(armDistMultNode + ".outputX", stretchToggleCondition + ".colorIfTrueR")
+	    cmds.setAttr(stretchToggleCondition + ".colorIfFalseR", 1)
+	    
+	    #set up the squash nodes
+	    squashMultNode = cmds.shadingNode("multiplyDivide", asUtility = True, name = "ik_arm_squash_mult_" + name)
+	    cmds.setAttr(squashMultNode + ".operation", 2)
+	    cmds.setAttr(squashMultNode + ".input1X", totalDist)
+	    cmds.connectAttr(ikArmCondition + ".outColorR", squashMultNode + ".input2X")
+	    
+	    #create a squash toggle mult node that multiplies the stretch factor by the bool of the squash attr. (0 or 1), this way our condition reads
+	    #if this result is greater than the original length(impossible if stretch bool is off, since result will be 0), than take this result and plug it
+	    #into the scale of our IK arm joints
+	    squashToggleCondition = cmds.shadingNode("condition", asUtility = True, name = "arm_squash_toggle_condition_" + name)
+	    cmds.setAttr(squashToggleCondition + ".operation", 0)
+	    cmds.connectAttr(ikCtrl + ".squash", squashToggleCondition + ".firstTerm")
+	    cmds.setAttr(squashToggleCondition + ".secondTerm", 1)
+	    cmds.connectAttr(squashMultNode + ".outputX", squashToggleCondition + ".colorIfTrueR")
+	    cmds.setAttr(squashToggleCondition + ".colorIfFalseR", 1)	    
+	    
+	    
+	    #connect to arm scale
+	    cmds.connectAttr(stretchToggleCondition + ".outColorR", ikUpArmJoint + ".sx")
+	    cmds.connectAttr(stretchToggleCondition + ".outColorR", ikLowArmJoint + ".sx")
+	    
+	    cmds.connectAttr(squashToggleCondition + ".outColorR", ikLowArmJoint + ".sy")
+	    cmds.connectAttr(squashToggleCondition + ".outColorR", ikLowArmJoint + ".sz")
+	    
+	    cmds.connectAttr(squashToggleCondition + ".outColorR", ikUpArmJoint + ".sy")
+	    cmds.connectAttr(squashToggleCondition + ".outColorR", ikUpArmJoint + ".sz")
+	    
+	    
+	    ###########################################################################################
+	    
+	    ###########################################################################################
+	    
+	    ###########################################################################################
+	    
+
+	    
+	    #add base and end groups to arm grp
+	    cmds.parent([baseGrp, endGrp], masterGrp)
+	    
+	    #color ikCtrl and ikPvCtrl
+	    for control in [ikCtrl, ikPvCtrl]:
+		cmds.setAttr(control + ".overrideEnabled", 1)
+		cmds.setAttr(control + ".overrideColor", color)
+		
+	    returnNodes = [ikUpArmJoint]
+	    return returnNodes
+	    
+    # # # # # # # # # # # # # # # # # # # # # # # # # # # # # # # # # # # # # # # # # # # # # # # # # # # # # # # # # # # # # # # 
+    # # # # # # # # # # # # # # # # # # # # # # # # # # # # # # # # # # # # # # # # # # # # # # # # # # # # # # # # # # # # # # # 
+    # # # # # # # # # # # # # # # # # # # # # # # # # # # # # # # # # # # # # # # # # # # # # # # # # # # # # # # # # # # # # # # 
+    # # # # # # # # # # # # # # # # # # # # # # # # # # # # # # # # # # # # # # # # # # # # # # # # # # # # # # # # # # # # # # # 
+    # # # # # # # # # # # # # # # # # # # # # # # # # # # # # # # # # # # # # # # # # # # # # # # # # # # # # # # # # # # # # # # 
+    def clavicleRig(self, fkArmJoint, ikArmJoint, fkArmCtrl, fkOrientNormal, armGrpParent, color):
+	
+	#create a nice name
+	name = self.prefix + self.name + self.suffix
+	if name.find("_") == 0:
+	    name = name.partition("_")[2]
+	
+	#create the result clavicle joint, parent under joints_grp
+	clavJoint = cmds.duplicate("driver_" + self.prefix + self.clavicle + self.suffix, po = True, name = "rig_" + self.prefix + self.clavicle + self.suffix)[0]
+	cmds.parent(clavJoint, "joints_" + name + "_grp")
+	
+	#parent fk and ik arm joint chains under result clavicle
+	cmds.parent([fkArmJoint, ikArmJoint], clavJoint)
+	
+	#hook up driver joint to result clav joint
+	cmds.pointConstraint(clavJoint, "driver_" + self.prefix + self.clavicle + self.suffix)
+	cmds.orientConstraint(clavJoint, "driver_" + self.prefix + self.clavicle + self.suffix)
+	cmds.connectAttr(clavJoint + ".scale", "driver_" + self.prefix + self.clavicle + self.suffix + ".scale")
+	
+	
+	#setup the ik/auto clavicle
+	ikClavNodes = self.ikClavRig(clavJoint, fkArmJoint, ikArmJoint, fkArmCtrl, fkOrientNormal, armGrpParent, color)
+	
+	
+	#setup fk clavicle
+	fkClavNodes = self.fkClavRig(clavJoint, fkArmJoint, ikArmJoint, fkArmCtrl, fkOrientNormal, armGrpParent, color)
+	
+	
+	#setup switch on rig clavicle joint and on fk orient normal node
+	clavModeConstraint = cmds.parentConstraint([fkClavNodes[0], ikClavNodes[0]], clavJoint)[0]
+	fkOrientConstraint = cmds.parentConstraint([fkClavNodes[0], ikClavNodes[2]], fkOrientNormal, mo = True)[0]
+	
+	
+	cmds.select("Rig_Settings")
+	cmds.addAttr(longName= name + "ClavMode", at = 'enum', en = "fk:ik:", keyable = True)
+	
+	# # # # # # # # # # # # # # # # # # # # # # # # # # # # # # # # # # # # # # # # # # # # # # 
+	#FK MODE
+	cmds.setAttr("Rig_Settings." + name + "ClavMode", 0)
+	
+	cmds.setAttr(clavModeConstraint + "." + fkClavNodes[0] + "W0", 1)
+	cmds.setAttr(clavModeConstraint + "." + ikClavNodes[0] + "W1", 0)
+	
+	cmds.setAttr(fkOrientConstraint + "." + fkClavNodes[0] + "W0", 1)
+	cmds.setAttr(fkOrientConstraint + "." + ikClavNodes[2] + "W1", 0)
+	
+	cmds.setAttr(fkClavNodes[1] + ".v", 1)
+	cmds.setAttr(ikClavNodes[1] + ".v", 0)
+	
+	#SET KEYS
+	cmds.setDrivenKeyframe([fkOrientConstraint + "." + ikClavNodes[2] + "W1", fkOrientConstraint + "." + fkClavNodes[0] + "W0", ikClavNodes[1] + ".v", fkClavNodes[1] + ".v", clavModeConstraint + "." + fkClavNodes[0] + "W0", clavModeConstraint + "." + ikClavNodes[0] + "W1"], cd = "Rig_Settings." + name + "ClavMode", itt = "linear", ott = "linear")
+
+	
+	# # # # # # # # # # # # # # # # # # # # # # # # # # # # # # # # # # # # # # # # # # # # # # 
+	#IK MODE
+	cmds.setAttr("Rig_Settings." + name + "ClavMode", 1)
+	
+	cmds.setAttr(clavModeConstraint + "." + fkClavNodes[0] + "W0", 0)
+	cmds.setAttr(clavModeConstraint + "." + ikClavNodes[0] + "W1", 1)
+	
+	cmds.setAttr(fkOrientConstraint + "." + fkClavNodes[0] + "W0", 0)
+	cmds.setAttr(fkOrientConstraint + "." + ikClavNodes[2] + "W1", 1)	
+	
+	cmds.setAttr(fkClavNodes[1] + ".v", 0)
+	cmds.setAttr(ikClavNodes[1] + ".v", 1)
+	
+	#SET KEYS
+	cmds.setDrivenKeyframe([fkOrientConstraint + "." + ikClavNodes[2] + "W1", fkOrientConstraint + "." + fkClavNodes[0] + "W0", ikClavNodes[1] + ".v", fkClavNodes[1] + ".v", clavModeConstraint + "." + fkClavNodes[0] + "W0", clavModeConstraint + "." + ikClavNodes[0] + "W1"], cd = "Rig_Settings." + name + "ClavMode", itt = "linear", ott = "linear")
+	
+	# # # # # # # # # # # # # # # # # # # # # # # # # # # # # # # # # # # # # # # # # # # # # # 
+	
+	
+	
+	#SETUP THE AUTO CLAV SWITCH (based on wether the arm is in fk or ik)
+	# # # # # # # # # # # # # # # # # # # # # # # # # # # # # # # # # # # # # # # # # # # # # # 
+	#FK MODE
+	cmds.setAttr("Rig_Settings." + name + "ArmMode", 0)
+	cmds.setAttr(ikClavNodes[3] + ".ikBlend", 0)
+	
+	#SET KEYS
+	cmds.setDrivenKeyframe([ikClavNodes[3] + ".ikBlend"], cd = "Rig_Settings." + name + "ArmMode", itt = "linear", ott = "linear")
+	
+	#IK MODE
+	cmds.setAttr("Rig_Settings." + name + "ArmMode", 1)
+	cmds.setAttr(ikClavNodes[3] + ".ikBlend", 1)
+
+	
+	#SET KEYS
+	cmds.setDrivenKeyframe([ikClavNodes[3] + ".ikBlend"], cd = "Rig_Settings." + name + "ArmMode", itt = "linear", ott = "linear")	
+	
+	
+	#return nodes
+	return clavJoint
+	
+    # # # # # # # # # # # # # # # # # # # # # # # # # # # # # # # # # # # # # # # # # # # # # # # # # # # # # # # # # # # # # # # 
+    # # # # # # # # # # # # # # # # # # # # # # # # # # # # # # # # # # # # # # # # # # # # # # # # # # # # # # # # # # # # # # # 
+    # # # # # # # # # # # # # # # # # # # # # # # # # # # # # # # # # # # # # # # # # # # # # # # # # # # # # # # # # # # # # # # 
+    # # # # # # # # # # # # # # # # # # # # # # # # # # # # # # # # # # # # # # # # # # # # # # # # # # # # # # # # # # # # # # # 
+    # # # # # # # # # # # # # # # # # # # # # # # # # # # # # # # # # # # # # # # # # # # # # # # # # # # # # # # # # # # # # # # 
+    def ikClavRig(self, clavJoint, fkArmJoint, ikArmJoint, fkArmCtrl, fkOrientNormal, armGrpParent, color):
+	
+	#create a nice name
+	name = self.prefix + self.name + self.suffix
+	if name.find("_") == 0:
+	    name = name.partition("_")[2]
+	    
+	    
+	#setup the auto clavicle
+	autoClavJointStart = cmds.duplicate("driver_" + self.prefix + self.clavicle + self.suffix, po = True, name = "auto_clavicle_" + name)[0]
+	ikClavJoint = cmds.duplicate("driver_" + self.prefix + self.clavicle + self.suffix, po = True, name = "ik_" + self.prefix + self.clavicle + self.suffix)[0]
+	ikClavJointEnd = cmds.duplicate(fkArmJoint, parentOnly = True, name = "ik_" + self.prefix + self.clavicle + self.suffix + "_end")[0]
+	cmds.parent(ikClavJointEnd, ikClavJoint)
+	
+	
+	
+	#parent the fk upperarm to the world
+	parent = cmds.listRelatives(ikClavJoint, parent = True)
+	if parent != None:
+	    cmds.parent(ikClavJoint, world = True)	
+
+	cmds.parent(ikClavJoint, "joints_" + name + "_grp")
+	
+        #create the shoulder hierarchy
+	parent = cmds.listRelatives(autoClavJointStart, parent = True)
+	if parent != None:
+	    cmds.parent(autoClavJointStart, world = True)
+	    
+	autoClavEndJoint = cmds.duplicate(fkArmJoint, parentOnly = True, name = "auto_clavicle_end_" + name)[0]
+	pos = cmds.xform(autoClavEndJoint, q = True, ws = True, t = True)
+	zPos = cmds.xform(autoClavJointStart, q = True, ws = True, t = True)[2]
+	cmds.xform(autoClavEndJoint, ws = True, t = [pos[0], pos[1], zPos])
+	cmds.parent(autoClavEndJoint, autoClavJointStart)
+	
+	#get all of the fk arm joints
+	fkArmJoints = cmds.listRelatives(fkArmJoint, allDescendents = True, type = "joint")
+	fkArmJoints.reverse()	
+
+	#create the IK for the clavicle
+	ikNodes = cmds.ikHandle(sj = autoClavJointStart, ee = autoClavEndJoint, sol = "ikSCsolver", name = "auto_clav_to_elbow_ikHandle_" + name)[0]
+	
+	#position the IK handle at the elbow joint
+	constraint = cmds.pointConstraint(fkArmJoints[0], ikNodes)[0]
+	cmds.delete(constraint)
+	
+	#create our autoClav world grp
+	autoClavWorld =  cmds.group(empty = True, name = "auto_clav_world_grp_" + name)
+	constraint = cmds.pointConstraint(autoClavEndJoint, autoClavWorld)[0]
+	cmds.delete(constraint)
+	cmds.makeIdentity(autoClavWorld, t = 1, r = 1, s = 1, apply = True)	
+	
+	
+	#duplicate the FK arm to create our invisible arm
+	invisUpArm = cmds.duplicate(fkArmJoint, po = True, name = "invis_" + fkArmJoint)[0]
+	invisLowArm = cmds.duplicate(fkArmJoints[0], po = True, name = "invis_" + fkArmJoints[0])[0]
+	invisHand = cmds.duplicate(fkArmJoints[1], po = True, name = "invis_" + fkArmJoints[1] )[0]
+	
+	cmds.parent(invisHand, invisLowArm)
+	cmds.parent(invisLowArm, invisUpArm)
+	cmds.parent(invisUpArm, autoClavWorld)
+	
+	
+	#create our upperarm orient locator
+	invisArmOrient = cmds.spaceLocator(name = "invis_arm_orient_loc_" + name)[0]
+	invisArmOrientGrp = cmds.group(empty = True, name = "invis_arm_orient_loc_grp_" + name)
+	
+	constraint = cmds.parentConstraint(fkArmJoint, invisArmOrient)[0]
+	cmds.delete(constraint)
+	constraint = cmds.parentConstraint(fkArmJoint, invisArmOrientGrp)[0]
+	cmds.delete(constraint)
+	cmds.parent(invisArmOrient, invisArmOrientGrp)	
+	
+	
+	#create the invis arm fk control (to derive autoClav info)
+	invisFkArmCtrl = utils.createControl("circle", 20, "invis_fk_" + self.prefix + "arm" + self.suffix + "_anim")
+	cmds.setAttr(invisFkArmCtrl + ".ry", -90)
+	cmds.makeIdentity(invisFkArmCtrl, r = 1, apply =True)
+	
+	constraint = cmds.parentConstraint(fkArmJoint, invisFkArmCtrl)[0]
+	cmds.delete(constraint)	
+	
+	
+	invisFkArmCtrlGrp = cmds.group(empty = True, name = "invis_fk_" + self.prefix + "arm" + self.suffix + "_anim_grp")
+	constraint = cmds.parentConstraint(fkArmJoint, invisFkArmCtrlGrp)[0]
+	cmds.delete(constraint)
+	cmds.parent(invisFkArmCtrl, invisFkArmCtrlGrp)
+
+
+	#orient constrain the invis fk up arm to the invis up arm orient loc. Also do this for the real arm
+	cmds.orientConstraint(invisArmOrient, invisUpArm)
+	cmds.parent(invisArmOrientGrp, invisFkArmCtrl)
+	
+
+	#connect invis arm ctrl rotates to be driven by real arm control rotates
+	cmds.connectAttr(fkArmCtrl + ".rotate", invisFkArmCtrl + ".rotate")	
+	
+	
+	#the following section of code will essentially give us a vector from the clav joint to the elbow. This will help to drive the rotations of the auto clav
+        #create our locators to determine elbow's position in space (will drive the ik handle on the auto clav)
+        autoTransLoc = cmds.spaceLocator(name = "elbow_auto_trans_loc_" + name)[0]
+        constraint = cmds.pointConstraint(fkArmJoints[0], autoTransLoc)[0]
+        cmds.delete(constraint)	
+	
+	
+        #duplicate the locator to create a parent loc
+        autoTransLocParent = cmds.duplicate(autoTransLoc, po = True, name = "elbow_auto_trans_loc_parent_" + name)[0]
+        cmds.pointConstraint(autoTransLoc, ikNodes)
+        cmds.parent(autoTransLoc, autoTransLocParent)
+        cmds.parent(autoTransLocParent, autoClavWorld)
+        cmds.makeIdentity(autoTransLocParent, t = 1, r = 1, s = 1, apply = True)	
+	
+	
+        #constrain the parent trans loc(elbow) to the invis elbow joint
+        cmds.pointConstraint(invisLowArm, autoTransLocParent, mo = True)
+        
+        
+        #create our locator that will handle switching between auto clav and manual clav. position at end joint (autoClavEndJoint)
+        autoClavSwitchLoc = cmds.spaceLocator(name = "auto_clav_switch_loc_" + name)[0]
+        
+        constraint = cmds.pointConstraint(autoClavEndJoint, autoClavSwitchLoc)[0]
+        cmds.delete(constraint)
+        cmds.parent(autoClavSwitchLoc, autoClavWorld)
+        cmds.makeIdentity(autoClavSwitchLoc, t = 1, r = 1, s = 1, apply = True)
+        cmds.parent(autoClavJointStart, autoClavWorld)	
+	
+        #setup constraint for switching between auto/manual
+        autoClavSwitchConstraint = cmds.pointConstraint([autoClavEndJoint, autoClavWorld], autoClavSwitchLoc, mo = True)[0]
+        cmds.setAttr(autoClavSwitchConstraint + "." + autoClavWorld + "W1", 0)
+	
+	
+        #create our IK for the auto clav to move
+        autoClavIK = cmds.ikHandle(sj = ikClavJoint, ee = ikClavJointEnd, sol = "ikSCsolver", name = "auto_clav_ikHandle_" + name)[0]
+        autoClavIKGrp = cmds.group(empty = True, name = "auto_clav_ikHandle_grp_" + name)
+        constraint = cmds.pointConstraint(autoClavIK, autoClavIKGrp)[0]
+        cmds.delete(constraint)
+        autoClavIKGrpMaster = cmds.duplicate(po = True, name = "auto_clav_ikHandle_grp_master_" + name)[0]	
+	
+	
+        cmds.parent(autoClavIKGrpMaster, autoClavSwitchLoc)
+        cmds.parent(autoClavIKGrp, autoClavIKGrpMaster)
+        cmds.parent(autoClavIK, autoClavIKGrp)
+	
+	
+        #create the shoulder control
+        shoulderCtrl = utils.createControl("pin", 1.5, "clavicle_" + name + "_anim")
+        cmds.setAttr(shoulderCtrl + ".ry", 90)
+        cmds.setAttr(shoulderCtrl + ".rx", 90)
+    
+        constraint = cmds.pointConstraint(fkArmJoint, shoulderCtrl)[0]
+        cmds.delete(constraint)	
+	
+        shoulderCtrlGrp = cmds.group(empty = True, name = "clavicle_" + name + "_anim_grp")
+        constraint = cmds.pointConstraint(fkArmJoint, shoulderCtrl)[0]
+        cmds.delete(constraint)
+        
+        cmds.parent(shoulderCtrl, shoulderCtrlGrp)
+        cmds.parent(shoulderCtrlGrp, autoClavWorld)	
+	
+        cmds.makeIdentity(shoulderCtrl, t = 1, r = 1, s = 1, apply = True)
+        cmds.setAttr(shoulderCtrl + ".sx", .65)
+        cmds.setAttr(shoulderCtrl + ".sy", 1.2)
+        cmds.setAttr(shoulderCtrl + ".sz", 1.2)
+        cmds.makeIdentity(shoulderCtrl, t = 1, r = 1, s = 1, apply = True)
+	
+	#lock out vis, rotates, and scale
+	for attr in [".sx", ".sy", ".sz", ".rx", ".ry", ".rz", ".v"]:
+	    cmds.setAttr(shoulderCtrl + attr, lock = False, keyable = False)
+	
+	#set the color
+	cmds.setAttr(shoulderCtrl + ".overrideEnabled", 1)
+	cmds.setAttr(shoulderCtrl + ".overrideColor", color)	
+	
+        #connect shoulder ctrl translate to the autoClavIKGrp translate
+        cmds.connectAttr(shoulderCtrl + ".translate", autoClavIKGrp + ".translate")
+        cmds.connectAttr(autoClavSwitchLoc + ".translate", shoulderCtrl + ".rotatePivotTranslate")	
+	
+	#clean up FK rig in scene
+	cmds.parent(invisFkArmCtrlGrp, autoClavWorld)
+	
+	
+	#create IK for invisible arm
+	invisRpIkHandle = cmds.ikHandle(name = "invis_arm_ikHandle_" + name, solver = "ikRPsolver", sj = invisUpArm, ee = invisHand)[0]
+	cmds.parent(invisRpIkHandle, "ik_wrist_" + name + "_anim")
+	cmds.poleVectorConstraint("ik_elbow_" + name + "_anim", invisRpIkHandle)
+	cmds.setAttr(invisRpIkHandle + ".v", 0)	
+
+	
+	#find children under autoClavWorld
+	children = cmds.listRelatives(autoClavWorld, children = True)
+	dntGrp = cmds.group(empty = True, name = "auto_clav_sys_grp_" +  name)
+	cmds.parent(dntGrp, autoClavWorld)
+	
+	for child in children:
+	    cmds.parent(child, dntGrp)
+	    
+	cmds.parent(shoulderCtrlGrp, autoClavWorld)	
+	cmds.parent([ikNodes, autoClavWorld], "arm_rig_master_grp_" + name)
+	
+	cmds.setAttr(dntGrp + ".v", 0)
+	cmds.setAttr(ikNodes + ".v", 0)	
+	
+	
+	
+	#setup autoShoulder attr
+	cmds.select(shoulderCtrl)
+	cmds.addAttr(longName='autoShoulders', defaultValue=0, minValue=0, maxValue=1, keyable = True)
+	
+	cmds.setAttr(shoulderCtrl + ".autoShoulders", 0)
+	
+	cmds.setAttr(autoClavSwitchConstraint + "." + autoClavEndJoint + "W0", 0)
+	cmds.setAttr(autoClavSwitchConstraint + "." + autoClavWorld + "W1", 1)
+	cmds.setDrivenKeyframe([autoClavSwitchConstraint + "." + autoClavEndJoint + "W0", autoClavSwitchConstraint + "." + autoClavWorld + "W1"], cd = shoulderCtrl + ".autoShoulders", itt = "linear", ott = "linear")
+	
+	
+	cmds.setAttr(shoulderCtrl + ".autoShoulders", 1)
+	cmds.setAttr(autoClavSwitchConstraint + "." + autoClavEndJoint + "W0", 1)
+	cmds.setAttr(autoClavSwitchConstraint + "." + autoClavWorld + "W1", 0)
+	cmds.setDrivenKeyframe([autoClavSwitchConstraint + "." + autoClavEndJoint + "W0", autoClavSwitchConstraint + "." + autoClavWorld + "W1"], cd = shoulderCtrl + ".autoShoulders", itt = "linear", ott = "linear")
+	
+	cmds.setAttr(shoulderCtrl + ".autoShoulders", 0)	
+	
+	
+
+	
+	return [ikClavJoint, autoClavWorld, shoulderCtrl, invisRpIkHandle]
+    
+    # # # # # # # # # # # # # # # # # # # # # # # # # # # # # # # # # # # # # # # # # # # # # # # # # # # # # # # # # # # # # # # 
+    # # # # # # # # # # # # # # # # # # # # # # # # # # # # # # # # # # # # # # # # # # # # # # # # # # # # # # # # # # # # # # # 
+    # # # # # # # # # # # # # # # # # # # # # # # # # # # # # # # # # # # # # # # # # # # # # # # # # # # # # # # # # # # # # # # 
+    # # # # # # # # # # # # # # # # # # # # # # # # # # # # # # # # # # # # # # # # # # # # # # # # # # # # # # # # # # # # # # # 
+    # # # # # # # # # # # # # # # # # # # # # # # # # # # # # # # # # # # # # # # # # # # # # # # # # # # # # # # # # # # # # # # 
+    def fkClavRig(self, clavJoint, fkArmJoint, ikArmJoint, fkArmCtrl, fkOrientNormal, armGrpParent, color):
+	
+	#create a nice name
+	name = self.prefix + self.name + self.suffix
+	if name.find("_") == 0:
+	    name = name.partition("_")[2]
+	    
+	#create the fk joint
+	fkClavJoint = cmds.duplicate("driver_" + self.prefix + self.clavicle + self.suffix, po = True, name = "fk_" + self.prefix + self.clavicle + self.suffix)[0]
+	
+	#create the control
+	control = utils.createControl("arrowOnBall", 4, "fk_clavicle_" + name + "_anim")
+	
+	#set the color
+	cmds.setAttr(control + ".overrideEnabled", 1)
+	cmds.setAttr(control + ".overrideColor", color)
+	
+	#position control
+	constraint = cmds.pointConstraint(fkArmJoint, control)[0]
+	cmds.delete(constraint)
+	
+	#move the pivot to the clav joint
+	piv = cmds.xform(fkClavJoint, q = True, ws = True, rotatePivot = True)
+	cmds.xform(control, ws = True, piv = [piv[0], piv[1], piv[2]])	
+	
+	
+	#create the control group
+	controlGrp = cmds.group(empty = True, name = "fk_clavicle_" + name + "_anim_grp")
+	constraint = cmds.parentConstraint(fkClavJoint, controlGrp)[0]
+	cmds.delete(constraint)
+	
+	#hierarchy and joint following control
+	cmds.parentConstraint(control, fkClavJoint)
+	cmds.parent(control, controlGrp)
+	cmds.parent(fkClavJoint, "joints_" + name + "_grp")
+	cmds.parent(controlGrp, "arm_rig_master_grp_" + name)
+	
+	cmds.makeIdentity(control, apply = True, t = 1, r = 1, s = 1)
+	for attr in [".sx", ".sy", ".sz", ".v"]:
+	    cmds.setAttr(control + attr, lock = True, keyable = False)
+	
+	return [fkClavJoint, controlGrp]
+	
+    # # # # # # # # # # # # # # # # # # # # # # # # # # # # # # # # # # # # # # # # # # # # # # # # # # # # # # # # # # # # # # # 
+    # # # # # # # # # # # # # # # # # # # # # # # # # # # # # # # # # # # # # # # # # # # # # # # # # # # # # # # # # # # # # # # 
+    # # # # # # # # # # # # # # # # # # # # # # # # # # # # # # # # # # # # # # # # # # # # # # # # # # # # # # # # # # # # # # # 
+    # # # # # # # # # # # # # # # # # # # # # # # # # # # # # # # # # # # # # # # # # # # # # # # # # # # # # # # # # # # # # # # 
+    # # # # # # # # # # # # # # # # # # # # # # # # # # # # # # # # # # # # # # # # # # # # # # # # # # # # # # # # # # # # # # # 
+    def upperArmTwistRig(self, color, rollGrpParent, fkArmJoint, ikArmJoint):	
+	
+	#create a nice name
+	name = self.prefix + self.name + self.suffix
+	if name.find("_") == 0:
+	    name = name.partition("_")[2]
+	    
+	#find number of twist joints
+	numRolls = 0
+	for joint in ["_twist_01", "_twist_02", "_twist_03"]:
+	    if cmds.objExists("driver_" + self.prefix + self.upperArm + joint + self.suffix):
+		numRolls = numRolls + 1
+		
+		
+	#create a nurbs plane for our ribbon
+	ribbon = cmds.nurbsPlane(ax = [0,0,1], lr = numRolls, width = 10, d = 3, u = 1, v = numRolls, ch = True, name = "upperarm_twist_ribbon_" + name)[0]
+	
+	#rebuild the ribbon with 1 U span
+	ribbon = cmds.rebuildSurface(ribbon, su = 1, du = 1, sv = numRolls, dv = 1, ch = 1)[0]
+	cmds.setAttr(ribbon + ".rz", -90)
+	cmds.makeIdentity(ribbon, apply = True, t = 1, r = 1, s = 1)    
+	
+	#create 2 temporary skin joints
+	moveVal = 0
+	for i in range(numRolls - 1):
+	    moveVal += 10
+	
+	cmds.select(clear = True)
+	topSkinJoint = cmds.joint(name = "top_skinJoint_temp")
+	cmds.move(moveVal, 0, 0, r = True, os = True, wd = True)
+	cmds.select(clear = True)
+	    
+	cmds.select(clear = True)
+	bottomSkinJoint = cmds.joint(name = "bottom_skinJoint_temp")
+	cmds.move(moveVal * -1, 0, 0, r = True, os = True, wd = True)
+	cmds.select(clear = True)
+	
+	#skin ribbon
+	cmds.select([ribbon, topSkinJoint, bottomSkinJoint])
+	skin = cmds.skinCluster(tsb = True, mi = 2, omi = True, dr = 5, bm = 0, sm = 0)
+	
+	
+	#position the joints, thus moving the ribbon
+	constraint = cmds.parentConstraint("driver_" + self.prefix + self.upperArm + self.suffix, topSkinJoint)[0]
+	cmds.delete(constraint)
+	
+	constraint = cmds.parentConstraint("driver_" + self.prefix + self.lowerArm + self.suffix, bottomSkinJoint)[0]
+	cmds.delete(constraint)	
+	
+	#delete ribbon history and skin joints
+	cmds.delete(ribbon, ch = True)
+	cmds.delete([bottomSkinJoint, topSkinJoint])
+	
+	
+	#create hair system on ribbon
+	cmds.select(ribbon)
+	mel.eval("createHair 1 3 10 0 0 0 0 5 0 2 1 1;")
+	
+	
+	#figure out which follicles created represent which areas on the ribbon
+	hairs = cmds.ls(type = "hairSystem")
+	
+	if len(hairs) > 0:
+	    hairSys = hairs[0]
+	    parent = cmds.listRelatives(hairs[0], parent = True)[0]
+	    hairSys = cmds.rename(parent, self.prefix + self.upperArm + "_twist" + self.suffix + "_hairSys")
+	    
+	follicles = cmds.listConnections(hairSys + "Shape", type = "follicle")
+	follicles = set(follicles)
+	hairFollicles = follicles	
+
+	#delete outputCurves
+	cmds.delete(parent + "OutputCurves")	
+
+
+	#create a joint per follicle
+	for follicle in hairFollicles:
+	    cmds.select(clear = True)
+	    joint = cmds.joint(name = follicle + "_joint")
+	    cmds.select(clear = True)
+	    constraint = cmds.parentConstraint(follicle, joint)[0]
+	    cmds.delete(constraint)
+	    
+	    cmds.parent(joint, follicle)
+	    cmds.makeIdentity(joint, apply = True, t = 0, r = 1, s = 0)
+	    
+	    
+	#create the skin joints (final)
+	skinJoints = []
+	for i in range(numRolls + 1):
+	    cmds.select(clear = True)
+	    skinJoint = cmds.joint(name = "skin_" + self.prefix + self.upperArm + "_twist_joint" + self.suffix + str(i))
+	    cmds.select(clear = True)
+	    skinJoints.append(skinJoint)
+	    
+	for i in range(numRolls):
+	    constraint = cmds.parentConstraint("driver_" + self.prefix + self.upperArm + "_twist_0" + str(i + 1) + self.suffix, skinJoints[i])[0]
+	    cmds.delete(constraint)
+	
+	constraint = cmds.parentConstraint("driver_" + self.prefix + self.lowerArm + self.suffix, skinJoint)[0]
+	cmds.delete(constraint)	
+	
+	
+	#create our manual control curves
+	x = 1
+	groups = []
+	for joint in skinJoints:
+	    
+	    if joint != skinJoints[-1]:
+		if joint == skinJoints[0]:
+		    grpName = self.prefix + self.upperArm + self.suffix + "_twist_anim_grp"
+		else:
+		    grpName = self.prefix + self.upperArm + self.suffix + "_twist_" + str(x) + "_anim_grp"
+		    
+		group = cmds.group(empty = True, name = grpName)
+		groups.append(group)
+		constraint = cmds.parentConstraint(joint, group)[0]
+		cmds.delete(constraint)
+		x = x + 1
+		
+	for i in range(int(len(groups))):
+	    grpName = groups[i].partition("_grp")[0]
+	    twistCtrl = utils.createControl("circle", 20, grpName)
+	    cmds.setAttr(twistCtrl + ".ry", -90)
+	    cmds.makeIdentity(twistCtrl, r = 1, apply =True)
+	    
+	    constraint = cmds.parentConstraint(groups[i], twistCtrl)[0]
+	    cmds.delete(constraint)
+	    
+	    cmds.parent(twistCtrl, groups[i])
+	    cmds.parent(skinJoints[i], twistCtrl)
+	    cmds.makeIdentity(skinJoints[i], apply = True, t = 0, r = 1, s = 0)
+	    
+	    #clean up control
+	    cmds.setAttr(twistCtrl + ".v", keyable = False)
+	    cmds.setAttr(twistCtrl + ".overrideEnabled", 1)
+	    cmds.setAttr(twistCtrl + ".overrideColor", color)	
+	    
+	
+	#organize groups
+	masterGrp = cmds.group(empty = True, name = "upperarm_twist_master_grp_" + name)
+	constraint = cmds.parentConstraint(fkArmJoint, masterGrp)[0]
+	cmds.delete(constraint)
+	
+	rollGrp = cmds.duplicate(masterGrp, name = "upperarm_twist_roll_grp_" + name)[0]
+	constraint = cmds.parentConstraint("driver_" + self.prefix + self.upperArm + self.suffix, rollGrp)[0]
+	cmds.delete(constraint)	
+	cmds.parent(rollGrp, masterGrp)
+	cmds.makeIdentity(rollGrp, apply = True, t = 1, r = 1, s = 1)
+	
+	#set rotate order on roll grp (xzy)
+	cmds.setAttr(rollGrp + ".rotateOrder", 3)
+    
+	for group in groups:
+	    cmds.parent(group, rollGrp)
+	    cmds.makeIdentity(group, apply = True, t = 1, r = 1, s = 1)
+	    
+	
+	#skin ribbon to skin joints
+	cmds.select(ribbon)
+	for joint in skinJoints:
+	    cmds.select(joint, add = True)
+	    
+	skin = cmds.skinCluster(tsb = True, mi = 2, omi = True, dr = 5, bm = 0, sm = 0)	
+	
+	
+	#orient roll grp to both fk/ik arm joints and set driven keys between them
+	upArmConstOrient = cmds.orientConstraint([fkArmJoint, ikArmJoint], rollGrp, mo = True)[0] #skip = x removed
+
+	cmds.setAttr("Rig_Settings." + name + "ArmMode", 0)
+	cmds.setAttr(upArmConstOrient + "." + fkArmJoint + "W0", 1)
+	cmds.setAttr(upArmConstOrient + "." + ikArmJoint + "W1", 0)
+	cmds.setDrivenKeyframe([upArmConstOrient + "." + fkArmJoint + "W0", upArmConstOrient + "." + ikArmJoint + "W1", ], cd = "Rig_Settings." + name + "ArmMode", itt = "linear", ott = "linear")
+
+	cmds.setAttr("Rig_Settings." + name + "ArmMode", 1)
+	cmds.setAttr(upArmConstOrient + "." + fkArmJoint + "W0", 0)
+	cmds.setAttr(upArmConstOrient + "." + ikArmJoint + "W1", 1)
+	cmds.setDrivenKeyframe([upArmConstOrient + "." + fkArmJoint + "W0", upArmConstOrient + "." + ikArmJoint + "W1", ], cd = "Rig_Settings." + name + "ArmMode", itt = "linear", ott = "linear")
+    
+	
+	
+	#parent end skin joint to masterGrp and orientConstrain twist to driver upper arm
+	cmds.parent(skinJoints[-1], rollGrp)
+	cmds.orientConstraint("driver_" + self.prefix + self.upperArm + self.suffix, skinJoints[-1], skip = ["y", "z"])
+	
+	
+	#parentConstraint master roll grp to driver clavicle
+	cmds.parentConstraint(rollGrpParent, masterGrp, mo = True)	
+	
+	
+	#hook up driver joints
+	hairFollicles = sorted(hairFollicles)
+	hairFollicles = hairFollicles[::-1]
+	num = 1
+	
+	for i in range(len(hairFollicles)):	    
+	    if cmds.objExists("driver_" + self.prefix + self.upperArm + "_twist_0" + str(num) + self.suffix):
+		cmds.orientConstraint(skinJoints[i], "driver_" + self.prefix + self.upperArm + "_twist_0" + str(num) + self.suffix)
+		cmds.pointConstraint(skinJoints[i], "driver_" + self.prefix + self.upperArm + "_twist_0" + str(num) + self.suffix, mo = True)
+		cmds.scaleConstraint(skinJoints[i], "driver_" + self.prefix + self.upperArm + "_twist_0" + str(num) + self.suffix)
+		num = num + 1	
+		
+	
+	#add attr on rig settings node for manual twist control visibility
+	cmds.select("Rig_Settings")
+	cmds.addAttr(longName=(name + "twistCtrlVis"), at = 'bool', dv = 0, keyable = True)
+	cmds.connectAttr("Rig_Settings." + name + "twistCtrlVis", rollGrp + ".v")		
+	
+	
+	#hook up multiply nodes so that twistAmount values from rig settings affect the ribbon twist
+        cmds.select("Rig_Settings")
+        cmds.addAttr(longName= ( name + "UpperarmTwistAmount" ), defaultValue= .9, minValue= 0 , maxValue= 1, keyable = True)	
+	
+	
+	#take twist ammount attr, multiply by -1, and feed into upperarm twist joint 1
+        multNodeA = cmds.shadingNode("multiplyDivide", asUtility = True, name = "upperarm_twist_" + name + "_multNodeA")
+        cmds.connectAttr("Rig_Settings." + name + "UpperarmTwistAmount", multNodeA + ".input1X")
+	cmds.setAttr(multNodeA + ".input2X", -1)
+	
+        multNodeB = cmds.shadingNode("multiplyDivide", asUtility = True, name = "upperarm_twist_" + name + "_multNodeB")
+        cmds.connectAttr(rollGrp + ".rx", multNodeB+ ".input1X")
+	cmds.connectAttr(multNodeA + ".outputX", multNodeB + ".input2X")
+        cmds.connectAttr(multNodeB + ".outputX", groups[0] + ".rx")	
+	
+	
+	
+	#any twist joints over the initial, setup simply mult nodes for carry down values
+	if numRolls > 1:
+	    for i in range(int(numRolls)):
+		if i == 1:
+		    cmds.select("Rig_Settings")
+		    cmds.addAttr(longName= ( name + "UpperarmTwist2Amount" ), defaultValue=.5, minValue=0, maxValue=1, keyable = True)
+		    
+		    #hook up multiply nodes so that twistAmount values from rig settings affect the ribbon twist
+		    multNode = cmds.shadingNode("multiplyDivide", asUtility = True, name = "upperarm_twist2_" + name + "_multNode")
+		    multNodeB = cmds.shadingNode("multiplyDivide", asUtility = True, name = "upperarm_twist2_" + name + "_multNodeB")
+		    blendNode = cmds.shadingNode("blendColors", asUtility = True, name = "upperarm_twist2_" + name + "_blendNode")
+		    
+		    #hook up blendnode to take in fk and ik upperarm rx values
+		    cmds.connectAttr( ikArmJoint + ".rx", blendNode + ".color1R")
+		    cmds.connectAttr( fkArmJoint+ ".rx", blendNode + ".color2R")
+		    
+		    #take output of that and plug into multNode. multiply by the twist ammount attribute value
+		    cmds.connectAttr(blendNode + ".outputR", multNode + ".input1X")
+		    cmds.connectAttr("Rig_Settings." + name + "UpperarmTwist2Amount" , multNodeB + ".input1X")
+		    cmds.setAttr(multNodeB + ".input2X", -1)
+		    cmds.connectAttr(multNodeB + ".outputX" , multNode + ".input2X")
+		    cmds.connectAttr(multNode + ".outputX", groups[i] + ".rx")
+		    
+		    #connect blendNode.blender to rig settings arm mode
+		    #cmds.connectAttr("Rig_Settings." + name + "ArmMode", blendNode + ".blender")
+		    
+		    #or set driven keyframes so blender attr doesn't go over 1
+		    cmds.setAttr("Rig_Settings." + name + "ArmMode", 0)
+		    cmds.setAttr(blendNode + ".blender", 0)
+		    cmds.setDrivenKeyframe(blendNode + ".blender", cd = "Rig_Settings." + name + "ArmMode", itt = "linear", ott = "linear")
+		    
+		    cmds.setAttr("Rig_Settings." + name + "ArmMode", 1)
+		    cmds.setAttr(blendNode + ".blender", 1)
+		    cmds.setDrivenKeyframe(blendNode + ".blender", cd = "Rig_Settings." + name + "ArmMode", itt = "linear", ott = "linear")
+		    
+		    
+
+		    
+		if i == 2:
+		    cmds.select("Rig_Settings")
+		    cmds.addAttr(longName= ( name + "UpperarmTwist3Amount" ), defaultValue=.5, minValue=0, maxValue=1, keyable = True)
+		    
+		    #hook up multiply nodes so that twistAmount values from rig settings affect the ribbon twist
+		    multNode = cmds.shadingNode("multiplyDivide", asUtility = True, name = "upperarm_twist3_" + name + "_multNode")
+		    multNodeB = cmds.shadingNode("multiplyDivide", asUtility = True, name = "upperarm_twist3_" + name + "_multNodeB")
+		    blendNode = cmds.shadingNode("blendColors", asUtility = True, name = "upperarm_twist3_" + name + "_blendNode")
+		    
+		    #hook up blendnode to take in fk and ik upperarm rx values
+		    cmds.connectAttr( ikArmJoint + ".rx", blendNode + ".color1R")
+		    cmds.connectAttr( fkArmJoint + ".rx", blendNode + ".color2R")
+		    
+		    #take output of that and plug into multNode. multiply by the twist ammount attribute value
+		    cmds.connectAttr(blendNode + ".outputR", multNode + ".input1X")
+		    cmds.connectAttr("Rig_Settings." + name + "UpperarmTwist3Amount" , multNodeB + ".input1X")
+		    cmds.setAttr(multNodeB + ".input2X", -1)
+		    cmds.connectAttr(multNodeB + ".outputX" , multNode + ".input2X")
+		    cmds.connectAttr(multNode + ".outputX", groups[i] + ".rx")
+		    
+		    #connect blendNode.blender to rig settings arm mode
+		    #cmds.connectAttr("Rig_Settings." + name+ "ArmMode", blendNode + ".blender")
+		    
+		    #or set driven keyframes so blender attr doesn't go over 1
+		    cmds.setAttr("Rig_Settings." + name + "ArmMode", 0)
+		    cmds.setAttr(blendNode + ".blender", 0)
+		    cmds.setDrivenKeyframe(blendNode + ".blender", cd = "Rig_Settings." + name + "ArmMode", itt = "linear", ott = "linear")
+		    
+		    cmds.setAttr("Rig_Settings." + name + "ArmMode", 1)
+		    cmds.setAttr(blendNode + ".blender", 1)
+		    cmds.setDrivenKeyframe(blendNode + ".blender", cd = "Rig_Settings." + name + "ArmMode", itt = "linear", ott = "linear")		    
+		    
+		    
+	#Group up and parent into rig
+	twistGrp = cmds.group(empty = True, name = "upperarm_twist_grp_" + name)
+	cmds.parent([ribbon, hairSys, masterGrp], twistGrp)
+	
+	#find follicles grp
+	for follicle in hairFollicles:
+		folliclesGrp = cmds.listRelatives(follicle, parent = True)	
+	cmds.parent(folliclesGrp[0], twistGrp)
+	
+	if cmds.objExists("nucleus1"):
+	    cmds.parent("nucleus1", twistGrp)
+	
+	#turn inherits transforms off
+	cmds.setAttr(folliclesGrp[0] + ".inheritsTransform", 0)
+	cmds.setAttr(ribbon + ".inheritsTransform", 0)
+	
+	#hide nodes
+	for node in [folliclesGrp[0], ribbon, hairSys, skinJoints[0]]:
+	    cmds.setAttr(node + ".v", 0)	
+	    
+	#add twist grp to arm_sys_grp
+	cmds.parent(twistGrp, "arm_sys_grp")
+    
+    
+    
+    # # # # # # # # # # # # # # # # # # # # # # # # # # # # # # # # # # # # # # # # # # # # # # # # # # # # # # # # # # # # # # # 
+    # # # # # # # # # # # # # # # # # # # # # # # # # # # # # # # # # # # # # # # # # # # # # # # # # # # # # # # # # # # # # # # 
+    # # # # # # # # # # # # # # # # # # # # # # # # # # # # # # # # # # # # # # # # # # # # # # # # # # # # # # # # # # # # # # # 
+    # # # # # # # # # # # # # # # # # # # # # # # # # # # # # # # # # # # # # # # # # # # # # # # # # # # # # # # # # # # # # # # 
+    # # # # # # # # # # # # # # # # # # # # # # # # # # # # # # # # # # # # # # # # # # # # # # # # # # # # # # # # # # # # # # # 
+    def lowerArmTwistRig(self, color, rollGrpParent, fkArmJoint, ikArmJoint):
+	
+	#create a nice name
+	name = self.prefix + self.name + self.suffix
+	if name.find("_") == 0:
+	    name = name.partition("_")[2]
+	    
+	    
+        #create our roll group
+        rollGrp = cmds.group(empty = True, name = self.prefix + self.lowerArm + self.suffix + "_roll_grp")
+        cmds.parentConstraint("driver_" + self.prefix + self.lowerArm + self.suffix, rollGrp)
+	cmds.parent(rollGrp, "arm_sys_grp")
+	            
+        #create our twist joint and twist mod joint
+        cmds.select(clear = True)
+        twistJoint = cmds.joint(name = self.prefix + self.lowerArm + self.suffix + "_twist_joint")
+        cmds.select(clear = True)
+        
+        constraint = cmds.parentConstraint("driver_" + self.prefix + self.lowerArm + "_twist_01" + self.suffix, twistJoint)[0]  
+        cmds.delete(constraint)
+        
+        cmds.parent(twistJoint, rollGrp)
+        
+        #twist mod joint
+        twistMod = cmds.duplicate(twistJoint, po = True, name = self.prefix + self.lowerArm + self.suffix + "_twist_mod")[0]
+        cmds.parent(twistMod, twistJoint)	
+	
+	
+        #create the manual twist control
+        twistCtrl = utils.createControl("circle", 15, self.prefix + self.lowerArm + self.suffix + "_twist_anim")
+        cmds.setAttr(twistCtrl + ".ry", -90)
+        cmds.makeIdentity(twistCtrl, r = 1, apply =True)
+        
+        constraint = cmds.parentConstraint(twistMod, twistCtrl)[0]
+        cmds.delete(constraint)
+        
+        twistCtrlGrp = cmds.group(empty = True, name = self.prefix + self.lowerArm + self.suffix + "_twist_anim_grp")
+        constraint = cmds.parentConstraint(twistMod, twistCtrlGrp)[0]
+        cmds.delete(constraint)
+        
+        cmds.parent(twistCtrl, twistCtrlGrp)
+        cmds.parent(twistCtrlGrp, twistMod)
+        cmds.makeIdentity(twistCtrl, t = 1, r = 1, s = 1, apply = True)
+	
+	cmds.setAttr(twistCtrl + ".overrideEnabled", 1)
+        cmds.setAttr(twistCtrl + ".overrideColor", color)
+	for attr in [".sx", ".sy", ".sz"]:
+	    cmds.setAttr(twistCtrl + attr, lock = True, keyable = False)
+	
+	cmds.setAttr(twistCtrl + ".v", keyable = False)	
+	
+	
+        #add attr on rig settings for manual twist control visibility
+        cmds.select("Rig_Settings")
+        cmds.addAttr(longName=(name + "twistCtrlVisLower"), at = 'bool', dv = 0, keyable = True)
+        cmds.connectAttr("Rig_Settings." + name + "twistCtrlVisLower", twistCtrl + ".v")
+        cmds.connectAttr("Rig_Settings." + name + "twistCtrlVisLower", twistMod + ".v")
+	cmds.connectAttr("Rig_Settings." + name + "twistCtrlVisLower", twistJoint + ".v")
+	cmds.setAttr(twistMod + ".radius", .01)
+	cmds.setAttr(twistJoint + ".radius", .01)
+	
+        #setup a simple relationship of foot rotateX value into mult node. input2X is driven by an attr on rig settings for twist amt(default is .5). Output into twist joint
+        twistMultNode = cmds.shadingNode("multiplyDivide", asUtility = True, name = self.prefix + self.lowerArm + "_twist" + self.suffix + "_mult_node")
+        
+        #add attr to rig settings
+        cmds.select("Rig_Settings")
+        cmds.addAttr(longName= ( name + "ForearmTwistAmount" ), defaultValue=.5, minValue=0, maxValue=1, keyable = True)
+	
+	
+        #connect output of driver hand into input1x
+        cmds.connectAttr("driver_" + self.prefix + self.hand + self.suffix + ".rx", twistMultNode + ".input1X")
+        
+        #connect attr into input2x
+        cmds.connectAttr("Rig_Settings." + name + "ForearmTwistAmount", twistMultNode + ".input2X")
+        
+        #connect output into driver calf twist
+        cmds.connectAttr(twistMultNode + ".outputX", twistJoint + ".rx")
+	
+	#constrain driver joint to twist joint
+	cmds.parentConstraint(twistCtrl, "driver_" + self.prefix + self.lowerArm + "_twist_01" + self.suffix,  mo = True)	
+	
+	
+	
+	#if there is more than 1 roll bone, set those up now:
+	numRolls = 0
+	for joint in ["_twist_01", "_twist_02", "_twist_03"]:
+	    if cmds.objExists("driver_" + self.prefix + self.lowerArm + joint + self.suffix):
+		numRolls = numRolls + 1
+	
+	
+	if numRolls > 1:
+	    for i in range(int(numRolls)):
+		
+		if i == 1:
+		    
+		    cmds.setAttr("Rig_Settings." + name + "ForearmTwistAmount", .75)
+		    cmds.select("Rig_Settings")
+		    cmds.addAttr(longName= ( name + "ForearmTwist2Amount" ), defaultValue=.5, minValue=0, maxValue=1, keyable = True)
+		    
+		    #create the manual twist control setup
+		    twistMod = cmds.duplicate("driver_" + self.prefix + self.lowerArm + "_twist_0" + str(i + 1) + self.suffix, po = True, name = self.prefix + self.lowerArm + self.suffix + "_twist2_mod")[0]
+		    cmds.parent(twistMod, rollGrp)
+		    
+        
+        
+		    #create the manual twist control
+		    twistCtrl = utils.createControl("circle", 15, self.prefix + self.lowerArm + self.suffix + "_twist2_anim")
+		    cmds.setAttr(twistCtrl + ".ry", -90)
+		    cmds.makeIdentity(twistCtrl, r = 1, apply =True)
+		    
+		    constraint = cmds.parentConstraint(twistMod, twistCtrl)[0]
+		    cmds.delete(constraint)
+		    
+		    twistCtrlGrp = cmds.group(empty = True, name = self.prefix + self.lowerArm + self.suffix + "_twist2_anim_grp")
+		    constraint = cmds.parentConstraint(twistMod, twistCtrlGrp)[0]
+		    cmds.delete(constraint)
+		    
+		    cmds.parent(twistCtrl, twistCtrlGrp)
+		    cmds.parent(twistCtrlGrp, twistMod)
+		    cmds.makeIdentity(twistCtrl, t = 1, r = 1, s = 1, apply = True)
+		    
+		    cmds.connectAttr("Rig_Settings." + name + "twistCtrlVisLower", twistCtrl + ".v")
+		    cmds.connectAttr("Rig_Settings." + name + "twistCtrlVisLower", twistMod + ".v")
+		    for attr in [".sx", ".sy", ".sz"]:
+			cmds.setAttr(twistCtrl + attr, lock = True, keyable = False)
+		    
+		    cmds.setAttr(twistCtrl + ".v", keyable = False)
+		    cmds.setAttr(twistCtrl + ".overrideEnabled", 1)
+		    cmds.setAttr(twistCtrl + ".overrideColor", color)
+	
+		    #drive the twist joint
+		    twistMultNode = cmds.shadingNode("multiplyDivide", asUtility = True, name = self.prefix + self.lowerArm + "_twist_2" + self.suffix + "_mult_node")
+		    cmds.connectAttr("driver_" + self.prefix + self.lowerArm + "_twist_01" + self.suffix + ".rx", twistMultNode + ".input1X")
+		    cmds.connectAttr("Rig_Settings." + name + "ForearmTwist2Amount", twistMultNode + ".input2X")
+		    cmds.connectAttr(twistMultNode + ".outputX", twistCtrlGrp + ".rx")
+		    cmds.parentConstraint(twistCtrl, "driver_" + self.prefix + self.lowerArm + "_twist_0" + str(i + 1) + self.suffix, mo = True)
+		    
+	
+		    
+
+		    
+		if i == 2:
+		    
+		    cmds.select("Rig_Settings")
+		    cmds.addAttr(longName= ( name + "ForearmTwist3Amount" ), defaultValue=.25, minValue=0, maxValue=1, keyable = True)
+		    
+		    #create the manual twist control setup
+		    twistMod = cmds.duplicate("driver_" + self.prefix + self.lowerArm + "_twist_0" + str(i + 1) + self.suffix, po = True, name = self.prefix + self.lowerArm + self.suffix + "_twist3_mod")[0]
+		    cmds.parent(twistMod, rollGrp)
+		    
+        
+        
+		    #create the manual twist control
+		    twistCtrl = utils.createControl("circle", 15, self.prefix + self.lowerArm + self.suffix + "_twist3_anim")
+		    cmds.setAttr(twistCtrl + ".ry", -90)
+		    cmds.makeIdentity(twistCtrl, r = 1, apply =True)
+		    
+		    constraint = cmds.parentConstraint(twistMod, twistCtrl)[0]
+		    cmds.delete(constraint)
+		    
+		    twistCtrlGrp = cmds.group(empty = True, name = self.prefix + self.lowerArm + self.suffix + "_twist3_anim_grp")
+		    constraint = cmds.parentConstraint(twistMod, twistCtrlGrp)[0]
+		    cmds.delete(constraint)
+		    
+		    cmds.parent(twistCtrl, twistCtrlGrp)
+		    cmds.parent(twistCtrlGrp, twistMod)
+		    cmds.makeIdentity(twistCtrl, t = 1, r = 1, s = 1, apply = True)
+		    
+		    cmds.connectAttr("Rig_Settings." + name + "twistCtrlVisLower", twistCtrl + ".v")
+		    cmds.connectAttr("Rig_Settings." + name + "twistCtrlVisLower", twistMod + ".v")
+		    for attr in [".sx", ".sy", ".sz"]:
+			cmds.setAttr(twistCtrl + attr, lock = True, keyable = False)
+		    
+		    cmds.setAttr(twistCtrl + ".v", keyable = False)
+		    cmds.setAttr(twistCtrl + ".overrideEnabled", 1)
+		    cmds.setAttr(twistCtrl + ".overrideColor", color)
+	
+		    #drive the twist joint
+		    twistMultNode = cmds.shadingNode("multiplyDivide", asUtility = True, name = self.prefix + self.lowerArm + "_twist_3" + self.suffix + "_mult_node")
+		    cmds.connectAttr("driver_" + self.prefix + self.lowerArm + "_twist_01" + self.suffix + ".rx", twistMultNode + ".input1X")
+		    cmds.connectAttr("Rig_Settings." + name + "ForearmTwist3Amount", twistMultNode + ".input2X")
+		    cmds.connectAttr(twistMultNode + ".outputX", twistCtrlGrp + ".rx")
+		    cmds.parentConstraint(twistCtrl, "driver_" + self.prefix + self.lowerArm + "_twist_0" + str(i + 1) + self.suffix, mo = True)
+		    
+		    
+		    
+    # # # # # # # # # # # # # # # # # # # # # # # # # # # # # # # # # # # # # # # # # # # # # # # # # # # # # # # # # # # # # # # 
+    # # # # # # # # # # # # # # # # # # # # # # # # # # # # # # # # # # # # # # # # # # # # # # # # # # # # # # # # # # # # # # # 
+    # # # # # # # # # # # # # # # # # # # # # # # # # # # # # # # # # # # # # # # # # # # # # # # # # # # # # # # # # # # # # # # 
+    # # # # # # # # # # # # # # # # # # # # # # # # # # # # # # # # # # # # # # # # # # # # # # # # # # # # # # # # # # # # # # # 
+    # # # # # # # # # # # # # # # # # # # # # # # # # # # # # # # # # # # # # # # # # # # # # # # # # # # # # # # # # # # # # # # 
+    def normalizeSubVector(vector1, vector2):
+	    import math
+	    returnVec = []
+	    for i in range(len(vector1)):
+		    returnVec.append(vector1[i] - vector2[i])
+	    
+	    #get length of vector
+	    length = math.sqrt( (returnVec[0] * returnVec[0]) + (returnVec[1] * returnVec[1]) + (returnVec[2] * returnVec[2]) )
+	    
+	    #normalize the vector
+	    normalizedVector = []
+	    for i in range(len(returnVec)):
+		    normalizedVector.append(returnVec[i]/length)
+	    
+	    absVector = []
+	    for i in range(len(normalizedVector)):
+		    absVector.append(abs(normalizedVector[i]))
+	    
+	    aimAxis = max(absVector)
+	    aimAxisIndex = absVector.index(aimAxis)
+	    
+	    if aimAxisIndex == 0:
+		    if normalizedVector[0] < 0:
+			    axis = "X"
+		    if normalizedVector[0] > 0:
+			    axis = "-X"
+    
+	    if aimAxisIndex == 1:
+		    if normalizedVector[1] < 0:
+			    axis = "Y"
+		    if normalizedVector[1] > 0:
+			    axis = "-Y"
+			    
+	    if aimAxisIndex == 2:
+		    if normalizedVector[2] < 0:
+			    axis = "Z"
+		    if normalizedVector[2] > 0:
+			    axis = "-Z"
+			    
+	    return axis
\ No newline at end of file
diff --git a/Engine/Extras/Maya_AnimationRiggingTools/MayaTools/General/Scripts/Modules/ART_Core.py b/Engine/Extras/Maya_AnimationRiggingTools/MayaTools/General/Scripts/Modules/ART_Core.py
new file mode 100644
--- /dev/null
+++ b/Engine/Extras/Maya_AnimationRiggingTools/MayaTools/General/Scripts/Modules/ART_Core.py
@@ -0,0 +1,143 @@
+import maya.cmds as cmds
+import maya.mel as mel
+import os
+import ART_rigUtils as utils
+reload(utils)
+
+
+class RigCore():
+    
+    #RigCore builds up our core components needed to start the rig build. This includes setting up the driver skeleton and building things like the rig settings, and master rig grps
+    #These are components that will be needed for every rig
+
+    def __init__(self):
+        
+        #create the rig settings node
+        self.rigSettings = cmds.group(empty = True, name = "Rig_Settings")
+        cmds.setAttr(self.rigSettings + ".tx", lock = True, keyable = False)
+        cmds.setAttr(self.rigSettings + ".ty", lock = True, keyable = False)
+        cmds.setAttr(self.rigSettings + ".tz", lock = True, keyable = False)
+        cmds.setAttr(self.rigSettings + ".rx", lock = True, keyable = False)
+        cmds.setAttr(self.rigSettings + ".ry", lock = True, keyable = False)
+        cmds.setAttr(self.rigSettings + ".rz", lock = True, keyable = False)
+        cmds.setAttr(self.rigSettings + ".sx", lock = True, keyable = False)
+        cmds.setAttr(self.rigSettings + ".sy", lock = True, keyable = False)
+        cmds.setAttr(self.rigSettings + ".sz", lock = True, keyable = False)
+        cmds.setAttr(self.rigSettings + ".v", lock = True, keyable = False)
+	
+	
+	#Setup the driver skeleton
+	self.createDriverSkeleton()
+	
+	#build the core rig components
+	self.buildCoreComponents()
+        
+        
+        
+    # # # # # # # # # # # # # # # # # # # # # # # # # # # # # # # # # # # # # # # # # # # # # # # # # # # # # # # # # # # # # # # 
+    # # # # # # # # # # # # # # # # # # # # # # # # # # # # # # # # # # # # # # # # # # # # # # # # # # # # # # # # # # # # # # # 
+    # # # # # # # # # # # # # # # # # # # # # # # # # # # # # # # # # # # # # # # # # # # # # # # # # # # # # # # # # # # # # # # 
+    # # # # # # # # # # # # # # # # # # # # # # # # # # # # # # # # # # # # # # # # # # # # # # # # # # # # # # # # # # # # # # # 
+    # # # # # # # # # # # # # # # # # # # # # # # # # # # # # # # # # # # # # # # # # # # # # # # # # # # # # # # # # # # # # # #     
+    def createDriverSkeleton(self):
+	
+	#there will always be a root bone, so let's duplicate that
+	dupe = cmds.duplicate("root", rc = True)[0]
+	cmds.select("root", hi = True)
+	joints = cmds.ls(sl = True)
+	
+	cmds.select(dupe, hi = True)
+	dupeJoints = cmds.ls(sl = True)
+	
+	driverJoints = []
+	for i in range(int(len(dupeJoints))):
+	    
+	    if cmds.objExists(dupeJoints[i]):
+		driverJoint = cmds.rename(dupeJoints[i], "driver_" + joints[i])
+		driverJoints.append(driverJoint)
+		
+
+	#create a direct connection between the driver and the export joints	
+	for joint in driverJoints:
+	    exportJoint = joint.partition("_")[2]
+	    
+	    cmds.connectAttr(joint + ".translate", exportJoint + ".translate")
+	    cmds.orientConstraint(joint, exportJoint)
+	    cmds.connectAttr(joint + ".scale", exportJoint + ".scale")
+	    
+	    
+	    
+    # # # # # # # # # # # # # # # # # # # # # # # # # # # # # # # # # # # # # # # # # # # # # # # # # # # # # # # # # # # # # # # 
+    # # # # # # # # # # # # # # # # # # # # # # # # # # # # # # # # # # # # # # # # # # # # # # # # # # # # # # # # # # # # # # # 
+    # # # # # # # # # # # # # # # # # # # # # # # # # # # # # # # # # # # # # # # # # # # # # # # # # # # # # # # # # # # # # # # 
+    # # # # # # # # # # # # # # # # # # # # # # # # # # # # # # # # # # # # # # # # # # # # # # # # # # # # # # # # # # # # # # # 
+    # # # # # # # # # # # # # # # # # # # # # # # # # # # # # # # # # # # # # # # # # # # # # # # # # # # # # # # # # # # # # # # 
+    def buildCoreComponents(self):
+        
+        #builds the master, the root, and the core rig groups
+	#MASTER CONTROL
+	masterControl = utils.createControl("circle", 150, "master_anim")
+	
+	constraint = cmds.pointConstraint("root", masterControl)[0]
+	cmds.delete(constraint)
+	
+	cmds.makeIdentity(masterControl, apply = True)
+	cmds.setAttr(masterControl + ".overrideEnabled", 1)
+	cmds.setAttr(masterControl + ".overrideColor", 18)
+	
+	spaceSwitchFollow = cmds.group(empty = True, name = masterControl + "_space_switcher_follow")
+	constraint = cmds.parentConstraint("root", spaceSwitchFollow)[0]
+	cmds.delete(constraint)
+	
+	spaceSwitcher = cmds.group(empty = True, name = masterControl + "_space_switcher")
+	constraint = cmds.parentConstraint("root", spaceSwitcher)[0]
+	cmds.delete(constraint)	
+	cmds.parent(spaceSwitcher, spaceSwitchFollow)
+	cmds.parent(masterControl, spaceSwitcher)
+	cmds.makeIdentity(masterControl, apply = True)
+	
+	
+	#OFFSET CONTROL
+	offsetControl = utils.createControl("square", 140, "offset_anim")
+	constraint = cmds.pointConstraint("root", offsetControl)[0]
+	cmds.delete(constraint)
+    
+	cmds.parent(offsetControl, masterControl)
+	cmds.makeIdentity(offsetControl, apply = True)
+	cmds.setAttr(offsetControl + ".overrideEnabled", 1)
+	cmds.setAttr(offsetControl + ".overrideColor", 17)
+	
+	
+	#ROOT ANIM
+	rootControl = utils.createControl("sphere", 10, "root_anim")
+	constraint = cmds.parentConstraint("driver_root", rootControl)[0]
+	cmds.delete(constraint)
+	cmds.parent(rootControl, masterControl)
+	cmds.makeIdentity(rootControl, apply = True)
+	cmds.parentConstraint(rootControl, "driver_root")
+	cmds.setAttr(rootControl + ".overrideEnabled", 1)
+	cmds.setAttr(rootControl + ".overrideColor", 30)
+	
+	for attr in [".sx", ".sy", ".sz", ".v"]:
+	    cmds.setAttr(masterControl + attr, lock = True, keyable = False)
+	    cmds.setAttr(offsetControl + attr, lock = True, keyable = False)
+	    cmds.setAttr(rootControl + attr, lock = True, keyable = False)
+      
+      
+	#Create the group that will hold all of the control rig components
+	rigGrp = cmds.group(empty = True, name = "ctrl_rig")
+	cmds.parent(rigGrp, "offset_anim")
+	    
+	    
+	#finish grouping everything under 1 character grp
+	controlRigGrp = cmds.group(empty = True, name = "rig_grp")
+	cmds.parent(["driver_root", "master_anim_space_switcher_follow"], controlRigGrp)
+	cmds.parent("Rig_Settings", controlRigGrp)
+	
+	if cmds.objExists("Proxy_Geo_Skin_Grp"):
+	    cmds.parent("Proxy_Geo_Skin_Grp", controlRigGrp)
+	    
+	    
+	    
+	returnNodes = [rigGrp, offsetControl]
+	return returnNodes
diff --git a/Engine/Extras/Maya_AnimationRiggingTools/MayaTools/General/Scripts/Modules/ART_Spine.py b/Engine/Extras/Maya_AnimationRiggingTools/MayaTools/General/Scripts/Modules/ART_Spine.py
new file mode 100644
--- /dev/null
+++ b/Engine/Extras/Maya_AnimationRiggingTools/MayaTools/General/Scripts/Modules/ART_Spine.py
@@ -0,0 +1,118 @@
+import maya.cmds as cmds
+import maya.mel as mel
+import os
+import ART_rigUtils as utils
+reload(utils)
+
+
+class Spine():
+    
+    #pass in items we need to know about (do we build a clavicle? What is the root name of the module?(If set to None, just use upperarm, lowerarm, etc), 
+    def __init__(self, includeClavicle, name, prefix, suffix, armGrpParent, color):
+        
+        
+        print "test"
+        
+        
+
+    # # # # # # # # # # # # # # # # # # # # # # # # # # # # # # # # # # # # # # # # # # # # # # # # # # # # # # # # # # # # # # # 
+    # # # # # # # # # # # # # # # # # # # # # # # # # # # # # # # # # # # # # # # # # # # # # # # # # # # # # # # # # # # # # # # 
+    # # # # # # # # # # # # # # # # # # # # # # # # # # # # # # # # # # # # # # # # # # # # # # # # # # # # # # # # # # # # # # # 
+    # # # # # # # # # # # # # # # # # # # # # # # # # # # # # # # # # # # # # # # # # # # # # # # # # # # # # # # # # # # # # # # 
+    # # # # # # # # # # # # # # # # # # # # # # # # # # # # # # # # # # # # # # # # # # # # # # # # # # # # # # # # # # # # # # # 
+    def buildFKSpine(self):
+	
+	#find the number of spine bones from the skeleton settings
+	spineJoints = self.getSpineJoints()
+	fkControls = []
+	parent = None
+	
+	for joint in spineJoints:
+	    if joint == "spine_01":
+		#add space switcher node to base of spine
+		spaceSwitcherFollow = cmds.group(empty = True, name = joint + "_space_switcher_follow")
+		constraint = cmds.parentConstraint(joint, spaceSwitcherFollow)[0]
+		cmds.delete(constraint)
+		spaceSwitcher = cmds.duplicate(spaceSwitcherFollow, name = joint + "_space_switcher")[0]
+		cmds.parent(spaceSwitcher, spaceSwitcherFollow)
+		
+	    #create an empty group in the same space as the joint
+	    group = cmds.group(empty = True, name = joint + "_anim_grp")
+	    constraint = cmds.parentConstraint(joint, group)[0]
+	    cmds.delete(constraint)
+	    
+	    #create an additional layer of group that has zeroed attrs
+	    offsetGroup = cmds.group(empty = True, name = joint + "_anim_offset_grp")
+	    constraint = cmds.parentConstraint(joint, offsetGroup)[0]
+	    cmds.delete(constraint)
+	    cmds.parent(offsetGroup, group)
+	    
+	    #create a control object in the same space as the joint
+	    control = utils.createControl("circle", 45, joint + "_anim")
+	    tempDupe = cmds.duplicate(control)[0]
+	    constraint = cmds.parentConstraint(joint, control)[0]
+	    cmds.delete(constraint)
+	    fkControls.append(control)
+	    
+	    
+	    #parent the control object to the group
+	    cmds.parent(control, offsetGroup)
+	    constraint = cmds.orientConstraint(tempDupe, control, skip = ["x", "z"])[0]
+	    cmds.delete(constraint)
+	    cmds.makeIdentity(control, t = 1, r = 1, s = 1, apply = True)
+
+	    
+	    #setup hierarchy
+	    if parent != None:
+		cmds.parent(group, parent, absolute = True)
+		
+	    else:
+		cmds.parent(group, spaceSwitcher)
+		cmds.parent(spaceSwitcherFollow, "body_anim")
+		
+		
+	    #set the parent to be the current spine control
+	    parent = control
+	    
+	    #clean up
+	    cmds.delete(tempDupe)
+	    
+	    for attr in [".v"]:
+		cmds.setAttr(control + attr, lock = True, keyable = False)
+		
+	    #set the control's color
+	    cmds.setAttr(control + ".overrideEnabled", 1)
+	    cmds.setAttr(control + ".overrideColor", 18)
+	    
+
+	return fkControls
+    
+    
+
+    # # # # # # # # # # # # # # # # # # # # # # # # # # # # # # # # # # # # # # # # # # # # # # # # # # # # # # # # # # # # # # # 
+    # # # # # # # # # # # # # # # # # # # # # # # # # # # # # # # # # # # # # # # # # # # # # # # # # # # # # # # # # # # # # # # 
+    # # # # # # # # # # # # # # # # # # # # # # # # # # # # # # # # # # # # # # # # # # # # # # # # # # # # # # # # # # # # # # # 
+    # # # # # # # # # # # # # # # # # # # # # # # # # # # # # # # # # # # # # # # # # # # # # # # # # # # # # # # # # # # # # # # 
+    # # # # # # # # # # # # # # # # # # # # # # # # # # # # # # # # # # # # # # # # # # # # # # # # # # # # # # # # # # # # # # # 
+    def getSpineJoints(self):
+        
+        if cmds.objExists("Skeleton_Settings"):
+	    numSpineBones = int(cmds.getAttr("Skeleton_Settings.numSpineBones"))
+	    spineJoints = []
+	    
+	    for i in range(int(numSpineBones)):
+		
+		if i < 10:
+		    spineJoint = "spine_0" + str(i + 1)
+		    
+		else:
+		    spineJoint = "spine_" + (i + 1)
+		    
+		spineJoints.append(spineJoint)
+		
+	    return spineJoints
+	
+	else:
+	    cmds.confirmDialog(title = "Error", icon = "critical", message = "Needed node, Skeleton_Settings could not be found.")
+	    return
+	
\ No newline at end of file
