Adding ART scripts to repository

diff --git a/Engine/Extras/Maya_AnimationRiggingTools/MayaTools/General/Scripts/Modules/ART_rigUtils.py b/Engine/Extras/Maya_AnimationRiggingTools/MayaTools/General/Scripts/Modules/ART_rigUtils.py
new file mode 100644
--- /dev/null
+++ b/Engine/Extras/Maya_AnimationRiggingTools/MayaTools/General/Scripts/Modules/ART_rigUtils.py
@@ -0,0 +1,154 @@
+import maya.cmds as cmds
+import maya.mel as mel
+import os
+
+#rigging utilities
+
+
+# # # # # # # # # # # # # # # # # # # # # # # # # # # # # # # # # # # # # # # # # # # # # # # # # # # # # # # # # # # # # # # 
+# # # # # # # # # # # # # # # # # # # # # # # # # # # # # # # # # # # # # # # # # # # # # # # # # # # # # # # # # # # # # # # 
+# # # # # # # # # # # # # # # # # # # # # # # # # # # # # # # # # # # # # # # # # # # # # # # # # # # # # # # # # # # # # # # 
+# # # # # # # # # # # # # # # # # # # # # # # # # # # # # # # # # # # # # # # # # # # # # # # # # # # # # # # # # # # # # # # 
+# # # # # # # # # # # # # # # # # # # # # # # # # # # # # # # # # # # # # # # # # # # # # # # # # # # # # # # # # # # # # # # 
+def createControl(controlType, size, name):
+    
+    scale = getScaleFactor()
+            
+    if controlType == "circle":
+        control = cmds.circle(c = (0,0,0), sw = 360, r = size * scale, d = 3, name = name)[0]
+        
+    if controlType == "circleSpecial":
+        control = cmds.circle(c = (0,0,0), sw = 360, r = 1, d = 3, name = name)[0]
+        
+        side = name.rpartition("_")[2]
+        
+        if side == "l":
+            cmds.xform(control, piv = (0, -1, 0))
+            
+        else:
+            cmds.xform(control, piv = (0, 1, 0))
+            
+        cmds.setAttr(control + ".scaleX", size * scale)
+        cmds.setAttr(control + ".scaleY", size * scale)
+        cmds.setAttr(control + ".scaleZ", size * scale)
+
+    if controlType == "square":
+        control = cmds.circle(c = (0,0,0), s = 4, sw = 360, r = size * scale, d = 1, name = name)[0]
+        cmds.setAttr(control + ".rz", 45)
+        
+    if controlType == "foot":
+        control = cmds.curve(name = name, d = 3, p = [(0, 40, 0), (-3.42, 39, 0), (-10.2, 37, 0), (-13, 22, 0), (-15.7, 13.2, 0), (-20, -14, 0), (-18.1, -25.6, 0), (-15, -44.8, 0), (1.1, -41.2, 0), (4.8, -41.7, 0), (15.5, -31.9, 0), (16.9, -22.7, 0), (18.6, -15.2, 0), (16.5, -.5, 0), (11.2, 29.2, 0), (10.7, 39.7, 0), (3.6, 39.9, 0), (0, 40, 0)])
+        
+        footLoc = cmds.spaceLocator(name = (name + "_end_loc"))[0]
+        cmds.parent(footLoc, control)
+        cmds.setAttr(footLoc + ".ty", -40)
+        cmds.setAttr(footLoc + ".v", 0)
+                     
+        cmds.setAttr(control + ".scaleX", size * scale)
+        cmds.setAttr(control + ".scaleY", size * scale)
+        cmds.setAttr(control + ".scaleZ", size * scale)
+        
+    
+    if controlType == "arrow":
+        control = cmds.curve(name = name, d = 1, p = [(0, -45, 0), (5, -45, 0), (5, -62, 0), (10, -62, 0), (0, -72, 0), (-10, -62, 0), (-5, -62, 0), (-5, -45, 0), (0, -45, 0)])
+        cmds.xform(control, cp = True)
+        cmds.setAttr(control + ".ty", 58.5)
+        cmds.makeIdentity(control, t = 1, apply = True)
+        cmds.xform(control, piv = (0, 13.5, 0))
+        
+        cmds.setAttr(control + ".scaleX", size * scale)
+        cmds.setAttr(control + ".scaleY", size * scale)
+        cmds.setAttr(control + ".scaleZ", size * scale)
+        
+    if controlType == "arrowOnBall":
+        
+        control = cmds.curve(name = name, d = 1, p = [(0.80718, 0.830576, 8.022739), (0.80718, 4.219206, 7.146586 ), (0.80718, 6.317059, 5.70073), (2.830981, 6.317059, 5.70073), (0, 8.422749, 2.94335), (-2.830981, 6.317059, 5.70073), (-0.80718, 6.317059, 5.70073), (-0.80718, 4.219352, 7.146486), (-0.80718, 0.830576, 8.022739), (-4.187851, 0.830576, 7.158003), (-6.310271, 0.830576, 5.705409), (-6.317059, 2.830981, 5.7007), (-8.422749, 0, 2.94335), (-6.317059, -2.830981, 5.70073), (-6.317059, -0.830576, 5.70073), (-4.225134, -0.830576, 7.142501), (-0.827872, -0.830576, 8.017446), (-0.80718, -4.176512, 7.160965), (-0.80718, -6.317059, 5.70073), (-2.830981, -6.317059, 5.70073), (0, -8.422749, 2.94335), (2.830981, -6.317059, 5.70073), (0.80718, -6.317059, 5.70073), (0.80718, -4.21137, 7.151987), (0.80718, -0.830576, 8.022739), (4.183345, -0.830576, 7.159155), (6.317059, -0.830576, 5.70073), (6.317059, -2.830981, 5.70073), (8.422749, 0, 2.94335), (6.317059, 2.830981, 5.70073), (6.317059, 0.830576, 5.70073), (4.263245, 0.830576, 7.116234), (0.80718, 0.830576, 8.022739)])
+        
+        cmds.setAttr(control + ".scaleX", size * scale)
+        cmds.setAttr(control + ".scaleY", size * scale)
+        cmds.setAttr(control + ".scaleZ", size * scale)
+        
+        
+    if controlType == "semiCircle":
+        
+        control = cmds.curve(name = name, d = 3, p = [(0,0,0), (7, 0, 0), (8, 0, 0), (5, 4, 0), (0, 5, 0), (-5, 4, 0), (-8, 0, 0), (-7, 0, 0), (0,0,0)])
+        cmds.xform(control, ws = True, t = (0, 5, 0))
+        cmds.xform(control, ws = True, piv = (0,0,0))
+        cmds.makeIdentity(control, t = 1, apply = True)
+        
+        
+        cmds.setAttr(control + ".scaleX", size * scale)
+        cmds.setAttr(control + ".scaleY", size * scale)
+        cmds.setAttr(control + ".scaleZ", size * scale)
+        
+        
+    if controlType == "pin":
+        control = cmds.curve(name = name, d = 1, p = [(12,0,0), (0, 0, 0), (-12, -12, 0), (-12, 12, 0), (0, 0, 0)])
+        cmds.xform(control, ws = True, piv = [12,0,0])
+        cmds.setAttr(control + ".scaleY", .5)
+        cmds.makeIdentity(control, t = 1, apply = True)
+        
+        
+        cmds.setAttr(control + ".scaleX", size * scale)
+        cmds.setAttr(control + ".scaleY", size * scale)
+        cmds.setAttr(control + ".scaleZ", size * scale)
+        
+    
+    if controlType == "sphere":
+
+        points = [(0, 0, 1), (0, 0.5, 0.866), (0, 0.866025, 0.5), (0, 1, 0), (0, 0.866025, -0.5), (0, 0.5, -0.866025), (0, 0, -1), (0, -0.5, -0.866025), (0, -0.866025, -0.5), (0, -1, 0), (0, -0.866025, 0.5), (0, -0.5, 0.866025), (0, 0, 1), (0.707107, 0, 0.707107), (1, 0, 0), (0.707107, 0, -0.707107), (0, 0, -1), (-0.707107, 0, -0.707107), (-1, 0, 0), (-0.866025, 0.5, 0), (-0.5, 0.866025, 0), (0, 1, 0), (0.5, 0.866025, 0), (0.866025, 0.5, 0), (1, 0, 0), (0.866025, -0.5, 0), (0.5, -0.866025, 0), (0, -1, 0), (-0.5, -0.866025, 0), (-0.866025, -0.5, 0), (-1, 0, 0), (-0.707107, 0, 0.707107), (0, 0, 1)]
+        control = cmds.curve(name = name, d = 1, p = points)
+        
+        cmds.setAttr(control + ".scaleX", size * scale)
+        cmds.setAttr(control + ".scaleY", size * scale)
+        cmds.setAttr(control + ".scaleZ", size * scale)
+        
+
+    return control    
+
+
+# # # # # # # # # # # # # # # # # # # # # # # # # # # # # # # # # # # # # # # # # # # # # # # # # # # # # # # # # # # # # # # 
+# # # # # # # # # # # # # # # # # # # # # # # # # # # # # # # # # # # # # # # # # # # # # # # # # # # # # # # # # # # # # # # 
+# # # # # # # # # # # # # # # # # # # # # # # # # # # # # # # # # # # # # # # # # # # # # # # # # # # # # # # # # # # # # # # 
+# # # # # # # # # # # # # # # # # # # # # # # # # # # # # # # # # # # # # # # # # # # # # # # # # # # # # # # # # # # # # # # 
+# # # # # # # # # # # # # # # # # # # # # # # # # # # # # # # # # # # # # # # # # # # # # # # # # # # # # # # # # # # # # # # 
+def getScaleFactor():
+    
+    headLoc = cmds.spaceLocator(name = "headLoc")[0]
+    cmds.parentConstraint("driver_head", headLoc)
+    
+    height = cmds.getAttr(headLoc + ".tz")
+    
+    defaultHeight = 400
+    scaleFactor = height/defaultHeight
+    cmds.delete(headLoc)
+    
+    return scaleFactor    
+
+# # # # # # # # # # # # # # # # # # # # # # # # # # # # # # # # # # # # # # # # # # # # # # # # # # # # # # # # # # # # # # # 
+# # # # # # # # # # # # # # # # # # # # # # # # # # # # # # # # # # # # # # # # # # # # # # # # # # # # # # # # # # # # # # # 
+# # # # # # # # # # # # # # # # # # # # # # # # # # # # # # # # # # # # # # # # # # # # # # # # # # # # # # # # # # # # # # # 
+# # # # # # # # # # # # # # # # # # # # # # # # # # # # # # # # # # # # # # # # # # # # # # # # # # # # # # # # # # # # # # # 
+# # # # # # # # # # # # # # # # # # # # # # # # # # # # # # # # # # # # # # # # # # # # # # # # # # # # # # # # # # # # # # # 
+def getUpAxis(obj):
+    
+    cmds.xform(obj, ws = True, relative = True, t = [0, 0, 10])
+    translate = cmds.getAttr(obj + ".translate")[0]
+    newTuple = (abs(translate[0]), abs(translate[1]), abs(translate[2]))
+    cmds.xform(obj, ws = True, relative = True, t = [0, 0, -10])
+    
+    highestVal = max(newTuple)
+    axis = newTuple.index(highestVal)
+    upAxis = None
+    
+    if axis == 0:
+        upAxis = "X"
+        
+    if axis == 1:
+        upAxis = "Y"
+        
+    if axis == 2:
+        upAxis = "Z"
+        
+    return upAxis
+    
\ No newline at end of file
diff --git a/Engine/Extras/Maya_AnimationRiggingTools/MayaTools/General/Scripts/customMayaMenu.py b/Engine/Extras/Maya_AnimationRiggingTools/MayaTools/General/Scripts/customMayaMenu.py
new file mode 100644
--- /dev/null
+++ b/Engine/Extras/Maya_AnimationRiggingTools/MayaTools/General/Scripts/customMayaMenu.py
@@ -0,0 +1,520 @@
+import maya.cmds as cmds
+import maya.mel as mel
+import os, cPickle
+from functools import partial
+
+def customMayaMenu():
+    gMainWindow = mel.eval('$temp1=$gMainWindow')
+    
+    
+    menus = cmds.window(gMainWindow, q = True, menuArray = True)
+    found = False
+    for menu in menus:
+        label = cmds.menu(menu, q = True, label = True)
+        if label == "Epic Games":
+            found = True
+    
+    if found == False:
+        customMenu = cmds.menu(parent=gMainWindow, label = 'Epic Games')
+        
+        #tools path
+	toolsPath = cmds.internalVar(usd = True) + "mayaTools.txt"
+	if os.path.exists(toolsPath):
+	    
+	    f = open(toolsPath, 'r')
+	    mayaToolsDir = f.readline()
+	    f.close()	
+	    
+        #ART
+	cmds.menuItem(parent = customMenu, label = "Animation Rigging Toolset", bld = True, enable = False)
+	cmds.menuItem(parent = customMenu, divider = True)
+	cmds.menuItem(parent = customMenu, label = "Character Rig Creator", c = launchSkeletonBuilder)
+	cmds.menuItem(parent = customMenu, label = "Edit Existing Character", c = editCharacter)
+	cmds.menuItem(parent = customMenu, label = "Add Character For Animation", c = launchAddCharacter)
+	cmds.menuItem(parent = customMenu, label = "Animation Interface", c = launchAnimUI)
+	
+	cmds.menuItem(parent = customMenu, label = "Settings", c = launchARTSettings)
+        artHelp = cmds.menuItem(parent = customMenu, label = "Help", subMenu=True)
+	cmds.menuItem(parent = artHelp, label = "Learning Videos", c = launchLearningVideos)
+	cmds.menuItem(parent = artHelp, label = "Help Documentation", c = launchRigHelp)
+        cmds.menuItem(parent = artHelp, label = "About", c = aboutARTTools)
+	cmds.menuItem(parent = customMenu, divider = True)
+	cmds.menuItem(parent = customMenu, label = "Misc. Tools", bld = True, enable = False)
+	cmds.menuItem(parent = customMenu, divider = True)
+	
+
+
+	#PERFORCE
+	p4Menu = cmds.menuItem(parent = customMenu, label = "Perforce", subMenu=True, to = True)
+	cmds.menuItem("perforceSubmitMenuItem", parent = p4Menu, label = "Submit Current File", enable = False, c = p4Submit)
+	cmds.menuItem("perforceAddAndSubmitMenuItem", parent = p4Menu, label = "Add and Submit Current File", enable = False, c = p4AddSubmit)
+	cmds.menuItem("perforceCheckOutMenuItem", parent = p4Menu, label = "Check Out Current File", enable = False, c = p4CheckOut)
+	cmds.menuItem("perforceFileHistoryMenuItem", parent = p4Menu, label = "Current File History", enable = False, c = p4GetHistory)
+	cmds.menuItem("perforceGetLatestMenuItem", parent = p4Menu, label = "Get Latest Revision of Current File", enable = False, c = p4GetLatest)
+	cmds.menuItem("perforceProjectList", parent = p4Menu, label = "Set Project", enable = False, subMenu = True)
+	cmds.menuItem("perforceProject_New", parent = "perforceProjectList", label = "New Project", c = createNewProj)
+	cmds.radioMenuItemCollection("perforceProjectRadioMenuCollection", parent = "perforceProjectList")
+	cmds.menuItem(parent = customMenu, divider = True)
+	
+	
+
+	
+	
+	
+	
+	#check settings to see if use source control is turned on
+	toolsPath = cmds.internalVar(usd = True) + "mayaTools.txt"
+	if os.path.exists(toolsPath):
+	    
+	    f = open(toolsPath, 'r')
+	    mayaToolsDir = f.readline()
+	    f.close()	
+
+	    settingsLocation = mayaToolsDir + "/General/Scripts/projectSettings.txt"
+	    
+	    if os.path.exists(settingsLocation):
+		f = open(settingsLocation, 'r')
+		settings = cPickle.load(f)
+		f.close()	
+		
+		#find use source control value
+		sourceControl = settings.get("UseSourceControl")
+		
+		if sourceControl:
+		    cmds.menuItem("perforceSubmitMenuItem", edit = True, enable = True)
+		    cmds.menuItem("perforceAddAndSubmitMenuItem", edit = True, enable = True)
+		    cmds.menuItem("perforceCheckOutMenuItem", edit = True, enable = True)
+		    cmds.menuItem("perforceFileHistoryMenuItem", edit = True, enable = True)
+		    cmds.menuItem("perforceGetLatestMenuItem", edit = True, enable = True)
+		    cmds.menuItem("perforceProjectList", edit = True, enable = True)
+		    
+		    #launch script job for checking Maya Tools
+		    cmds.scriptJob(event = ["NewSceneOpened", autoUpdateTools])
+		    
+
+	    
+
+    
+#############################################################################################
+#############################################################################################
+#############################################################################################
+def p4ProjectMenu(*args):
+    
+    #clear any projects that are in the collection first
+    items = cmds.lsUI(menuItems = True)
+    for i in items:
+	data = cmds.menuItem(i, q = True, docTag = True)
+	if data == "P4Proj":
+	    cmds.deleteUI(i)
+	    
+    #find projects
+    toolsPath = cmds.internalVar(usd = True) + "mayaTools.txt"
+    if os.path.exists(toolsPath):
+	
+	f = open(toolsPath, 'r')
+	mayaToolsDir = f.readline()
+	f.close()
+	
+    projects = os.listdir(mayaToolsDir + "/General/Scripts/")
+    p4Projects = []
+    #Test_Project_Settings
+    for proj in projects:
+	if proj.rpartition(".")[2] == "txt":
+	    if proj.partition("_")[2].partition("_")[0] == "Project":
+		p4Projects.append(proj)
+
+
+    #set the current project
+    try:
+	f = open(mayaToolsDir + "/General/Scripts/projectSettings.txt", 'r')
+	settings = cPickle.load(f)
+	f.close()
+	currentProj = settings.get("CurrentProject")
+    except:
+	pass
+    
+
+    
+    #add the projects to the menu	
+    for proj in p4Projects:
+	projectName = proj.partition("_")[0]
+	
+	if projectName == currentProj:
+	    val = True
+	else:
+	    val = False
+	    
+	menuItem = cmds.menuItem(label = projectName, parent = "perforceProjectList", cl = "perforceProjectRadioMenuCollection", rb = val, docTag = "P4Proj", c = partial(setProj, projectName))
+	cmds.menuItem(parent = "perforceProjectList", optionBox = True, c = partial(editProj, projectName))	
+	
+#############################################################################################
+#############################################################################################
+#############################################################################################
+def setProj(projectName, *args):
+    import perforceUtils
+    reload(perforceUtils)
+    perforceUtils.setCurrentProject(projectName)
+    
+#############################################################################################
+#############################################################################################
+#############################################################################################
+def editProj(projectName, *args):
+    import perforceUtils
+    reload(perforceUtils)
+    perforceUtils.editProject(projectName)
+    
+#############################################################################################
+#############################################################################################
+#############################################################################################
+def createNewProj(*args):
+    import perforceUtils
+    reload(perforceUtils)
+    perforceUtils.createNewProject()
+    
+#############################################################################################
+#############################################################################################
+#############################################################################################
+def autoUpdateTools(*args):
+    import perforceUtils
+    reload(perforceUtils)
+    perforceUtils.p4_checkForUpdates()
+    
+
+#############################################################################################
+#############################################################################################
+#############################################################################################
+def launchARTSettings(*args):
+    import ART_Settings
+    reload(ART_Settings)
+    ART_Settings.ART_Settings()
+    
+#############################################################################################
+#############################################################################################
+#############################################################################################
+def aboutARTTools(*args):
+    
+    cmds.confirmDialog(title = "About", message = "(c) Epic Games, Inc. 2013\nCreated by: [REDACTED_PERSON]\[REDACTED_EMAIL]\nVisit www.epicgames.com", icon = "information")
+    
+#############################################################################################
+#############################################################################################
+#############################################################################################
+def editCharacter(*args):
+    
+    if cmds.window("artEditCharacterUI", exists = True):
+        cmds.deleteUI("artEditCharacterUI")
+    
+    window = cmds.window("artEditCharacterUI", w = 300, h = 400, title = "Edit Character", mxb = False, mnb = False, sizeable = False)
+    mainLayout = cmds.columnLayout(w = 300, h = 400, rs = 5, co = ["both", 5])
+    
+    #banner image
+    toolsPath = cmds.internalVar(usd = True) + "mayaTools.txt"
+    if os.path.exists(toolsPath):
+        
+        f = open(toolsPath, 'r')
+        mayaToolsDir = f.readline()
+        f.close()
+	
+    cmds.image(w = 300, h = 50, image = mayaToolsDir + "/General/Icons/ART/artBanner300px.bmp", parent = mainLayout)
+    
+    cmds.text(label = "", h = 1, parent = mainLayout)
+    optionMenu = cmds.optionMenu("artProjOptionMenu", label = "Project:", w =290, h = 40, cc = getProjCharacters, parent = mainLayout)
+    textScrollList = cmds.textScrollList("artProjCharacterList", w = 290, h = 300, parent = mainLayout)
+    button = cmds.button(w = 290, h = 40, label = "Edit Export File", c = editSelectedCharacter, ann = "Edit the character's skeleton settings, joint positions, or skin weights.", parent = mainLayout)
+    button2 = cmds.button(w = 290, h = 40, label = "Edit Rig File", c = editSelectedCharacterRig, ann = "Edit the character's control rig that will be referenced in by animation.", parent = mainLayout)
+    
+    cmds.text(label = "", h = 1)
+    
+    
+    cmds.showWindow(window)
+    getProjects()
+    getProjCharacters()
+    
+#############################################################################################
+#############################################################################################
+#############################################################################################
+def getProjects(*args):
+    
+    toolsPath = cmds.internalVar(usd = True) + "mayaTools.txt"
+    if os.path.exists(toolsPath):
+        
+        f = open(toolsPath, 'r')
+        mayaToolsDir = f.readline()
+        f.close()
+        
+    projects = os.listdir(mayaToolsDir + "/General/ART/Projects/")
+    for project in projects:
+        cmds.menuItem(label = project, parent = "artProjOptionMenu")
+
+#############################################################################################
+#############################################################################################
+#############################################################################################        
+def getProjCharacters(*args):
+    toolsPath = cmds.internalVar(usd = True) + "mayaTools.txt"
+    if os.path.exists(toolsPath):
+        
+        f = open(toolsPath, 'r')
+        mayaToolsDir = f.readline()
+        f.close()
+        
+    proj = cmds.optionMenu("artProjOptionMenu", q = True, value = True)
+    
+    cmds.textScrollList("artProjCharacterList", edit = True, removeAll = True)
+    characters = os.listdir(mayaToolsDir + "/General/ART/Projects/" + proj + "/ExportFiles/")
+    
+    for character in characters:
+	if os.path.isfile(mayaToolsDir + "/General/ART/Projects/" + proj + "/ExportFiles/" + character):
+	    if character.rpartition(".")[2] == "mb":
+		niceName = character.rpartition(".")[0]
+		niceName = niceName.partition("_Export")[0]
+		cmds.textScrollList("artProjCharacterList", edit = True, append = niceName)
+        
+#############################################################################################
+#############################################################################################
+#############################################################################################
+def editSelectedCharacter(*args):
+    toolsPath = cmds.internalVar(usd = True) + "mayaTools.txt"
+    if os.path.exists(toolsPath):
+        
+        f = open(toolsPath, 'r')
+        mayaToolsDir = f.readline()
+        f.close()
+        
+    proj = cmds.optionMenu("artProjOptionMenu", q = True, value = True)
+    character = cmds.textScrollList("artProjCharacterList", q = True, si = True)[0]
+    
+    cmds.file(mayaToolsDir + "/General/ART/Projects/" + proj + "/ExportFiles/" + character + "_Export.mb", open = True, force = True)
+    cmds.deleteUI("artEditCharacterUI")
+    launchSkeletonBuilder()
+    
+    
+#############################################################################################
+#############################################################################################
+#############################################################################################
+def editSelectedCharacterRig(*args):
+    toolsPath = cmds.internalVar(usd = True) + "mayaTools.txt"
+    if os.path.exists(toolsPath):
+        
+        f = open(toolsPath, 'r')
+        mayaToolsDir = f.readline()
+        f.close()
+        
+    proj = cmds.optionMenu("artProjOptionMenu", q = True, value = True)
+    character = cmds.textScrollList("artProjCharacterList", q = True, si = True)[0]
+    
+    cmds.file(mayaToolsDir + "/General/ART/Projects/" + proj + "/AnimRigs/" + character + ".mb", open = True, force = True)
+    cmds.deleteUI("artEditCharacterUI")
+    launchSkeletonBuilder()
+    
+
+    
+#############################################################################################
+#############################################################################################
+#############################################################################################    
+def changeMayaToolsLoc(*args):
+    path = cmds.internalVar(usd = True) + "mayaTools.txt"
+    if os.path.exists(path):
+        os.remove(path)
+        cmds.confirmDialog(title = "Change Location", message = "Once you have chosen your new tools location, it is recommended that you restart Maya.", button = "OK")
+        cmds.file(new = True, force = True)
+        
+#############################################################################################
+#############################################################################################
+#############################################################################################
+def launchSkeletonBuilder(*args):
+    
+    import ART_skeletonBuilder_UI
+    reload(ART_skeletonBuilder_UI)
+    UI = ART_skeletonBuilder_UI.SkeletonBuilder_UI()
+    
+    
+#############################################################################################
+#############################################################################################
+#############################################################################################
+def launchAddCharacter(*args):
+    
+    import ART_addCharacter_UI
+    reload(ART_addCharacter_UI)
+    UI = ART_addCharacter_UI.AddCharacter_UI()
+    
+    
+#############################################################################################
+#############################################################################################
+#############################################################################################
+def launchAnimUI(*args):
+    
+    import ART_animationUI
+    reload(ART_animationUI)
+    UI = ART_animationUI.AnimationUI()
+
+#############################################################################################
+#############################################################################################
+#############################################################################################
+def launchEpic(*args):
+	
+    cmds.launch(web="http://www.epicgames.com")
+    
+#############################################################################################
+#############################################################################################
+#############################################################################################
+def launchUnreal(*args):
+	
+    cmds.launch(web="http://www.unrealengine.com")
+    
+#############################################################################################
+#############################################################################################
+#############################################################################################
+def launchAnimHelp(*args):
+	
+    toolsPath = cmds.internalVar(usd = True) + "mayaTools.txt"
+    if os.path.exists(toolsPath):
+        
+        f = open(toolsPath, 'r')
+        mayaToolsDir = f.readline()
+        f.close()
+	
+	
+    if os.path.exists(mayaToolsDir + "/General/ART/Help/ART_AnimHelp.pdf"):
+	cmds.launch(pdfFile = mayaToolsDir + "/General/ART/Help/ART_AnimHelp.pdf")
+
+#############################################################################################
+#############################################################################################
+#############################################################################################
+def launchRigHelp(self, *args):
+	
+    cmds.launch(web = "https://docs.unrealengine.com/latest/INT/Engine/Content/Tools/MayaRiggingTool/index.html")
+	    
+
+#############################################################################################
+#############################################################################################
+#############################################################################################
+def launchLearningVideos(self, *args):
+	
+    import ART_Help
+    reload(ART_Help)
+    ART_Help.ART_LearningVideos()
+	
+	
+#############################################################################################
+#############################################################################################
+#############################################################################################    
+def setupScene(*args):
+    
+    cmds.currentUnit(time = 'ntsc')
+    cmds.playbackOptions(min = 0, max = 100, animationStartTime = 0, animationEndTime = 100)
+    cmds.currentTime(0)
+    
+    
+    #check for skeleton builder or animation UIs
+    if cmds.dockControl("skeletonBuilder_dock", exists = True):
+	print "Custom Maya Menu: SetupScene"
+	channelBox = cmds.formLayout("SkelBuilder_channelBoxFormLayout", q = True, childArray = True)
+	if channelBox != None:
+	    channelBox = channelBox[0]
+	
+	    #reparent the channelBox Layout back to maya's window
+	    cmds.control(channelBox, e = True, p = "MainChannelsLayersLayout")
+	    channelBoxLayout = mel.eval('$temp1=$gChannelsLayersForm')
+	    channelBoxForm = mel.eval('$temp1 = $gChannelButtonForm')
+	
+	    #edit the channel box pane's attachment to the formLayout
+	    cmds.formLayout(channelBoxLayout, edit = True, af = [(channelBox, "left", 0),(channelBox, "right", 0), (channelBox, "bottom", 0)], attachControl = (channelBox, "top", 0, channelBoxForm))
+	    
+	    
+	#print "deleting dock and window and shit"
+	cmds.deleteUI("skeletonBuilder_dock")
+	if cmds.window("SkelBuilder_window", exists = True):
+	    cmds.deleteUI("SkelBuilder_window")	
+
+
+
+
+	
+    if cmds.dockControl("artAnimUIDock", exists = True):
+	
+	channelBox = cmds.formLayout("ART_cbFormLayout", q = True, childArray = True)
+	if channelBox != None:
+	    channelBox = channelBox[0]
+	
+	    #reparent the channelBox Layout back to maya's window
+	    cmds.control(channelBox, e = True, p = "MainChannelsLayersLayout")
+	    channelBoxLayout = mel.eval('$temp1=$gChannelsLayersForm')
+	    channelBoxForm = mel.eval('$temp1 = $gChannelButtonForm')
+	
+	    #edit the channel box pane's attachment to the formLayout
+	    cmds.formLayout(channelBoxLayout, edit = True, af = [(channelBox, "left", 0),(channelBox, "right", 0), (channelBox, "bottom", 0)], attachControl = (channelBox, "top", 0, channelBoxForm))
+	    
+	    
+	#print "deleting dock and window and shit"
+	cmds.deleteUI("artAnimUIDock")
+	if cmds.window("artAnimUI", exists = True):
+	    cmds.deleteUI("artAnimUI")
+    
+    
+#############################################################################################
+#############################################################################################
+#############################################################################################
+def autoOpenAnimUI():
+    if cmds.objExists("*:master_anim_space_switcher_follow"):
+	launchAnimUI()
+	
+	
+#############################################################################################
+#############################################################################################
+#############################################################################################
+def p4GetLatest(*args):
+    import perforceUtils
+    reload(perforceUtils)
+    perforceUtils.p4_getLatestRevision(None)
+    
+    
+#############################################################################################
+#############################################################################################
+#############################################################################################    
+def p4CheckOut(*args):
+    import perforceUtils
+    reload(perforceUtils)
+    perforceUtils.p4_checkOutCurrentFile(None)
+    
+#############################################################################################
+#############################################################################################
+#############################################################################################
+def p4GetHistory(*args):
+    import perforceUtils
+    reload(perforceUtils)
+    perforceUtils.p4_getRevisionHistory()
+    
+    
+#############################################################################################
+#############################################################################################
+#############################################################################################
+def p4Submit(*args):
+    import perforceUtils
+    reload(perforceUtils)
+    perforceUtils.p4_submitCurrentFile(None, None)
+    
+    
+#############################################################################################
+#############################################################################################
+#############################################################################################
+def p4AddSubmit(*args):
+    import perforceUtils
+    reload(perforceUtils)
+    perforceUtils.p4_addAndSubmitCurrentFile(None, None)
+    
+    
+# # # # # # # # # # # # # # # # # # # # # # # # # # # # # # # # # # # # # # # # # # # # # #
+# # # # # # # # # # # # # # # # # # # # # # # # # # # # # # # # # # # # # # # # # # # # # #
+# # # # # # # # # # # # # # # # # # # # # # # # # # # # # # # # # # # # # # # # # # # # # #
+# # # # # # # # # # # # # # # # # # # # # # # # # # # # # # # # # # # # # # # # # # # # # #
+# # # # # # # # # # # # # # # # # # # # # # # # # # # # # # # # # # # # # # # # # # # # # #
+
+#LAUNCH SCRIPT JOBS
+scriptJobNum2 = cmds.scriptJob(event = ["PostSceneRead", autoOpenAnimUI])
+scriptJobNum = cmds.scriptJob(event = ["NewSceneOpened", setupScene])
+p4ScriptJob = cmds.scriptJob(event = ["NewSceneOpened", p4ProjectMenu])
+
+    
+    
+    
+    
+    
diff --git a/Engine/Extras/Maya_AnimationRiggingTools/MayaTools/General/Scripts/mayaSetup.py b/Engine/Extras/Maya_AnimationRiggingTools/MayaTools/General/Scripts/mayaSetup.py
new file mode 100644
--- /dev/null
+++ b/Engine/Extras/Maya_AnimationRiggingTools/MayaTools/General/Scripts/mayaSetup.py
@@ -0,0 +1,37 @@
+import maya.cmds as cmds
+import os
+import sys
+from functools import partial
+import maya.mel as mel
+
+
+    
+def setupTools():
+    
+    path = cmds.internalVar(usd = True) + "mayaTools.txt"
+    
+
+    f = open(path, 'r')
+    
+    mayaToolsDir = f.readline()
+    
+    path = mayaToolsDir + "/General/Scripts"
+    pluginPath = mayaToolsDir + "/General/Plugins"
+    
+    #look in sys.path to see if path is in sys.path. if not, add it
+    if not path in sys.path:
+        sys.path.append(path)
+        
+    #make sure MAYA_PLUG_IN_PATH has our plugin path
+    pluginPaths = os.environ["MAYA_PLUG_IN_PATH"]
+    pluginPaths = pluginPaths + ";" + pluginPath
+    os.environ["MAYA_PLUG_IN_PATH"] = pluginPaths
+        
+
+    
+    
+    #setup menu item in main window
+    import customMayaMenu as cmm
+    cmm.customMayaMenu()
+
+setupTools()
diff --git a/Engine/Extras/Maya_AnimationRiggingTools/MayaTools/General/Scripts/perforceUtils.py b/Engine/Extras/Maya_AnimationRiggingTools/MayaTools/General/Scripts/perforceUtils.py
new file mode 100644
--- /dev/null
+++ b/Engine/Extras/Maya_AnimationRiggingTools/MayaTools/General/Scripts/perforceUtils.py
@@ -0,0 +1,1015 @@
+import maya.cmds as cmds
+from P4 import P4,P4Exception
+import os, cPickle
+from functools import partial
+
+# # # # # # # # # # # # # # # # # # # # # # # # # # # # # # # # # # # # # 
+def p4_getLatestRevision(fileName, *args):
+    fileArg = fileName
+    #try to connect
+    p4 = P4()
+    
+    try:
+        p4.connect()
+        
+    except:
+        cmds.confirmDialog(title = "Perforce", icon = "critical", message = "Unable to connect to perforce server.")
+        return
+    
+    #find currently opened file name
+    if fileName == None:
+	fileName = cmds.file(q = True, sceneName = True)
+	
+    syncFiles = []
+    
+    
+    try:
+	#client info
+	spec = p4.run( "client", "-o" )[0]
+	client = spec.get("Client")
+	owner = spec.get("Owner")	
+	p4.user = owner
+	p4.client = client
+        
+    except:
+        cmds.confirmDialog(title = "Perforce", icon = "critical", message = "Unable to obtain client spec information.")
+    
+    
+    
+    #try to get the current file revision on local, and compare to depot
+    try:
+        #Find out the revision of the local version of the file
+        myFile = p4.run_have(fileName)[0]
+        
+        #This will find the revision number of your local file.
+        localRevision = int(myFile['haveRev'])
+        
+        #find out the revision number of the depot version of the file
+        depotVersion = p4.run_files(myFile['depotFile'])[0]
+        
+        #find the depot file path
+        depotFile = depotVersion['depotFile']
+        
+        #find the depot revision number of the file
+        depotRevision = int(depotVersion['rev'])        
+    
+        #check for latest
+        if localRevision != depotRevision:
+            syncFiles.append(depotFile)
+                
+                
+        
+        #Check for scene references in the file
+        allRefs = []
+        references = cmds.file(q = True, reference = True)
+        for reference in references:
+            nestedRef = cmds.file(reference, q = True, reference = True)
+            allRefs.append(reference)
+            allRefs.append(nestedRef)
+            
+        #loop through all found references and check for latest
+        for ref in allRefs:
+            #get revision of local file
+            myFile = p4.run_have(ref)[0]
+            
+            #get revision number
+            localRefRevision = int(myFile['haveRev']) 
+            
+            #grab depot file info
+            depotRefVersion = p4.run_files(myFile['depotFile'])[0]
+            
+            #depot file path
+            depotFile = depotRefVersion['depotFile']
+            
+            #get depot's revision #
+            depotRefRevision = int(depotRefVersion['rev'])
+            
+            
+            #compare
+            if localRefRevision != depotRefRevision:
+                syncFiles.append(depotFile)
+            
+    
+        
+        
+        #if there are files to sync, do it now
+        if len(syncFiles) > 0:
+            message = "The following files are not at latest revision:\n\n"
+            for file in syncFiles:
+                message += file + "\n"
+                
+            result = cmds.confirmDialog(title = "Perforce", icon = "warning", message =  message, button = ["Sync", "Cancel"])
+            
+            if result == "Sync":
+                #sync files
+                for f in syncFiles:
+                    p4.run_sync(f)
+                    
+                #ask if user would like to reopen
+		if fileArg == None:
+		    result = cmds.confirmDialog(title = "Perforce", icon = "question", message =  "Sync Complete. Reopen file to get changes?", button = ["Yes", "Cancel"])
+		    
+		    if result == "Yes":
+			cmds.file(fileName, open = True, force = True)
+		    
+	else:
+	    cmds.confirmDialog(title = "Perforce", icon = "information", message = "This file is already at head revision.", button = "Close")
+                    
+        #disconnect from server
+        p4.disconnect()
+                
+        
+        
+    #Handle any p4 errors that come back from trying to run the above code
+    except P4Exception:
+        errorString = "The following errors were encountered:\n\n"
+        
+        for e in p4.errors:  
+            errorString += e + "\n"
+        
+        cmds.confirmDialog(title = "Perforce", icon = "critical", message = errorString)
+	p4.disconnect()
+        return 
+    
+    
+# # # # # # # # # # # # # # # # # # # # # # # # # # # # # # # # # # # # # 
+def p4_checkOutCurrentFile(fileName, *args):
+    fileArg = fileName
+    #try to connect
+    p4 = P4()
+    
+    try:
+        p4.connect()
+        
+    except:
+        cmds.confirmDialog(title = "Perforce", icon = "critical", message = "Unable to connect to perforce server.")
+        return False
+    
+    #find currently opened file name
+    if fileName == None:
+	fileName = cmds.file(q = True, sceneName = True)
+	
+    reopen = False
+    syncFiles = []
+    
+    
+    try:
+	#client info
+	spec = p4.run( "client", "-o" )[0]
+	client = spec.get("Client")
+	owner = spec.get("Owner")	
+	p4.user = owner
+	p4.client = client
+        
+    except:
+        cmds.confirmDialog(title = "Perforce", icon = "critical", message = "Unable to obtain client spec information.")
+        
+    
+    try:
+        #check to see if file is at head revision
+        myFile = p4.run_have(fileName)[0]
+        
+        #This will find the revision number of your local file.
+        localRevision = int(myFile['haveRev'])
+        
+        #find out the revision number of the depot version of the file
+        depotVersion = p4.run_files(myFile['depotFile'])[0]
+        
+        #find the depot file path
+        depotFile = depotVersion['depotFile']
+        
+        #find the depot revision number of the file
+        depotRevision = int(depotVersion['rev'])        
+    
+        #check for latest
+        if localRevision != depotRevision:
+            result = cmds.confirmDialog(title = "Perforce", icon = "warning", message = "This file is not at head revision. Please get latest and try again.", button = ["Get Latest", "Cancel"])
+            
+            if result == "Get Latest":
+                p4_getLatestRevision(fileArg)
+                p4.disconnect()
+		
+            else:
+                return False
+            
+        else:
+	    try:
+		#check to see if file is checked out
+		opened = p4.run_opened(depotFile)
+		
+		if len(opened) > 0:
+		    user = opened[0]['user']
+		    cmds.confirmDialog(title = "Perforce", icon = "warning", message = "This file is already checked out by: " + user, button = "Close")
+		    p4.disconnect()
+	    
+		else:
+		    #check out the file
+		    p4.run_edit(depotFile)
+		    cmds.confirmDialog(title = "Perfoce", icon = "information", message = "This file is now checked out.", button = "Close")
+		    p4.disconnect()
+		    #tools path
+		    toolsPath = cmds.internalVar(usd = True) + "mayaTools.txt"
+		    if os.path.exists(toolsPath):
+			
+			f = open(toolsPath, 'r')
+			mayaToolsDir = f.readline()
+			f.close()
+			
+		    return True
+		    
+            #Handle any p4 errors that come back from trying to run the above code
+            except P4Exception:
+                errorString = "The following errors were encountered:\n\n"
+                
+                for e in p4.errors:  
+                    errorString += e + "\n"
+                
+                cmds.confirmDialog(title = "Perforce", icon = "critical", message = errorString)
+                return False
+            
+
+
+    #Handle any p4 errors that come back from trying to run the above code
+    except P4Exception:
+        errorString = "The following errors were encountered:\n\n"
+        
+        for e in p4.errors:  
+            errorString += e + "\n"
+        
+        cmds.confirmDialog(title = "Perforce", icon = "critical", message = errorString)
+	p4.disconnect()
+        return False
+    
+    
+    
+# # # # # # # # # # # # # # # # # # # # # # # # # # # # # # # # # # # # # 
+def p4_getRevisionHistory(*args):
+    
+    #try to connect
+	p4 = P4()
+	
+	try:
+	    p4.connect()
+	    
+	except:
+	    cmds.confirmDialog(title = "Perforce", icon = "critical", message = "Unable to connect to perforce server.")
+	    return
+	
+	#find currently opened file name
+	clientFile = cmds.file(q = True, sceneName = True)
+	reopen = False
+	syncFiles = []
+	
+	
+	try:
+	    #client info
+	    spec = p4.run( "client", "-o" )[0]
+	    client = spec.get("Client")
+	    owner = spec.get("Owner")	
+	    p4.user = owner
+	    p4.client = client    
+
+	except:
+	    cmds.confirmDialog(title = "Perforce", icon = "critical", message = "Unable to obtain client spec information.")
+	    
+	    
+	#get revision history of current file
+	try:
+	    #check to see if file is at head revision
+	    myFile = p4.run_have(clientFile)[0]
+	    depotVersion = p4.run_files(myFile['depotFile'])[0]
+	    depotFile = depotVersion['depotFile']	    
+	    history = p4.run_changes(depotFile)
+	    
+	    info = ""
+	    
+	    for h in history:
+		user = h.get("user")
+		change = h.get("change")
+		desc = h.get("desc")
+		
+		if desc.find("\n") == -1:
+		    desc = desc + "...\n"
+		    
+		else:
+		    desc = desc.partition("\n")[0] + "...\n"
+		    
+		info += change + " by " + user + ": " + desc
+		
+	    
+	    #print report into a confirm dialog
+	    cmds.confirmDialog(title = "History", icon = "information", ma = "left",  message = info, button = "Close")
+	    p4.disconnect()
+		
+	    
+	    
+	
+	except P4Exception:
+	    errorString = "The following errors were encountered:\n\n"
+	    
+	    for e in p4.errors:  
+		errorString += e + "\n"
+	    
+	    cmds.confirmDialog(title = "Perforce", icon = "critical", message = errorString)
+	    p4.disconnect()
+	    return
+	
+	
+# # # # # # # # # # # # # # # # # # # # # # # # # # # # # # # # # # # # # 
+def p4_submitCurrentFile(fileName, desc, *args):
+    
+        fileArg = fileName
+        #try to connect
+	p4 = P4()
+	
+	
+	try:
+	    p4.connect()
+	    
+	except:
+	    cmds.confirmDialog(title = "Perforce", icon = "critical", message = "Unable to connect to perforce server.")
+	    return
+	
+	#find currently opened file name
+	if fileName == None:
+	    fileName = cmds.file(q = True, sceneName = True)
+	    
+	reopen = False
+	syncFiles = []
+	
+	
+	try:
+	    #client info
+	    spec = p4.run( "client", "-o" )[0]
+	    client = spec.get("Client")
+	    owner = spec.get("Owner")	
+	    p4.user = owner
+	    p4.client = client    
+	    
+	except:
+	    cmds.confirmDialog(title = "Perforce", icon = "critical", message = "Unable to obtain client spec information.")
+	    
+	    
+	
+	#SUBMIT
+	try:
+	    	    
+	    if desc == None:
+		result = cmds.promptDialog(title = "Perforce", message = "Please Enter a Description..", button = ["Accept", "Cancel"], defaultButton = "Accept", dismissString = "Cancel", cancelButton = "Cancel")
+		
+	    else:
+		result = "Accept"
+		
+	    #process
+	    if result == "Accept":
+		
+		#description = "test"
+		myFile = p4.run_have(fileName)[0]
+		depotVersion = p4.run_files(myFile['depotFile'])[0]
+		depotFile = depotVersion['depotFile']
+		
+		#check to see if file is checked out
+		opened = p4.run_opened(depotFile)
+		
+		if len(opened) > 0:
+		    opendBy = opened[0]['user']
+
+		    
+		    if opendBy.lower() != owner.lower():
+			cmds.confirmDialog(title = "Perforce", icon = "warning", message = "This file is already checked out by: " + opendBy, button = "Close")
+			p4.disconnect()
+			return
+		
+		    else:
+		    
+			#fetch the description
+			if desc == None:
+			    desc = cmds.promptDialog(q = True, text = True)	    
+			
+			#save the file locally (so depot and HD are in sync)
+			openedFile = cmds.file(q = True, sceneName = True)
+			saveFileName = openedFile.rpartition("/")[2]
+			if fileArg == None:
+			    cmds.file(f = True, save = True, options = "v = 0", type = "mayaBinary")
+			
+			
+			#grab the name of the file
+			fileNameWithExt = openedFile.rpartition("/")[2]
+			fileName = fileNameWithExt.rpartition(".")[0]
+			
+			description = (desc + "\n Affected Files: " + openedFile)
+			
+			#create new changelist
+			newChange = p4.fetch_change()
+			newChange._description = description
+			
+			#make sure we don't add existing default changelist files.
+			newChange._files = []
+			
+			#determine the new number so we can refetch it. 
+			newChangeNum = int(p4.save_change(newChange)[0].split()[1])
+	    
+	
+			#change changelist number
+			p4.run_reopen('-c', newChangeNum, depotFile)
+			
+			#submit the changelist
+			p4.run_submit('-c', newChangeNum)
+			
+			#tell the user submit was successful
+			if fileArg == None:
+			    cmds.confirmDialog(title = "Perforce", icon = "information", message = "Submit Operation was successful!", button = "Close")
+			    
+			else:
+			    return True
+			
+		else:
+		    #if the file is not checked out by the user, let them know
+		    result = cmds.confirmDialog(title = "Perforce", icon = "warning", message = "File is not checked out. Unable to continue submit operation on this file:\n\n" + fileName)
+		
+
+		
+	except P4Exception:
+	    if fileArg == None:
+		errorString = "The following errors were encountered:\n\n"
+		
+		for e in p4.errors:  
+		    errorString += e + "\n"
+		
+		cmds.confirmDialog(title = "Perforce", icon = "critical", message = errorString)
+		
+	    p4.disconnect()
+	    return False
+		
+# # # # # # # # # # # # # # # # # # # # # # # # # # # # # # # # # # # # # 			
+def p4_addAndSubmitCurrentFile(fileName, description, *args):
+    fileArg = fileName
+    #try to connect
+    p4 = P4()
+    
+    
+    try:
+	p4.connect()
+	
+    except:
+	cmds.confirmDialog(title = "Perforce", icon = "critical", message = "Unable to connect to perforce server.")
+	return
+    
+
+    #find currently opened file name
+    if fileName == None:
+	fileName = cmds.file(q = True, sceneName = True)
+    reopen = False
+    syncFiles = []
+    
+    
+    
+    try:
+	#client info
+	spec = p4.run( "client", "-o" )[0]
+	client = spec.get("Client")
+	owner = spec.get("Owner")	
+	p4.user = owner
+	p4.client = client
+	
+    except:
+	cmds.confirmDialog(title = "Perforce", icon = "critical", message = "Unable to obtain client spec information.")
+	
+    
+    #find currently opened file name
+    proceed = False
+    if fileArg == None:
+	fileName = cmds.file(q = True, sceneName = True)
+	
+	if fileName == "":
+	    cmds.confirmDialog(title = "Perforce", icon = "warning", message = "Cannot Add file to perforce as file has no name.", button = "Close")
+	    p4.disconnect()
+	    return
+	
+	else:
+	    proceed = True
+	    
+    else:
+	proceed = True
+	
+    
+    #if the file has a filename, 
+    if proceed:
+	try:
+	    
+	    
+	    clientRoot = p4.fetch_client(p4.client)._Root
+	    
+	    #check to make sure client root is in the client file path
+	    if os.path.normpath(fileName).find(os.path.normpath(clientRoot))  == 0:
+		#if it was, then get a description for the changelist
+		if description == None:
+		    result = cmds.promptDialog(title = "Perforce", message = "Please Enter a Description..", button = ["Accept", "Cancel"], defaultButton = "Accept", dismissString = "Cancel", cancelButton = "Cancel")
+		
+		else:
+		    result = "Accept"
+		    
+		
+		if result == "Accept":
+		    
+		    #get changelist description
+		    if description == None:
+			description = cmds.promptDialog(q = True, text = True)
+			
+		    
+		    #create changelist
+		    newChange = p4.fetch_change()
+		    newChange._description = description
+		    
+		    #make sure we don't add existing default changelist files.
+		    newChange._files = []
+			    
+		    #determine the new number so we can refetch it.
+		    newChangeNum = int(p4.save_change(newChange)[0].split()[1])		    
+		    
+		    
+		    #description = "test"
+		    p4.run_add('-c', newChangeNum, fileName)
+	
+		    #submit the changelist
+		    p4.run_submit('-c', newChangeNum)
+		    
+		    #tell user operation was successful
+		    if fileArg == None:
+			result = cmds.confirmDialog(title = "Perforce", icon = "information", message = "File has been successfully added to perforce and submitted!", button = ["Close", "Check Out File"])
+			
+			if result == "Close":
+			    p4.disconnect()
+			    return
+			
+			if result == "Check Out File":
+			    p4_checkOutCurrentFile(fileName)
+			    
+		    #return operation succuessful
+		    return True
+			
+			
+
+
+		else:
+		    p4.disconnect()
+		    return
+		
+		
+	    else:
+		cmds.confirmDialog(title = "Perforce", icon = "warning", message = "Cannot proceed. File is not under client's root, " + clientRoot, button = "Close")
+		p4.disconnect()
+		return False
+	    
+	    
+	    
+	    
+	except P4Exception:
+	    errorString = "The following errors were encountered:\n\n"
+	    
+	    for e in p4.errors:  
+		errorString += e + "\n"
+	    
+	    cmds.confirmDialog(title = "Perforce", icon = "critical", message = errorString)
+	    p4.disconnect()
+	    return False
+	
+	
+# # # # # # # # # # # # # # # # # # # # # # # # # # # # # # # # # # # # # 
+def p4_checkForUpdates(*args):
+    #try to connect
+    p4 = P4()
+    
+    #get maya tools path
+    toolsPath = cmds.internalVar(usd = True) + "mayaTools.txt"
+    if os.path.exists(toolsPath):
+	
+	f = open(toolsPath, 'r')
+	mayaToolsDir = f.readline()
+	f.close()
+	
+    #connect to p4
+    try:
+	p4.connect()
+	
+    except:
+	cmds.confirmDialog(title = "Perforce", icon = "critical", message = "Unable to connect to perforce server.")
+	return
+    
+    #find currently opened file name
+    clientFile = cmds.file(q = True, sceneName = True)
+    reopen = False
+    syncFiles = []
+    
+    
+    try:
+	#client info
+	spec = p4.run( "client", "-o" )[0]
+	client = spec.get("Client")
+	owner = spec.get("Owner")	
+	p4.user = owner
+	p4.client = client
+
+	
+    except:
+	cmds.confirmDialog(title = "Perforce", icon = "critical", message = "Unable to obtain client spec information.")
+	
+	
+    #this will check the maya tools directory in p4 for any updates
+    try:
+	syncFiles = []
+	clientRoot = p4.fetch_client(p4.client)._Root
+	depotDirectories = []
+	
+	#get current project
+	if os.path.exists(mayaToolsDir + "/General/Scripts/projectSettings.txt"):
+	    #read to find current project
+	    f = open(mayaToolsDir + "/General/Scripts/projectSettings.txt", 'r')
+	    settings = cPickle.load(f)
+	    f.close()
+	    
+	    #write out new settings
+	    project = settings.get("CurrentProject")
+	    
+	    if os.path.exists(mayaToolsDir + "/General/Scripts/" + project + "_Project_Settings.txt"):
+		#read the depot paths to sync from the project settings
+		f = open(mayaToolsDir + "/General/Scripts/" + project + "_Project_Settings.txt", 'r')
+		settings = cPickle.load(f)
+		f.close() 
+		depotDirectories = settings.get("depotPaths")		
+		print depotDirectories
+
+	
+	#look at each directory inside MayaTools
+	for dir in depotDirectories:
+	    depotFiles = p4.run_files(dir + "...")
+	    
+	    for each in depotFiles:
+		
+		#try to compare depot to local. It is possible that there are local files not in depot, and vise versa
+		try:
+		    fileInfo = p4.run_files(each['depotFile'])[0]
+		    depotFilePath = fileInfo['depotFile']
+		    fileName = depotFilePath.rpartition("/")[2]
+		    
+		    #compare local files
+		    localFile = p4.run_have(depotFilePath)[0]
+		    localRevision = int(localFile['haveRev'])
+		    depotRevision = int(fileInfo['rev'])
+		    
+		    if localRevision < depotRevision:
+			syncFiles.append(depotFilePath)
+			
+		    
+		except:
+		    
+		    try:
+			#check to see if it errors out because we don't have a local version of the file
+			fileInfo = p4.run_files(each['depotFile'])[0]
+			depotFilePath = fileInfo['depotFile']
+			fileName = depotFilePath.rpartition("/")[2]			
+			localFile = p4.run_have(depotFilePath)[0]
+			
+		    except:
+			action = each.get("action")
+			if action != "delete":
+			    syncFiles.append(depotFilePath)
+			    
+			pass
+		    
+		    
+		
+	
+	#check size of syncFiles and ask user if they want to sync
+	if len(syncFiles) > 0:
+	    result = cmds.confirmDialog(title = "MayaTools", icon = "warning", message = "There are new updates available to the depot directories specified by your project settings.", button = ["Update", "Not Now"])
+	    
+	    if result == "Update":
+		for file in syncFiles:
+		    p4.run_sync(file)
+		    
+		
+		cmds.confirmDialog(title = "MayaTools", icon = "information", message = "Tools are now up to date!", button = "Close")
+		p4.disconnect()
+		
+	    else:
+		p4.disconnect()
+		return
+		    
+
+    #handle any errors
+    except P4Exception:
+	errorString = "The following errors were encountered:\n\n"
+	
+	for e in p4.errors:  
+	    errorString += e + "\n"
+	
+	cmds.confirmDialog(title = "Perforce", icon = "critical", message = errorString)
+	p4.disconnect()
+	return
+    
+    
+
+   
+# # # # # # # # # # # # # # # # # # # # # # # # # # # # # # # # # # # # # 
+def createNewProject(*args):
+    
+    if cmds.window("createNewARTProject_Window", exists = True):
+	cmds.deleteUI("createNewARTProject_Window")
+	
+    #create window
+    window = cmds.window("createNewARTProject_Window", w = 400, h = 600, mnb = False, mxb = False, title = "Create New Project")
+    
+    #frameLayouts for settings: Perforce/Auto-Sync
+    mainLayout = cmds.columnLayout(w = 400, h = 600)
+    
+    #project name field
+    projectName = cmds.textFieldGrp("newARTProjectNameField", label = "Project Name: ", w = 400, h = 30, parent = mainLayout, cal = [1, "left"])
+    
+    
+    scrollLayout = cmds.scrollLayout(w = 400, h = 520, parent = mainLayout)
+    columnLayout = cmds.columnLayout(w = 380, parent = scrollLayout)
+    
+
+    
+    #perforce/auto-sync layout
+    p4Frame = cmds.frameLayout(parent = columnLayout, w = 370, cll = True, label='Perforce/Auto-Sync', borderStyle='in')
+    p4Layout = cmds.columnLayout(w = 360, parent = p4Frame, co = ["both", 5], rs = 5)
+    
+    #create the scrollField with the information
+    cmds.scrollField(parent = p4Layout, w = 350, h = 100, editable = False, wordWrap = True, text = "Add depot paths you would like the tools to check for updates on. If updates are found, you will be notified, and asked if you would like to sync. Valid depot paths look like:\n\n//depot/usr/jeremy_ernst/MayaTools")
+    
+    #crete the add button
+    cmds.button(w = 350, label = "Add Perforce Depot Path", parent = p4Layout, c = partial(addPerforceDepotPath, p4Layout))
+    
+    
+    
+    
+    #save settings button
+    cmds.button(parent = mainLayout, w = 400, h = 50, label = "Save Settings and Close", c = partial(saveProjectSettings, p4Layout, False))
+    
+    #show window
+    cmds.showWindow(window)
+    
+    
+    
+    
+# # # # # # # # # # # # # # # # # # # # # # # # # # # # # # # # # # # # # 
+def editProject(project, *args):
+    
+    if cmds.window("createNewARTProject_Window", exists = True):
+	cmds.deleteUI("createNewARTProject_Window")
+	
+    #create window
+    window = cmds.window("createNewARTProject_Window", w = 400, h = 600, mnb = False, mxb = False, title = "Edit Project")
+    
+    #frameLayouts for settings: Perforce/Auto-Sync
+    mainLayout = cmds.columnLayout(w = 400, h = 600)
+    
+    #project name field
+    projectName = cmds.textFieldGrp("newARTProjectNameField", label = "Project Name: ", text = project, w = 400, h = 30, parent = mainLayout, cal = [1, "left"])
+    
+    
+    scrollLayout = cmds.scrollLayout(w = 400, h = 520, parent = mainLayout)
+    columnLayout = cmds.columnLayout(w = 380, parent = scrollLayout)
+    
+
+    
+    #perforce/auto-sync layout
+    p4Frame = cmds.frameLayout(parent = columnLayout, w = 370, cll = True, label='Perforce/Auto-Sync', borderStyle='in')
+    p4Layout = cmds.columnLayout(w = 360, parent = p4Frame, co = ["both", 5], rs = 5)
+    
+    #create the scrollField with the information
+    cmds.scrollField(parent = p4Layout, w = 350, h = 100, editable = False, wordWrap = True, text = "Add depot paths you would like the tools to check for updates on. If updates are found, you will be notified, and asked if you would like to sync. Valid depot paths look like:\n\n//depot/usr/jeremy_ernst/MayaTools")
+    
+    #crete the add button
+    cmds.button(w = 350, label = "Add Perforce Depot Path", parent = p4Layout, c = partial(addPerforceDepotPath, p4Layout))
+    
+    #get maya tools path
+    toolsPath = cmds.internalVar(usd = True) + "mayaTools.txt"
+    if os.path.exists(toolsPath):
+	
+	f = open(toolsPath, 'r')
+	mayaToolsDir = f.readline()
+	f.close()
+    
+    
+    #open the project settings and auto-fill in the info
+    if os.path.exists(mayaToolsDir + "/General/Scripts/" + project + "_Project_Settings.txt"):
+	f = open(mayaToolsDir + "/General/Scripts/" + project + "_Project_Settings.txt", 'r')
+	settings = cPickle.load(f)
+	f.close() 
+	
+	paths = settings.get("depotPaths")
+	
+    
+	if len(paths) > 0:
+	    for path in paths:
+		#add the path
+		field = addPerforceDepotPath(p4Layout)
+		cmds.textField(field, edit = True, text = path)
+	
+	
+	
+	
+	   
+    
+    #save settings button
+    cmds.button(parent = mainLayout, w = 400, h = 50, label = "Save Settings and Close", c = partial(saveProjectSettings, p4Layout, True))
+    
+    #show window
+    cmds.showWindow(window)
+    
+# # # # # # # # # # # # # # # # # # # # # # # # # # # # # # # # # # # # # 
+def addPerforceDepotPath(layout, *args):
+    
+    field = cmds.textField(docTag = "P4DepotPath", w = 350, parent = layout)
+    
+    #add a RMB menu to remove the field
+    menu = cmds.popupMenu(parent = field, b = 3)
+    cmds.menuItem(parent = menu, label = "Remove Path", c = partial(removePerforceDepotPath, field))
+    
+    return field
+   
+
+
+
+
+
+# # # # # # # # # # # # # # # # # # # # # # # # # # # # # # # # # # # # # 
+def removePerforceDepotPath(field, *args):
+    cmds.textField(field, edit = True, visible = False, h = 1)
+    #cmds.deleteUI(field) This crashes maya instantly. Come ON AUTODESK
+    
+
+
+
+
+# # # # # # # # # # # # # # # # # # # # # # # # # # # # # # # # # # # # # 
+def saveProjectSettings(perforceLayout, edit, *args):
+    
+    #find p4 depot path textfields
+    children = cmds.columnLayout(perforceLayout, q = True, childArray = True)
+    textFields = []
+    for child in children:
+	if child.find("textField") == 0:
+	    data = cmds.textField(child, q = True, docTag = True)
+	    if data == "P4DepotPath":
+		textFields.append(child)
+	
+    
+    #make sure paths are valid
+    savePaths = []
+    for field in textFields:
+	path = cmds.textField(field, q = True, text = True)
+	if path != "":
+	    
+	    try:
+		p4 = P4()
+		p4.connect()
+    
+		    
+	    except:
+		cmds.confirmDialog(title = "Perforce", icon = "critical", message = "Unable to connect to perforce server.")
+		return
+	    
+	    
+	    try:
+		#client info
+		spec = p4.run( "client", "-o" )[0]
+		client = spec.get("Client")
+		owner = spec.get("Owner")	
+		p4.user = owner
+		p4.client = client
+	
+		
+	    except:
+		cmds.confirmDialog(title = "Perforce", icon = "critical", message = "Unable to obtain client spec information.")	
+    
+	    
+	    #now check
+	    try:
+		depotFiles = p4.run_files(path + "...")
+		if len(depotFiles) > 0:
+		    savePaths.append(path)
+	    
+	    #handle any errors
+	    except P4Exception:
+		errorString = "The following errors were encountered:\n\n"
+		
+		for e in p4.errors:  
+		    errorString += e + "\n"
+		
+		cmds.confirmDialog(title = "Perforce", icon = "critical", message = errorString)
+		p4.disconnect()
+		return
+	    
+	else:
+	    #see if the text field is just hidden or if it is actually blank
+	    vis = cmds.textField(field, q = True, visible = True)
+	    if vis == True:
+		cmds.confirmDialog(title = "Error", icon = "critical", message = "Empty string not allowed as a path name. Either remove that field, or enter a correct depot path.")
+		return
+	    
+	    
+    #write out to disk
+    projectName = cmds.textFieldGrp("newARTProjectNameField", q = True, text = True)
+    if projectName == "":
+	cmds.confirmDialog(title = "Error", icon = "critical", message = "Empty string not allowed as a project name.")
+	return
+    
+    #save the new project file under MayaTools/General/Scripts as projName + "_Project_Settings.txt"
+    toolsPath = cmds.internalVar(usd = True) + "mayaTools.txt"
+    if os.path.exists(toolsPath):
+	
+	f = open(toolsPath, 'r')
+	mayaToolsDir = f.readline()
+	f.close()
+	
+    if edit == False:
+	if os.path.exists(mayaToolsDir + "/General/Scripts/" + projectName + "_Project_Settings.txt"):
+	    cmds.confirmDialog(title = "Error", icon = "critical", message = "Project already exists with that name")
+	    return
+    
+    #save out
+    f = open(mayaToolsDir + "/General/Scripts/" + projectName + "_Project_Settings.txt", 'w')
+
+
+    #create a dictionary with  values
+    settings = {}
+    settings["depotPaths"] = savePaths
+    
+    
+    #write our dictionary to file
+    cPickle.dump(settings, f)
+    f.close()
+    
+    
+    #delete the UI
+    cmds.deleteUI("createNewARTProject_Window")
+    
+    
+    #add the project to the menu
+    create = True
+    items = cmds.lsUI(menuItems = True)
+    for i in items:
+	data = cmds.menuItem(i, q = True, docTag = True)
+	if data == "P4Proj":
+	    label = cmds.menuItem(i, q = True, label = True)
+	    print label
+	    
+	    if label == projectName:
+		create = False
+		
+		
+	    
+    if create:
+	menuItem = cmds.menuItem(label = projectName, parent = "perforceProjectList", cl = "perforceProjectRadioMenuCollection", rb = True, docTag = "P4Proj", c = partial(setCurrentProject, projectName))
+	cmds.menuItem(parent = "perforceProjectList", optionBox = True, c = partial(editProject, projectName))
+    
+    
+    
+    #open up the projectSettings.txt file and add an entry for current project
+    if os.path.exists(mayaToolsDir + "/General/Scripts/projectSettings.txt"):
+	f = open(mayaToolsDir + "/General/Scripts/projectSettings.txt", 'r')
+	oldSettings = cPickle.load(f)
+	useSourceControl = oldSettings.get("UseSourceControl")
+	f.close()
+	
+	#write out new settings
+	settings = {}
+	settings["UseSourceControl"] = useSourceControl
+	settings["CurrentProject"] = projectName
+	
+	f = open(mayaToolsDir + "/General/Scripts/projectSettings.txt", 'w')
+	cPickle.dump(settings, f)
+	f.close()	
+	
+	
+# # # # # # # # # # # # # # # # # # # # # # # # # # # # # # # # # # # # # 
+def setCurrentProject(projectName, *args):
+    #get access to maya tools path
+    toolsPath = cmds.internalVar(usd = True) + "mayaTools.txt"
+    if os.path.exists(toolsPath):
+	
+	f = open(toolsPath, 'r')
+	mayaToolsDir = f.readline()
+	f.close()
+	
+	
+    #re-write settings
+    if os.path.exists(mayaToolsDir + "/General/Scripts/projectSettings.txt"):
+	f = open(mayaToolsDir + "/General/Scripts/projectSettings.txt", 'r')
+	oldSettings = cPickle.load(f)
+	useSourceControl = oldSettings.get("UseSourceControl")
+	f.close()
+	
+	#write out new settings
+	settings = {}
+	settings["UseSourceControl"] = useSourceControl
+	settings["CurrentProject"] = projectName
+	
+	f = open(mayaToolsDir + "/General/Scripts/projectSettings.txt", 'w')
+	cPickle.dump(settings, f)
+	f.close()	    
+    
diff --git a/Engine/Extras/Maya_AnimationRiggingTools/MayaTools/General/Scripts/setupScene.py b/Engine/Extras/Maya_AnimationRiggingTools/MayaTools/General/Scripts/setupScene.py
new file mode 100644
--- /dev/null
+++ b/Engine/Extras/Maya_AnimationRiggingTools/MayaTools/General/Scripts/setupScene.py
@@ -0,0 +1,11 @@
+import maya.cmds as cmds
+
+
+def setupScene():
+    #Setup scene to set timeline to 0-100 with current frame being 0 and the fps to be 30
+    cmds.currentUnit(time = 'ntsc')
+    cmds.playbackOptions(min = 0, max = 100, animationStartTime = 0, animationEndTime = 100)
+    cmds.currentTime(0)
+    
+    
+setupScene()
\ No newline at end of file
diff --git a/Engine/Extras/Maya_AnimationRiggingTools/MayaTools/userSetup.py b/Engine/Extras/Maya_AnimationRiggingTools/MayaTools/userSetup.py
new file mode 100644
--- /dev/null
+++ b/Engine/Extras/Maya_AnimationRiggingTools/MayaTools/userSetup.py
@@ -0,0 +1,131 @@
+import maya.cmds as cmds
+import os
+import sys
+from functools import partial
+import maya.mel as mel
+
+MAYA_SEP = '/'
+
+def correctSysPath():
+    for i, path in enumerate(sys.path):
+        if path[-1] in (MAYA_SEP, os.sep, os.altsep):
+            sys.path[i] = sys.path[i][:-1]
+
+correctSysPath()
+
+def deleteOldAutoSaves():
+    import maya.cmds as cmds
+    import time
+    import os.path, time
+    
+    
+    save = cmds.autoSave(q = True, enable = True)
+    destDir = cmds.autoSave(q = True, destinationFolder = True)
+    deleteDays = 3
+    
+    if save:
+        now = time.time()
+        fileCap = now - 60*60*24*deleteDays
+    
+        files = os.listdir(destDir)
+        
+        for file in files:
+            fileCreation = os.path.getctime(os.path.join(destDir, file))
+            if fileCreation < fileCap:
+                os.remove(os.path.join(destDir, file))
+                
+                
+def mayaToolsInstall_UI():
+    
+    if cmds.window("mayaToolsInstall_UI", exists = True):
+        cmds.deleteUI("mayaToolsInstall_UI")
+        
+    window = cmds.window("mayaToolsInstall_UI", w = 300, h = 100, title = "Maya Tools Install", mnb = False, mxb = False)
+    
+    mainLayout = cmds.columnLayout(w = 300, h = 100)
+    formLayout = cmds.formLayout(w = 300, h = 100)
+    
+    text = cmds.text(label = "ERROR: Could not find Maya Tools directory.\n Please locate folder using the \'Browse\' button.", w = 300)
+    cancelButton = cmds.button(label = "Cancel", w = 140, h = 50, c = mayaToolsInstall_Cancel)
+    browseButton = cmds.button(label = "Browse", w = 140, h = 50, c = mayaToolsInstall_Browse)
+    
+    cmds.formLayout(formLayout, edit = True, af = [(text, 'left', 10), (text, 'top', 10)])
+    cmds.formLayout(formLayout, edit = True, af = [(cancelButton, 'left', 5), (cancelButton, 'top', 50)])
+    cmds.formLayout(formLayout, edit = True, af = [(browseButton, 'right', 5), (browseButton, 'top', 50)])
+    
+    cmds.showWindow(window)
+    cmds.window(window, edit = True, w = 300, h = 100)
+    
+    
+    
+def mayaToolsInstall_Cancel(*args):
+    
+    cmds.deleteUI("mayaToolsInstall_UI")
+    cmds.warning("Maya Tools will not be setup")
+    
+
+
+def mayaToolsInstall_Browse(*args):
+    
+    mayaToolsDir = cmds.fileDialog2(dialogStyle = 2, fileMode = 3)[0]
+
+    #confirm that this is actually the maya tools directory
+    if mayaToolsDir.rpartition("/")[2] != "MayaTools":
+        cmds.warning("Selected directory is not valid. Please locate the MayaTools directory.")
+        
+        
+    else:
+        cmds.deleteUI("mayaToolsInstall_UI")
+        
+        #create file that contains this path
+        path = cmds.internalVar(usd = True) + "mayaTools.txt"
+        
+        f = open(path, 'w')
+        f.write(mayaToolsDir)
+        f.close()
+        
+        
+        #run setup
+        cmds.file(new = True, force = True)
+        
+        
+        
+def mayaTools():
+    
+    #import custom maya setup script
+    print cmds.internalVar(usd = True)
+    print cmds.internalVar(upd = True)
+    path = cmds.internalVar(usd = True)
+    path = path + "mayaTools.txt"
+    print path
+    
+    if os.path.exists(path):
+        f = open(path, 'r')
+        
+        mayaToolsDir = f.readline()
+        if not os.path.exists(mayaToolsDir):
+            mayaToolsInstall_UI()
+        
+        path = mayaToolsDir + "/General/Scripts"
+        pluginPath = mayaToolsDir + "/General/Plugins"
+        
+        #look in sys.path to see if path is in sys.path. if not, add it
+        if not path in sys.path:
+            sys.path.append(path)
+            
+            
+        #run setup
+        import customMayaMenu as cmm
+        cmm.customMayaMenu()
+        cmds.file(new = True, force = True)
+    
+    
+    else:
+        mayaToolsInstall_UI()
+        
+        
+
+#setup script jobs
+scriptJobNum2 = cmds.scriptJob(event = ["NewSceneOpened", mayaTools])
+
+
