Tools update

diff --git a/Engine/Extras/Maya_AnimationRiggingTools/MayaTools/General/Scripts/ART_addCharacter_UI.py b/Engine/Extras/Maya_AnimationRiggingTools/MayaTools/General/Scripts/ART_addCharacter_UI.py
--- a/Engine/Extras/Maya_AnimationRiggingTools/MayaTools/General/Scripts/ART_addCharacter_UI.py
+++ b/Engine/Extras/Maya_AnimationRiggingTools/MayaTools/General/Scripts/ART_addCharacter_UI.py
@@ -24,7 +24,7 @@ class AddCharacter_UI():
         if cmds.window("addCharacter_UI", exists = True):
             cmds.deleteUI("addCharacter_UI")
             
-        self.widgets["window"] = cmds.window("addCharacter_UI", title = "Add Character", w = 430, h = 260, sizeable = False, mxb = False, mnb = False)
+        self.widgets["window"] = cmds.window("addCharacter_UI", title = "Add Character", w = 430, h = 260, sizeable = True, mxb = False, mnb = False)
         
         #create the main layout
 	main = cmds.columnLayout(w = 430)
diff --git a/Engine/Extras/Maya_AnimationRiggingTools/MayaTools/General/Scripts/ART_animationUI.py b/Engine/Extras/Maya_AnimationRiggingTools/MayaTools/General/Scripts/ART_animationUI.py
--- a/Engine/Extras/Maya_AnimationRiggingTools/MayaTools/General/Scripts/ART_animationUI.py
+++ b/Engine/Extras/Maya_AnimationRiggingTools/MayaTools/General/Scripts/ART_animationUI.py
@@ -12,7 +12,7 @@ class AnimationUI():
         
         #check to see if there are any rigs in the scene, if not return
         characters = self.getCharacters()
-        
+        print characters
         if len(characters) == 0:
             result = cmds.confirmDialog(title = "Error", message = "No Characters found in scene. Would you like to add a character now?", button = ["Yes", "No"], defaultButton = "Yes", cancelButton = "No", dismissString = "No")
             if result == "Yes":
@@ -116,7 +116,7 @@ class AnimationUI():
             
         
         #build window
-        self.widgets["window"] = cmds.window("artAnimUI", w = 400, h = 700, title = "Animation UI", sizeable = False)
+        self.widgets["window"] = cmds.window("artAnimUI", w = 400, h = 700, title = "Animation UI", sizeable = True)
         
         #create the main layout
         self.widgets["topLevelLayout"] = cmds.columnLayout()
@@ -299,7 +299,7 @@ class AnimationUI():
                         "body_anim", "hip_anim", "clavicle_l_anim", "clavicle_r_anim", "fk_arm_l_anim", "fk_arm_r_anim", "fk_elbow_l_anim", "fk_elbow_r_anim", "fk_wrist_l_anim", "fk_wrist_r_anim",
                         "ik_elbow_l_anim", "ik_elbow_r_anim", "ik_wrist_l_anim", "ik_wrist_r_anim", "fk_thigh_l_anim", "fk_thigh_r_anim", "fk_calf_l_anim", "fk_calf_r_anim", "fk_foot_l_anim", "fk_foot_r_anim",
                         "fk_ball_l_anim", "fk_ball_r_anim", "ik_foot_anim_l", "ik_foot_anim_r", "heel_ctrl_l", "heel_ctrl_r", "toe_wiggle_ctrl_l", "toe_wiggle_ctrl_r",
-                        "toe_tip_ctrl_l", "toe_tip_ctrl_r", "master_anim", "offset_anim", "root_anim", "upperarm_l_twist_anim", "upperarm_l_twist2_anim", "upperarm_l_twist3_anim", "upperarm_r_twist_anim", "upperarm_r_twist2_anim", "upperarm_r_twist3_anim", "thigh_l_twist_anim", "thigh_r_twist_anim",
+                        "toe_tip_ctrl_l", "toe_tip_ctrl_r", "master_anim", "offset_anim", "root_anim", "upperarm_l_twist_anim", "upperarm_l_twist_2_anim", "upperarm_l_twist_3_anim", "upperarm_r_twist_anim", "upperarm_r_twist_2_anim", "upperarm_r_twist_3_anim", "thigh_l_twist_anim", "thigh_r_twist_anim",
                         "pinky_metacarpal_ctrl_l", "pinky_metacarpal_ctrl_r", "pinky_finger_fk_ctrl_1_l", "pinky_finger_fk_ctrl_1_r", "pinky_finger_fk_ctrl_2_l", "pinky_finger_fk_ctrl_2_r", "pinky_finger_fk_ctrl_3_l", "pinky_finger_fk_ctrl_3_r",
                         "ring_metacarpal_ctrl_l", "ring_metacarpal_ctrl_r", "ring_finger_fk_ctrl_1_l", "ring_finger_fk_ctrl_1_r", "ring_finger_fk_ctrl_2_l", "ring_finger_fk_ctrl_2_r", "ring_finger_fk_ctrl_3_l", "ring_finger_fk_ctrl_3_r",
                         "middle_metacarpal_ctrl_l", "middle_metacarpal_ctrl_r", "middle_finger_fk_ctrl_1_l", "middle_finger_fk_ctrl_1_r", "middle_finger_fk_ctrl_2_l", "middle_finger_fk_ctrl_2_r", "middle_finger_fk_ctrl_3_l", "middle_finger_fk_ctrl_3_r",
@@ -307,8 +307,8 @@ class AnimationUI():
                         "thumb_finger_fk_ctrl_1_l", "thumb_finger_fk_ctrl_1_r", "thumb_finger_fk_ctrl_2_l", "thumb_finger_fk_ctrl_2_r", "thumb_finger_fk_ctrl_3_l", "thumb_finger_fk_ctrl_3_r",
                         "index_l_ik_anim", "index_r_ik_anim", "middle_l_ik_anim", "middle_r_ik_anim", "ring_l_ik_anim", "ring_r_ik_anim", "pinky_l_ik_anim", "pinky_r_ik_anim", "thumb_l_ik_anim", "thumb_r_ik_anim",
                         "index_l_poleVector", "index_r_poleVector", "middle_l_poleVector", "middle_r_poleVector", "ring_l_poleVector", "ring_r_poleVector", "pinky_l_poleVector", "pinky_r_poleVector", "thumb_l_poleVector", "thumb_r_poleVector",
-                        "l_global_ik_anim", "r_global_ik_anim", "lowerarm_l_twist_anim", "lowerarm_l_twist2_anim", "lowerarm_l_twist3_anim", "lowerarm_r_twist_anim", "lowerarm_r_twist2_anim", "lowerarm_r_twist3_anim", "calf_r_twist_anim", "calf_r_twist2_anim", "calf_r_twist3_anim",
-	                "calf_l_twist_anim", "calf_l_twist2_anim", "calf_l_twist3_anim", "thigh_l_twist2_anim", "thigh_l_twist3_anim", "thigh_r_twist2_anim", "thigh_r_twist3_anim"]:
+                        "l_global_ik_anim", "r_global_ik_anim", "lowerarm_l_twist_anim", "lowerarm_l_twist_2_anim", "lowerarm_l_twist_3_anim", "lowerarm_r_twist_anim", "lowerarm_r_twist_2_anim", "lowerarm_r_twist_3_anim", "calf_r_twist_anim", "calf_r_twist_2_anim", "calf_r_twist_3_anim",
+	                "calf_l_twist_anim", "calf_l_twist_2_anim", "calf_l_twist_3_anim", "thigh_l_twist_2_anim", "thigh_l_twist_3_anim", "thigh_r_twist_2_anim", "thigh_r_twist_3_anim"]:
 	    self.controls.append(control)
 	    
 	#hack
@@ -389,7 +389,7 @@ class AnimationUI():
 	    if cmds.window("bakeDynToFKControls_UI", exists = True):
 		cmds.deleteUI("bakeDynToFKControls_UI")
 		
-	    window = cmds.window("bakeDynToFKControls_UI", title = "Bake Dynamics", w = 300, h = 400, sizeable = False, mnb = False, mxb = False)
+	    window = cmds.window("bakeDynToFKControls_UI", title = "Bake Dynamics", w = 300, h = 400, sizeable = True, mnb = False, mxb = False)
 	    mainLayout = cmds.formLayout(w = 300, h = 400)
 	    
 	    #textScrollList
@@ -1056,7 +1056,7 @@ class AnimationUI():
                         "body_anim", "hip_anim", "clavicle_l_anim", "clavicle_r_anim", "fk_arm_l_anim", "fk_arm_r_anim", "fk_elbow_l_anim", "fk_elbow_r_anim", "fk_wrist_l_anim", "fk_wrist_r_anim",
                         "ik_elbow_l_anim", "ik_elbow_r_anim", "ik_wrist_l_anim", "ik_wrist_r_anim", "fk_thigh_l_anim", "fk_thigh_r_anim", "fk_calf_l_anim", "fk_calf_r_anim", "fk_foot_l_anim", "fk_foot_r_anim",
                         "fk_ball_l_anim", "fk_ball_r_anim", "ik_foot_anim_l", "ik_foot_anim_r", "heel_ctrl_l", "heel_ctrl_r", "toe_wiggle_ctrl_l", "toe_wiggle_ctrl_r",
-                        "toe_tip_ctrl_l", "toe_tip_ctrl_r", "master_anim", "offset_anim", "root_anim", "upperarm_l_twist_anim", "upperarm_l_twist2_anim", "upperarm_l_twist3_anim", "upperarm_r_twist_anim", "upperarm_r_twist2_anim", "upperarm_r_twist3_anim", "thigh_l_twist_anim", "thigh_r_twist_anim",
+                        "toe_tip_ctrl_l", "toe_tip_ctrl_r", "master_anim", "offset_anim", "root_anim", "upperarm_l_twist_anim", "upperarm_l_twist_2_anim", "upperarm_l_twist_3_anim", "upperarm_r_twist_anim", "upperarm_r_twist_2_anim", "upperarm_r_twist_3_anim", "thigh_l_twist_anim", "thigh_r_twist_anim",
                         "pinky_metacarpal_ctrl_l", "pinky_metacarpal_ctrl_r", "pinky_finger_fk_ctrl_1_l", "pinky_finger_fk_ctrl_1_r", "pinky_finger_fk_ctrl_2_l", "pinky_finger_fk_ctrl_2_r", "pinky_finger_fk_ctrl_3_l", "pinky_finger_fk_ctrl_3_r",
                         "ring_metacarpal_ctrl_l", "ring_metacarpal_ctrl_r", "ring_finger_fk_ctrl_1_l", "ring_finger_fk_ctrl_1_r", "ring_finger_fk_ctrl_2_l", "ring_finger_fk_ctrl_2_r", "ring_finger_fk_ctrl_3_l", "ring_finger_fk_ctrl_3_r",
                         "middle_metacarpal_ctrl_l", "middle_metacarpal_ctrl_r", "middle_finger_fk_ctrl_1_l", "middle_finger_fk_ctrl_1_r", "middle_finger_fk_ctrl_2_l", "middle_finger_fk_ctrl_2_r", "middle_finger_fk_ctrl_3_l", "middle_finger_fk_ctrl_3_r",
@@ -1064,8 +1064,8 @@ class AnimationUI():
                         "thumb_finger_fk_ctrl_1_l", "thumb_finger_fk_ctrl_1_r", "thumb_finger_fk_ctrl_2_l", "thumb_finger_fk_ctrl_2_r", "thumb_finger_fk_ctrl_3_l", "thumb_finger_fk_ctrl_3_r",
                         "index_l_ik_anim", "index_r_ik_anim", "middle_l_ik_anim", "middle_r_ik_anim", "ring_l_ik_anim", "ring_r_ik_anim", "pinky_l_ik_anim", "pinky_r_ik_anim", "thumb_l_ik_anim", "thumb_r_ik_anim",
                         "index_l_poleVector", "index_r_poleVector", "middle_l_poleVector", "middle_r_poleVector", "ring_l_poleVector", "ring_r_poleVector", "pinky_l_poleVector", "pinky_r_poleVector", "thumb_l_poleVector", "thumb_r_poleVector",
-                        "l_global_ik_anim", "r_global_ik_anim", "lowerarm_l_twist_anim", "lowerarm_l_twist2_anim", "lowerarm_l_twist3_anim", "lowerarm_r_twist_anim", "lowerarm_r_twist2_anim", "lowerarm_r_twist3_anim", "calf_r_twist_anim", "calf_r_twist2_anim", "calf_r_twist3_anim",
-                        "calf_l_twist_anim", "calf_l_twist2_anim", "calf_l_twist3_anim", "thigh_l_twist2_anim", "thigh_l_twist3_anim", "thigh_r_twist2_anim", "thigh_r_twist3_anim"]:
+                        "l_global_ik_anim", "r_global_ik_anim", "lowerarm_l_twist_anim", "lowerarm_l_twist_2_anim", "lowerarm_l_twist_3_anim", "lowerarm_r_twist_anim", "lowerarm_r_twist_2_anim", "lowerarm_r_twist_3_anim", "calf_r_twist_anim", "calf_r_twist_2_anim", "calf_r_twist_3_anim",
+                        "calf_l_twist_anim", "calf_l_twist_2_anim", "calf_l_twist_3_anim", "thigh_l_twist_2_anim", "thigh_l_twist_3_anim", "thigh_r_twist_2_anim", "thigh_r_twist_3_anim"]:
 	    self.controls.append(control)
 	    
 	
@@ -1204,11 +1204,13 @@ class AnimationUI():
         
         referenceNodes = []
         references = cmds.ls(type = "reference")
-	
+	print references
         
         for reference in references:
             niceName = reference.rpartition("RN")[0]
+	    print niceName
             suffix = reference.rpartition("RN")[2]
+	    print suffix
             if suffix != "":
                 if cmds.objExists(niceName + suffix + ":" + "Skeleton_Settings"):
                     referenceNodes.append(niceName + suffix)
@@ -1216,7 +1218,8 @@ class AnimationUI():
             else:
                 if cmds.objExists(niceName + ":" + "Skeleton_Settings"):
                     referenceNodes.append(niceName)
-                
+	
+	print referenceNodes
         return referenceNodes
         
         
@@ -4389,74 +4392,74 @@ class AnimationUI():
             if cmds.objExists(character + ":upperarm_l_twist_anim"):
                 cmds.button(self.widgets[character + "_leftArmRollPickerButton"], edit = True, bgc = self.purple)
 		
-            if cmds.objExists(character + ":upperarm_l_twist2_anim"):
+            if cmds.objExists(character + ":upperarm_l_twist_2_anim"):
                 cmds.button(self.widgets[character + "_leftArmRoll2PickerButton"], edit = True, bgc = self.purple)
 		
-            if cmds.objExists(character + ":upperarm_l_twist3_anim"):
+            if cmds.objExists(character + ":upperarm_l_twist_3_anim"):
                 cmds.button(self.widgets[character + "_leftArmRoll3PickerButton"], edit = True, bgc = self.purple)
                 
             if cmds.objExists(character + ":upperarm_r_twist_anim"):
                 cmds.button(self.widgets[character + "_rightArmRollPickerButton"], edit = True, bgc = self.purple)
 		
-            if cmds.objExists(character + ":upperarm_r_twist2_anim"):
+            if cmds.objExists(character + ":upperarm_r_twist_2_anim"):
                 cmds.button(self.widgets[character + "_rightArmRoll2PickerButton"], edit = True, bgc = self.purple)
 	    
-            if cmds.objExists(character + ":upperarm_r_twist3_anim"):
+            if cmds.objExists(character + ":upperarm_r_twist_3_anim"):
                 cmds.button(self.widgets[character + "_rightArmRoll3PickerButton"], edit = True, bgc = self.purple)
 		
             if cmds.objExists(character + ":lowerarm_l_twist_anim"):
                 cmds.button(self.widgets[character + "_leftForeTwistPickerButton"], edit = True, bgc = self.purple)
 		
-            if cmds.objExists(character + ":lowerarm_l_twist2_anim"):
+            if cmds.objExists(character + ":lowerarm_l_twist_2_anim"):
                 cmds.button(self.widgets[character + "_leftForeTwist2PickerButton"], edit = True, bgc = self.purple)
 		
-            if cmds.objExists(character + ":lowerarm_l_twist3_anim"):
+            if cmds.objExists(character + ":lowerarm_l_twist_3_anim"):
                 cmds.button(self.widgets[character + "_leftForeTwist3PickerButton"], edit = True, bgc = self.purple)
 		
             if cmds.objExists(character + ":lowerarm_r_twist_anim"):
                 cmds.button(self.widgets[character + "_rightForeTwistPickerButton"], edit = True, bgc = self.purple)
 		
-            if cmds.objExists(character + ":lowerarm_r_twist2_anim"):
+            if cmds.objExists(character + ":lowerarm_r_twist_2_anim"):
                 cmds.button(self.widgets[character + "_rightForeTwist2PickerButton"], edit = True, bgc = self.purple)
 		
-            if cmds.objExists(character + ":lowerarm_r_twist3_anim"):
+            if cmds.objExists(character + ":lowerarm_r_twist_3_anim"):
                 cmds.button(self.widgets[character + "_rightForeTwist3PickerButton"], edit = True, bgc = self.purple)
                 
 
             if cmds.objExists(character + ":thigh_l_twist_anim"):
                 cmds.button(self.widgets[character + "_leftThighTwistPickerButton"], edit = True, bgc = self.purple)
 		
-            if cmds.objExists(character + ":thigh_l_twist2_anim"):
+            if cmds.objExists(character + ":thigh_l_twist_2_anim"):
                 cmds.button(self.widgets[character + "_leftThighTwist2PickerButton"], edit = True, bgc = self.purple)
 		
-            if cmds.objExists(character + ":thigh_l_twist3_anim"):
+            if cmds.objExists(character + ":thigh_l_twist_3_anim"):
                 cmds.button(self.widgets[character + "_leftThighTwist3PickerButton"], edit = True, bgc = self.purple)
                 
             if cmds.objExists(character + ":thigh_r_twist_anim"):
                 cmds.button(self.widgets[character + "_rightThighTwistPickerButton"], edit = True, bgc = self.purple)
 		
-            if cmds.objExists(character + ":thigh_r_twist2_anim"):
+            if cmds.objExists(character + ":thigh_r_twist_2_anim"):
                 cmds.button(self.widgets[character + "_rightThighTwist2PickerButton"], edit = True, bgc = self.purple)
 		
-            if cmds.objExists(character + ":thigh_r_twist3_anim"):
+            if cmds.objExists(character + ":thigh_r_twist_3_anim"):
                 cmds.button(self.widgets[character + "_rightThighTwist3PickerButton"], edit = True, bgc = self.purple)
 		
             if cmds.objExists(character + ":calf_l_twist_anim"):
                 cmds.button(self.widgets[character + "_leftCalfTwistPickerButton"], edit = True, bgc = self.purple)
 		
-            if cmds.objExists(character + ":calf_l_twist2_anim"):
+            if cmds.objExists(character + ":calf_l_twist_2_anim"):
                 cmds.button(self.widgets[character + "_leftCalfTwist2PickerButton"], edit = True, bgc = self.purple)
 		
-            if cmds.objExists(character + ":calf_l_twist3_anim"):
+            if cmds.objExists(character + ":calf_l_twist_3_anim"):
                 cmds.button(self.widgets[character + "_leftCalfTwist3PickerButton"], edit = True, bgc = self.purple)
 		
             if cmds.objExists(character + ":calf_r_twist_anim"):
                 cmds.button(self.widgets[character + "_rightCalfTwistPickerButton"], edit = True, bgc = self.purple)
 		
-            if cmds.objExists(character + ":calf_r_twist2_anim"):
+            if cmds.objExists(character + ":calf_r_twist_2_anim"):
                 cmds.button(self.widgets[character + "_rightCalfTwist2PickerButton"], edit = True, bgc = self.purple)
 		
-            if cmds.objExists(character + ":calf_r_twist3_anim"):
+            if cmds.objExists(character + ":calf_r_twist_3_anim"):
                 cmds.button(self.widgets[character + "_rightCalfTwist3PickerButton"], edit = True, bgc = self.purple)
 		
 		
@@ -6845,7 +6848,7 @@ class AnimationUI():
 	    
 	character = cmds.symbolButton(self.widgets["activeCharacterThumb"], q = True, ann = True)
 	
-	self.widgets["matchFrameRange_Window"] = cmds.window("matchOverFrameRange_UI", w = 500, h = 300, sizeable = False, title = "Match Over Frame Range", titleBarMenu = False)
+	self.widgets["matchFrameRange_Window"] = cmds.window("matchOverFrameRange_UI", w = 500, h = 300, sizeable = True, title = "Match Over Frame Range", titleBarMenu = False)
 	
 	#main layout
 	self.widgets["matchFrameRange_MainLayout"] = cmds.formLayout(w = 500, h = 300)
@@ -6912,7 +6915,7 @@ class AnimationUI():
 	if cmds.window("controlScaleWindow", exists = True):
 	    cmds.deleteUI("controlScaleWindow")
 	    
-	self.widgets["controlScaleWindow"] = cmds.window("controlScaleWindow", title = "CV Scale", w = 150, h = 50, sizeable = False, mnb = False, mxb = False)
+	self.widgets["controlScaleWindow"] = cmds.window("controlScaleWindow", title = "CV Scale", w = 150, h = 50, sizeable = True, mnb = False, mxb = False)
 	mainLayout = cmds.formLayout(w = 150, h = 50)
 	
 	smallerButton = cmds.button(label = "v", w = 50, h = 30, c = partial(self.control_scale_execute, .9))
@@ -7712,28 +7715,28 @@ class AnimationUI():
 	    
 		    
 	    try:
-		if cmds.objExists(character + ":" + "upperarm_l_twist2_anim"):
-		    if cmds.objExists(character + ":" + "upperarm_l_twist2_anim.buttonName"):
-			cmds.setAttr(character + ":" + "upperarm_l_twist2_anim.buttonName",self.widgets[character + "_leftArmRoll2PickerButton"], type = "string")
+		if cmds.objExists(character + ":" + "upperarm_l_twist_2_anim"):
+		    if cmds.objExists(character + ":" + "upperarm_l_twist_2_anim.buttonName"):
+			cmds.setAttr(character + ":" + "upperarm_l_twist_2_anim.buttonName",self.widgets[character + "_leftArmRoll2PickerButton"], type = "string")
 			
 		    else:
-			cmds.select(character + ":" + "upperarm_l_twist2_anim")
+			cmds.select(character + ":" + "upperarm_l_twist_2_anim")
 			cmds.addAttr(ln = "buttonName", dt = "string", keyable = False)
-			cmds.setAttr(character + ":" + "upperarm_l_twist2_anim.buttonName",self.widgets[character + "_leftArmRoll2PickerButton"], type = "string")
+			cmds.setAttr(character + ":" + "upperarm_l_twist_2_anim.buttonName",self.widgets[character + "_leftArmRoll2PickerButton"], type = "string")
 	    except:
 		pass
 	    
 		   
 		    
 	    try:
-		if cmds.objExists(character + ":" + "upperarm_l_twist3_anim"):
-		    if cmds.objExists(character + ":" + "upperarm_l_twist3_anim.buttonName"):
-			cmds.setAttr(character + ":" + "upperarm_l_twist3_anim.buttonName",self.widgets[character + "_leftArmRoll3PickerButton"], type = "string")
+		if cmds.objExists(character + ":" + "upperarm_l_twist_3_anim"):
+		    if cmds.objExists(character + ":" + "upperarm_l_twist_3_anim.buttonName"):
+			cmds.setAttr(character + ":" + "upperarm_l_twist_3_anim.buttonName",self.widgets[character + "_leftArmRoll3PickerButton"], type = "string")
 			
 		    else:
-			cmds.select(character + ":" + "upperarm_l_twist3_anim")
+			cmds.select(character + ":" + "upperarm_l_twist_3_anim")
 			cmds.addAttr(ln = "buttonName", dt = "string", keyable = False)
-			cmds.setAttr(character + ":" + "upperarm_l_twist3_anim.buttonName",self.widgets[character + "_leftArmRoll3PickerButton"], type = "string")
+			cmds.setAttr(character + ":" + "upperarm_l_twist_3_anim.buttonName",self.widgets[character + "_leftArmRoll3PickerButton"], type = "string")
 	    except:
 		pass
 	    
@@ -7755,28 +7758,28 @@ class AnimationUI():
 	    
 		    
 	    try:
-		if cmds.objExists(character + ":" + "upperarm_r_twist2_anim"):
-		    if cmds.objExists(character + ":" + "upperarm_r_twist2_anim.buttonName"):
-			cmds.setAttr(character + ":" + "upperarm_r_twist2_anim.buttonName",self.widgets[character + "_rightArmRoll2PickerButton"], type = "string")
+		if cmds.objExists(character + ":" + "upperarm_r_twist_2_anim"):
+		    if cmds.objExists(character + ":" + "upperarm_r_twist_2_anim.buttonName"):
+			cmds.setAttr(character + ":" + "upperarm_r_twist_2_anim.buttonName",self.widgets[character + "_rightArmRoll2PickerButton"], type = "string")
 			
 		    else:
-			cmds.select(character + ":" + "upperarm_r_twist2_anim")
+			cmds.select(character + ":" + "upperarm_r_twist_2_anim")
 			cmds.addAttr(ln = "buttonName", dt = "string", keyable = False)
-			cmds.setAttr(character + ":" + "upperarm_r_twist2_anim.buttonName",self.widgets[character + "_rightArmRoll2PickerButton"], type = "string")
+			cmds.setAttr(character + ":" + "upperarm_r_twist_2_anim.buttonName",self.widgets[character + "_rightArmRoll2PickerButton"], type = "string")
 	    except:
 		pass
 	    
 		    
 		    
 	    try:
-		if cmds.objExists(character + ":" + "upperarm_r_twist3_anim"):
-		    if cmds.objExists(character + ":" + "upperarm_r_twist3_anim.buttonName"):
-			cmds.setAttr(character + ":" + "upperarm_r_twist3_anim.buttonName",self.widgets[character + "_rightArmRoll3PickerButton"], type = "string")
+		if cmds.objExists(character + ":" + "upperarm_r_twist_3_anim"):
+		    if cmds.objExists(character + ":" + "upperarm_r_twist_3_anim.buttonName"):
+			cmds.setAttr(character + ":" + "upperarm_r_twist_3_anim.buttonName",self.widgets[character + "_rightArmRoll3PickerButton"], type = "string")
 			
 		    else:
-			cmds.select(character + ":" + "upperarm_r_twist3_anim")
+			cmds.select(character + ":" + "upperarm_r_twist_3_anim")
 			cmds.addAttr(ln = "buttonName", dt = "string", keyable = False)
-			cmds.setAttr(character + ":" + "upperarm_r_twist3_anim.buttonName",self.widgets[character + "_rightArmRoll3PickerButton"], type = "string")
+			cmds.setAttr(character + ":" + "upperarm_r_twist_3_anim.buttonName",self.widgets[character + "_rightArmRoll3PickerButton"], type = "string")
 	    except:
 		pass
 	    
@@ -7799,27 +7802,27 @@ class AnimationUI():
 	    
 		    
 	    try:
-		if cmds.objExists(character + ":" + "lowerarm_l_twist2_anim"):
-		    if cmds.objExists(character + ":" + "lowerarm_l_twist2_anim.buttonName"):
-			cmds.setAttr(character + ":" + "lowerarm_l_twist2_anim.buttonName",self.widgets[character + "_leftForeTwist2PickerButton"], type = "string")
+		if cmds.objExists(character + ":" + "lowerarm_l_twist_2_anim"):
+		    if cmds.objExists(character + ":" + "lowerarm_l_twist_2_anim.buttonName"):
+			cmds.setAttr(character + ":" + "lowerarm_l_twist_2_anim.buttonName",self.widgets[character + "_leftForeTwist2PickerButton"], type = "string")
 			
 		    else:
-			cmds.select(character + ":" + "lowerarm_l_twist2_anim")
+			cmds.select(character + ":" + "lowerarm_l_twist_2_anim")
 			cmds.addAttr(ln = "buttonName", dt = "string", keyable = False)
-			cmds.setAttr(character + ":" + "lowerarm_l_twist2_anim.buttonName",self.widgets[character + "_leftForeTwist2PickerButton"], type = "string")
+			cmds.setAttr(character + ":" + "lowerarm_l_twist_2_anim.buttonName",self.widgets[character + "_leftForeTwist2PickerButton"], type = "string")
 	    except:
 		pass
 	    
 	    
 	    try:
-		if cmds.objExists(character + ":" + "lowerarm_l_twist3_anim"):
-		    if cmds.objExists(character + ":" + "lowerarm_l_twist3_anim.buttonName"):
-			cmds.setAttr(character + ":" + "lowerarm_l_twist3_anim.buttonName",self.widgets[character + "_leftForeTwist3PickerButton"], type = "string")
+		if cmds.objExists(character + ":" + "lowerarm_l_twist_3_anim"):
+		    if cmds.objExists(character + ":" + "lowerarm_l_twist_3_anim.buttonName"):
+			cmds.setAttr(character + ":" + "lowerarm_l_twist_3_anim.buttonName",self.widgets[character + "_leftForeTwist3PickerButton"], type = "string")
 			
 		    else:
-			cmds.select(character + ":" + "lowerarm_l_twist3_anim")
+			cmds.select(character + ":" + "lowerarm_l_twist_3_anim")
 			cmds.addAttr(ln = "buttonName", dt = "string", keyable = False)
-			cmds.setAttr(character + ":" + "lowerarm_l_twist3_anim.buttonName",self.widgets[character + "_leftForeTwist3PickerButton"], type = "string")
+			cmds.setAttr(character + ":" + "lowerarm_l_twist_3_anim.buttonName",self.widgets[character + "_leftForeTwist3PickerButton"], type = "string")
 	    except:
 		pass
 	    
@@ -7840,28 +7843,28 @@ class AnimationUI():
 	    
 		    
 	    try:
-		if cmds.objExists(character + ":" + "lowerarm_r_twist2_anim"):
-		    if cmds.objExists(character + ":" + "lowerarm_r_twist2_anim.buttonName"):
-			cmds.setAttr(character + ":" + "lowerarm_r_twist2_anim.buttonName",self.widgets[character + "_rightForeTwist2PickerButton"], type = "string")
+		if cmds.objExists(character + ":" + "lowerarm_r_twist_2_anim"):
+		    if cmds.objExists(character + ":" + "lowerarm_r_twist_2_anim.buttonName"):
+			cmds.setAttr(character + ":" + "lowerarm_r_twist_2_anim.buttonName",self.widgets[character + "_rightForeTwist2PickerButton"], type = "string")
 			
 		    else:
-			cmds.select(character + ":" + "lowerarm_r_twist2_anim")
+			cmds.select(character + ":" + "lowerarm_r_twist_2_anim")
 			cmds.addAttr(ln = "buttonName", dt = "string", keyable = False)
-			cmds.setAttr(character + ":" + "lowerarm_r_twist2_anim.buttonName",self.widgets[character + "_rightForeTwist2PickerButton"], type = "string")
+			cmds.setAttr(character + ":" + "lowerarm_r_twist_2_anim.buttonName",self.widgets[character + "_rightForeTwist2PickerButton"], type = "string")
 	    except:
 		pass
 	    
 	    
 		    
 	    try:
-		if cmds.objExists(character + ":" + "lowerarm_r_twist3_anim"):
-		    if cmds.objExists(character + ":" + "lowerarm_r_twist3_anim.buttonName"):
-			cmds.setAttr(character + ":" + "lowerarm_r_twist3_anim.buttonName",self.widgets[character + "_rightForeTwist3PickerButton"], type = "string")
+		if cmds.objExists(character + ":" + "lowerarm_r_twist_3_anim"):
+		    if cmds.objExists(character + ":" + "lowerarm_r_twist_3_anim.buttonName"):
+			cmds.setAttr(character + ":" + "lowerarm_r_twist_3_anim.buttonName",self.widgets[character + "_rightForeTwist3PickerButton"], type = "string")
 			
 		    else:
-			cmds.select(character + ":" + "lowerarm_r_twist3_anim")
+			cmds.select(character + ":" + "lowerarm_r_twist_3_anim")
 			cmds.addAttr(ln = "buttonName", dt = "string", keyable = False)
-			cmds.setAttr(character + ":" + "lowerarm_r_twist3_anim.buttonName",self.widgets[character + "_rightForeTwist3PickerButton"], type = "string")
+			cmds.setAttr(character + ":" + "lowerarm_r_twist_3_anim.buttonName",self.widgets[character + "_rightForeTwist3PickerButton"], type = "string")
 	    except:
 		pass
 	    
@@ -7884,28 +7887,28 @@ class AnimationUI():
 		    
 		    
 	    try:
-		if cmds.objExists(character + ":" + "thigh_l_twist2_anim"):
-		    if cmds.objExists(character + ":" + "thigh_l_twist2_anim.buttonName"):
-			cmds.setAttr(character + ":" + "thigh_l_twist2_anim.buttonName",self.widgets[character + "_leftThighTwist2PickerButton"], type = "string")
+		if cmds.objExists(character + ":" + "thigh_l_twist_2_anim"):
+		    if cmds.objExists(character + ":" + "thigh_l_twist_2_anim.buttonName"):
+			cmds.setAttr(character + ":" + "thigh_l_twist_2_anim.buttonName",self.widgets[character + "_leftThighTwist2PickerButton"], type = "string")
 			
 		    else:
-			cmds.select(character + ":" + "thigh_l_twist2_anim")
+			cmds.select(character + ":" + "thigh_l_twist_2_anim")
 			cmds.addAttr(ln = "buttonName", dt = "string", keyable = False)
-			cmds.setAttr(character + ":" + "thigh_l_twist2_anim.buttonName",self.widgets[character + "_leftThighTwist2PickerButton"], type = "string")
+			cmds.setAttr(character + ":" + "thigh_l_twist_2_anim.buttonName",self.widgets[character + "_leftThighTwist2PickerButton"], type = "string")
 	    except:
 		pass
 	    
 		    
 		    
 	    try:
-		if cmds.objExists(character + ":" + "thigh_l_twist3_anim"):
-		    if cmds.objExists(character + ":" + "thigh_l_twist3_anim.buttonName"):
-			cmds.setAttr(character + ":" + "thigh_l_twist3_anim.buttonName",self.widgets[character + "_leftThighTwist3PickerButton"], type = "string")
+		if cmds.objExists(character + ":" + "thigh_l_twist_3_anim"):
+		    if cmds.objExists(character + ":" + "thigh_l_twist_3_anim.buttonName"):
+			cmds.setAttr(character + ":" + "thigh_l_twist_3_anim.buttonName",self.widgets[character + "_leftThighTwist3PickerButton"], type = "string")
 			
 		    else:
-			cmds.select(character + ":" + "thigh_l_twist3_anim")
+			cmds.select(character + ":" + "thigh_l_twist_3_anim")
 			cmds.addAttr(ln = "buttonName", dt = "string", keyable = False)
-			cmds.setAttr(character + ":" + "thigh_l_twist3_anim.buttonName",self.widgets[character + "_leftThighTwist3PickerButton"], type = "string")
+			cmds.setAttr(character + ":" + "thigh_l_twist_3_anim.buttonName",self.widgets[character + "_leftThighTwist3PickerButton"], type = "string")
 	    except:
 		pass
 	    
@@ -7928,28 +7931,28 @@ class AnimationUI():
 		    
 		    
 	    try:
-		if cmds.objExists(character + ":" + "thigh_r_twist2_anim"):
-		    if cmds.objExists(character + ":" + "thigh_r_twist2_anim.buttonName"):
-			cmds.setAttr(character + ":" + "thigh_r_twist2_anim.buttonName",self.widgets[character + "_rightThighTwist2PickerButton"], type = "string")
+		if cmds.objExists(character + ":" + "thigh_r_twist_2_anim"):
+		    if cmds.objExists(character + ":" + "thigh_r_twist_2_anim.buttonName"):
+			cmds.setAttr(character + ":" + "thigh_r_twist_2_anim.buttonName",self.widgets[character + "_rightThighTwist2PickerButton"], type = "string")
 			
 		    else:
-			cmds.select(character + ":" + "thigh_r_twist2_anim")
+			cmds.select(character + ":" + "thigh_r_twist_2_anim")
 			cmds.addAttr(ln = "buttonName", dt = "string", keyable = False)
-			cmds.setAttr(character + ":" + "thigh_r_twist2_anim.buttonName",self.widgets[character + "_rightThighTwist2PickerButton"], type = "string")
+			cmds.setAttr(character + ":" + "thigh_r_twist_2_anim.buttonName",self.widgets[character + "_rightThighTwist2PickerButton"], type = "string")
 	    except:
 		pass
 	    
 	    
 		    
 	    try:
-		if cmds.objExists(character + ":" + "thigh_r_twist3_anim"):
-		    if cmds.objExists(character + ":" + "thigh_r_twist3_anim.buttonName"):
-			cmds.setAttr(character + ":" + "thigh_r_twist3_anim.buttonName",self.widgets[character + "_rightThighTwist3PickerButton"], type = "string")
+		if cmds.objExists(character + ":" + "thigh_r_twist_3_anim"):
+		    if cmds.objExists(character + ":" + "thigh_r_twist_3_anim.buttonName"):
+			cmds.setAttr(character + ":" + "thigh_r_twist_3_anim.buttonName",self.widgets[character + "_rightThighTwist3PickerButton"], type = "string")
 			
 		    else:
-			cmds.select(character + ":" + "thigh_r_twist3_anim")
+			cmds.select(character + ":" + "thigh_r_twist_3_anim")
 			cmds.addAttr(ln = "buttonName", dt = "string", keyable = False)
-			cmds.setAttr(character + ":" + "thigh_r_twist3_anim.buttonName",self.widgets[character + "_rightThighTwist3PickerButton"], type = "string")
+			cmds.setAttr(character + ":" + "thigh_r_twist_3_anim.buttonName",self.widgets[character + "_rightThighTwist3PickerButton"], type = "string")
 			
 	    except:
 		pass
@@ -7973,28 +7976,28 @@ class AnimationUI():
 	    
 		    
 	    try:
-		if cmds.objExists(character + ":" + "calf_l_twist2_anim"):
-		    if cmds.objExists(character + ":" + "calf_l_twist2_anim.buttonName"):
-			cmds.setAttr(character + ":" + "calf_l_twist2_anim.buttonName",self.widgets[character + "_leftCalfTwist2PickerButton"], type = "string")
+		if cmds.objExists(character + ":" + "calf_l_twist_2_anim"):
+		    if cmds.objExists(character + ":" + "calf_l_twist_2_anim.buttonName"):
+			cmds.setAttr(character + ":" + "calf_l_twist_2_anim.buttonName",self.widgets[character + "_leftCalfTwist2PickerButton"], type = "string")
 			
 		    else:
-			cmds.select(character + ":" + "calf_l_twist2_anim")
+			cmds.select(character + ":" + "calf_l_twist_2_anim")
 			cmds.addAttr(ln = "buttonName", dt = "string", keyable = False)
-			cmds.setAttr(character + ":" + "calf_l_twist2_anim.buttonName",self.widgets[character + "_leftCalfTwist2PickerButton"], type = "string")
+			cmds.setAttr(character + ":" + "calf_l_twist_2_anim.buttonName",self.widgets[character + "_leftCalfTwist2PickerButton"], type = "string")
 	    except:
 		pass
 	    
 		    
 
 	    try:
-		if cmds.objExists(character + ":" + "calf_l_twist3_anim"):
-		    if cmds.objExists(character + ":" + "calf_l_twist3_anim.buttonName"):
-			cmds.setAttr(character + ":" + "calf_l_twist3_anim.buttonName",self.widgets[character + "_leftCalfTwist3PickerButton"], type = "string")
+		if cmds.objExists(character + ":" + "calf_l_twist_3_anim"):
+		    if cmds.objExists(character + ":" + "calf_l_twist_3_anim.buttonName"):
+			cmds.setAttr(character + ":" + "calf_l_twist_3_anim.buttonName",self.widgets[character + "_leftCalfTwist3PickerButton"], type = "string")
 			
 		    else:
-			cmds.select(character + ":" + "calf_l_twist3_anim")
+			cmds.select(character + ":" + "calf_l_twist_3_anim")
 			cmds.addAttr(ln = "buttonName", dt = "string", keyable = False)
-			cmds.setAttr(character + ":" + "calf_l_twist3_anim.buttonName",self.widgets[character + "_leftCalfTwist3PickerButton"], type = "string")
+			cmds.setAttr(character + ":" + "calf_l_twist_3_anim.buttonName",self.widgets[character + "_leftCalfTwist3PickerButton"], type = "string")
 	    except:
 		pass
 	    
@@ -8015,28 +8018,28 @@ class AnimationUI():
 		    
 
 	    try:
-		if cmds.objExists(character + ":" + "calf_r_twist2_anim"):
-		    if cmds.objExists(character + ":" + "calf_r_twist2_anim.buttonName"):
-			cmds.setAttr(character + ":" + "calf_r_twist2_anim.buttonName",self.widgets[character + "_rightCalfTwist2PickerButton"], type = "string")
+		if cmds.objExists(character + ":" + "calf_r_twist_2_anim"):
+		    if cmds.objExists(character + ":" + "calf_r_twist_2_anim.buttonName"):
+			cmds.setAttr(character + ":" + "calf_r_twist_2_anim.buttonName",self.widgets[character + "_rightCalfTwist2PickerButton"], type = "string")
 			
 		    else:
-			cmds.select(character + ":" + "calf_r_twist2_anim")
+			cmds.select(character + ":" + "calf_r_twist_2_anim")
 			cmds.addAttr(ln = "buttonName", dt = "string", keyable = False)
-			cmds.setAttr(character + ":" + "calf_r_twist2_anim.buttonName",self.widgets[character + "_rightCalfTwist2PickerButton"], type = "string")
+			cmds.setAttr(character + ":" + "calf_r_twist_2_anim.buttonName",self.widgets[character + "_rightCalfTwist2PickerButton"], type = "string")
 	    except:
 		pass
 	    
 		    
 		    
 	    try:
-		if cmds.objExists(character + ":" + "calf_r_twist3_anim"):
-		    if cmds.objExists(character + ":" + "calf_r_twist3_anim.buttonName"):
-			cmds.setAttr(character + ":" + "calf_r_twist3_anim.buttonName",self.widgets[character + "_rightCalfTwist3PickerButton"], type = "string")
+		if cmds.objExists(character + ":" + "calf_r_twist_3_anim"):
+		    if cmds.objExists(character + ":" + "calf_r_twist_3_anim.buttonName"):
+			cmds.setAttr(character + ":" + "calf_r_twist_3_anim.buttonName",self.widgets[character + "_rightCalfTwist3PickerButton"], type = "string")
 			
 		    else:
-			cmds.select(character + ":" + "calf_r_twist3_anim")
+			cmds.select(character + ":" + "calf_r_twist_3_anim")
 			cmds.addAttr(ln = "buttonName", dt = "string", keyable = False)
-			cmds.setAttr(character + ":" + "calf_r_twist3_anim.buttonName",self.widgets[character + "_rightCalfTwist3PickerButton"], type = "string")
+			cmds.setAttr(character + ":" + "calf_r_twist_3_anim.buttonName",self.widgets[character + "_rightCalfTwist3PickerButton"], type = "string")
 	    
 	    except:
 		pass
diff --git a/Engine/Extras/Maya_AnimationRiggingTools/MayaTools/General/Scripts/ART_autoRigger.py b/Engine/Extras/Maya_AnimationRiggingTools/MayaTools/General/Scripts/ART_autoRigger.py
--- a/Engine/Extras/Maya_AnimationRiggingTools/MayaTools/General/Scripts/ART_autoRigger.py
+++ b/Engine/Extras/Maya_AnimationRiggingTools/MayaTools/General/Scripts/ART_autoRigger.py
@@ -1122,6 +1122,8 @@ class AutoRigger():
 		    twistJoint = child
 	    
 	    topSpineJointConstraint = cmds.pointConstraint(topJoint, twistJoint, mo = True)[0]
+	    topSpineBone = twistJoint.partition("twist_")[2]
+	    cmds.pointConstraint(topSpineBone, twistJoint)[0]
 		    
 		    
 	    #connect attr on top spine joint constraint
@@ -4515,7 +4517,7 @@ class AutoRigger():
 
 	#skin ribbon
 	cmds.select([ribbon, topSkinJoint, bottomSkinJoint])
-	skin = cmds.skinCluster(tsb = True, mi = 2, omi = True, dr = 5, bm = 0, sm = 0)
+	skin = cmds.skinCluster(tsb = True, mi = 2, omi = True, dr = 5, sm = 0)
 	
 	
 	#position the joints, thus moving the ribbon
@@ -4643,7 +4645,7 @@ class AutoRigger():
 	for joint in skinJoints:
 	    cmds.select(joint, add = True)
 	    
-	skin = cmds.skinCluster(tsb = True, mi = 2, omi = True, dr = 5, bm = 0, sm = 0)
+	skin = cmds.skinCluster(tsb = True, mi = 2, omi = True, dr = 5, sm = 0)
 	
 	
 	#orient roll grp to both fk/ik arm joints and set driven keys between them
diff --git a/Engine/Extras/Maya_AnimationRiggingTools/MayaTools/General/Scripts/ART_poseEditor.py b/Engine/Extras/Maya_AnimationRiggingTools/MayaTools/General/Scripts/ART_poseEditor.py
--- a/Engine/Extras/Maya_AnimationRiggingTools/MayaTools/General/Scripts/ART_poseEditor.py
+++ b/Engine/Extras/Maya_AnimationRiggingTools/MayaTools/General/Scripts/ART_poseEditor.py
@@ -253,10 +253,11 @@ class PoseEditor_UI():
 	for node in selectNodes:
 	    attrs = cmds.listAttr(node, string = "space_*")
 	    
-	    for attr in attrs:
-		if cmds.getAttr(node + "." + attr) != 0:
-		    if not attr in safeAttrs:
-			warningNodes.append(node + "." + attr)
+	    if attrs != None:
+		for attr in attrs:
+		    if cmds.getAttr(node + "." + attr) != 0:
+			if not attr in safeAttrs:
+			    warningNodes.append(node + "." + attr)
 		    
 	#if length of warning nodes is greater than 0, warn user that cannot save poses with spaces
 	if len(warningNodes) > 0:
diff --git a/Engine/Extras/Maya_AnimationRiggingTools/MayaTools/General/Scripts/ART_skeletonBuilder_UI.py b/Engine/Extras/Maya_AnimationRiggingTools/MayaTools/General/Scripts/ART_skeletonBuilder_UI.py
--- a/Engine/Extras/Maya_AnimationRiggingTools/MayaTools/General/Scripts/ART_skeletonBuilder_UI.py
+++ b/Engine/Extras/Maya_AnimationRiggingTools/MayaTools/General/Scripts/ART_skeletonBuilder_UI.py
@@ -159,7 +159,7 @@ class SkeletonBuilder_UI():
 	if cmds.window("SkelBuilder_window", exists = True):
 	    cmds.deleteUI("SkelBuilder_window")
 	    
-        self.widgets["window"] = cmds.window("SkelBuilder_window", title = "Skeleton Builder", titleBarMenu = False, sizeable = False)
+        self.widgets["window"] = cmds.window("SkelBuilder_window", title = "Skeleton Builder", titleBarMenu = False, sizeable = True)
         
         #create layout
         self.widgets["mainLayout"] = cmds.columnLayout(rs = 5)
@@ -2373,7 +2373,7 @@ class SkeletonBuilder_UI():
 	    cmds.deleteUI("extraJointInfo_UI")
 	    
 	#create the window and the layouts
-	window = cmds.window("extraJointInfo_UI", w = 300, h = 150, title = "Joint Info", tbm = False, sizeable = False)
+	window = cmds.window("extraJointInfo_UI", w = 300, h = 150, title = "Joint Info", tbm = False, sizeable = True)
 	mainLayout = cmds.columnLayout(w = 300, h = 150)
 	formLayout = cmds.formLayout(w = 300, h = 150, parent = mainLayout)
 	
@@ -2426,7 +2426,7 @@ class SkeletonBuilder_UI():
 	    cmds.deleteUI("dynamicJointInfo_UI")
 	    
 	#create the window and the layouts
-	window = cmds.window("dynamicJointInfo_UI", w = 300, h = 200, title = "Joint Info", tbm = False, sizeable = False)
+	window = cmds.window("dynamicJointInfo_UI", w = 300, h = 200, title = "Joint Info", tbm = False, sizeable = True)
 	mainLayout = cmds.columnLayout(w = 300, h = 200)
 	formLayout = cmds.formLayout(w = 300, h = 200, parent = mainLayout)
 	
@@ -3219,7 +3219,7 @@ class SkeletonBuilder_UI():
         if cmds.window("loadTemplateUI", exists = True):
             cmds.deleteUI("loadTemplateUI")
             
-        window = cmds.window("loadTemplateUI", w = 300, h = 200, title = "Load Template", titleBarMenu = False, sizeable = False)
+        window = cmds.window("loadTemplateUI", w = 300, h = 200, title = "Load Template", titleBarMenu = False, sizeable = True)
         mainLayout = cmds.columnLayout(w = 300, h = 300)
         menuBarLayout = cmds.menuBarLayout(h = 30, w = 300, ebg = True, bgc = [.1, .1, .1], parent = mainLayout)
         cmds.menu(label = "File")
@@ -6280,7 +6280,7 @@ class SkeletonBuilder_UI():
 	if cmds.window("rigPoseCreationUI", exists = True):
 	    cmds.deleteUI("rigPoseCreationUI", window = True)
 	    
-	window = cmds.window("rigPoseCreationUI", w = 200, h = 130, sizeable = False, titleBar = True, titleBarMenu = True, mnb = False, mxb = False, title = "Create Rig Pose")
+	window = cmds.window("rigPoseCreationUI", w = 200, h = 130, sizeable = True, titleBar = True, titleBarMenu = True, mnb = False, mxb = False, title = "Create Rig Pose")
 	layout = cmds.columnLayout(w = 200, h = 130, bgc = [0,0,0])
 	
 	#banner
@@ -6690,7 +6690,7 @@ class SkeletonBuilder_UI():
 	if cmds.window("skinProxyGeo_UI_question", exists = True):
 	    cmds.deleteUI("skinProxyGeo_UI_question")
 	    
-	window = cmds.window("skinProxyGeo_UI_question", title = "Skin Proxy Mesh", w = 500, h = 225, sizeable = False, titleBarMenu = False)
+	window = cmds.window("skinProxyGeo_UI_question", title = "Skin Proxy Mesh", w = 500, h = 225, sizeable = True, titleBarMenu = False)
 	mainLayout = cmds.formLayout()
 	
 	imagePath = self.mayaToolsDir + "/General/Icons/ART/"
@@ -6774,7 +6774,7 @@ class SkeletonBuilder_UI():
 	if cmds.window("createRigPoseConfirmationUI", exists = True):
 	    cmds.deleteUI("createRigPoseConfirmationUI")
 	    
-	window = cmds.window("createRigPoseConfirmationUI", title = "Create Rig Pose", w = 500, h = 305, sizeable = False, titleBarMenu = False)
+	window = cmds.window("createRigPoseConfirmationUI", title = "Create Rig Pose", w = 500, h = 305, sizeable = True, titleBarMenu = False)
 	mainLayout = cmds.formLayout()
 	
 	#image
@@ -7602,7 +7602,7 @@ class SkeletonBuilder_UI():
 	    cmds.deleteUI("publishResultsInterface")
 
 	#create window
-	self.widgets["publishResultsWindow"] = cmds.window("publishResultsInterface", title = "Character Published!", w = 480, h = 200, mnb = False, mxb = False, sizeable = False)
+	self.widgets["publishResultsWindow"] = cmds.window("publishResultsInterface", title = "Character Published!", w = 480, h = 200, mnb = False, mxb = False, sizeable = True)
 	
 	#widgets
 	self.widgets["publishResults_mainLayout"] = cmds.columnLayout(rs = 10)
@@ -7815,7 +7815,7 @@ class SkeletonBuilder_UI():
 	if cmds.window("publish_ui_window", exists = True):
 	    cmds.deleteUI("publish_ui_window")
 	    
-	self.widgets["publishUIWindow"] = cmds.window("publish_ui_window", w = 300, h = 310, mxb = False, mnb = False, sizeable = False, title = "Publish Character")
+	self.widgets["publishUIWindow"] = cmds.window("publish_ui_window", w = 300, h = 310, mxb = False, mnb = False, sizeable = True, title = "Publish Character")
 	self.widgets["publishUIMainLayout"] = cmds.columnLayout(w = 300, h =310, co = ["both", 5])
 	
 	#banner image
@@ -8190,7 +8190,7 @@ class SkeletonBuilder_UI():
 		    
 		    #Create a window with a model editor 
 	    
-		    self.widgets["screenshotWindow"] =  cmds.window(title = "Character Thumbnail Creator", w = 400, h = 350, sizeable = False, titleBar = False)
+		    self.widgets["screenshotWindow"] =  cmds.window(title = "Character Thumbnail Creator", w = 400, h = 350, sizeable = True, titleBar = False)
 		    self.widgets["screenshotFormLayout"] = cmds.formLayout()
 		    self.widgets["screenshot_modelEditor"] = cmds.modelEditor(parent = self.widgets["screenshotFormLayout"])
 		    self.widgets["screenshot_button_frame"] = cmds.frameLayout(h = 488, borderStyle = "etchedIn", lv = False, collapsable = False,  parent = self.widgets["screenshotFormLayout"])
@@ -8333,7 +8333,7 @@ class SkeletonBuilder_UI():
 	    
 	    #Create a window with a model editor 
     
-	    self.widgets["screenshotWindow"] =  cmds.window(title = "Character Thumbnail Creator", w = 400, h = 350, sizeable = False, titleBar = False)
+	    self.widgets["screenshotWindow"] =  cmds.window(title = "Character Thumbnail Creator", w = 400, h = 350, sizeable = True, titleBar = False)
 	    self.widgets["screenshotFormLayout"] = cmds.formLayout()
 	    self.widgets["screenshot_modelEditor"] = cmds.modelEditor(parent = self.widgets["screenshotFormLayout"])
 	    self.widgets["screenshot_button_frame"] = cmds.frameLayout(h = 488, borderStyle = "etchedIn", lv = False, collapsable = False,  parent = self.widgets["screenshotFormLayout"])
@@ -9203,7 +9203,7 @@ class SkeletonBuilder_UI():
         if cmds.window("loadTemplateUI", exists = True):
             cmds.deleteUI("loadTemplateUI")
             
-        window = cmds.window("loadTemplateUI", w = 300, h = 200, title = "Load Template", titleBarMenu = False, sizeable = False)
+        window = cmds.window("loadTemplateUI", w = 300, h = 200, title = "Load Template", titleBarMenu = False, sizeable = True)
         mainLayout = cmds.columnLayout(w = 300, h = 300)
         menuBarLayout = cmds.menuBarLayout(h = 30, w = 300, ebg = True, bgc = [.1, .1, .1], parent = mainLayout)
         cmds.menu(label = "File")
@@ -9320,7 +9320,7 @@ class SkeletonBuilder_UI():
 	if cmds.window("rangeOfMotion_UI", exists = True):
 	    cmds.deleteUI("rangeOfMotion_UI")
 	    
-	self.widgets["rom_window"] = cmds.window("rangeOfMotion_UI", w = 600, h = 400, title = "Range Of Motion Generator", mxb = False, mnb = False, sizeable = False)
+	self.widgets["rom_window"] = cmds.window("rangeOfMotion_UI", w = 600, h = 400, title = "Range Of Motion Generator", mxb = False, mnb = False, sizeable = True)
 	
 	#create the main layout
 	self.widgets["rom_formLayout"] = cmds.formLayout(w = 600, h = 400)
@@ -10148,7 +10148,7 @@ class SkeletonBuilder_UI():
 	    if cmds.window("paintWeightsMoveInfsUI", exists = True):
 		cmds.deleteUI("paintWeightsMoveInfsUI")
 		
-	    window = cmds.window("paintWeightsMoveInfsUI", w = 300, h = 60, title = "Move Influences", titleBarMenu = False, sizeable = False)
+	    window = cmds.window("paintWeightsMoveInfsUI", w = 300, h = 60, title = "Move Influences", titleBarMenu = False, sizeable = True)
 	    mainLayout = cmds.columnLayout(w = 300, h = 60)
 	    
 	    cmds.separator(h = 10)
@@ -10448,7 +10448,7 @@ class SkeletonBuilder_UI():
 	if cmds.window("paintWeightsAddInf", exists = True):
 	    cmds.deleteUI("paintWeightsAddInf")
 	    
-	window = cmds.window("paintWeightsAddInf", t = "Add Influence", tbm = False, w = 300, h = 300, sizeable = False)
+	window = cmds.window("paintWeightsAddInf", t = "Add Influence", tbm = False, w = 300, h = 300, sizeable = True)
 	
 	mainLayout = cmds.columnLayout(w = 300, h = 300)
 	widgetLayout = cmds.rowColumnLayout(w = 300, h = 300, nc = 2, cw = [(1, 200), (2, 100)])
@@ -10526,7 +10526,7 @@ class SkeletonBuilder_UI():
 	if cmds.window("paintWeightsRemoveInf", exists = True):
 	    cmds.deleteUI("paintWeightsRemoveInf")
 	    
-	window = cmds.window("paintWeightsRemoveInf", t = "Remove Influence", tbm = False, w = 300, h = 300, sizeable = False)
+	window = cmds.window("paintWeightsRemoveInf", t = "Remove Influence", tbm = False, w = 300, h = 300, sizeable = True)
 	
 	mainLayout = cmds.columnLayout(w = 300, h = 300)
 	widgetLayout = cmds.rowColumnLayout(w = 300, h = 300, nc = 2, cw = [(1, 200), (2, 100)])
