OSX port

diff --git a/Engine/Extras/Maya_AnimationRiggingTools/ArtToolsOSX/MayaTools/General/Scripts/P4.py b/Engine/Extras/Maya_AnimationRiggingTools/ArtToolsOSX/MayaTools/General/Scripts/P4.py
new file mode 100644
--- /dev/null
+++ b/Engine/Extras/Maya_AnimationRiggingTools/ArtToolsOSX/MayaTools/General/Scripts/P4.py
@@ -0,0 +1,431 @@
+#!/usr/local/bin/python
+from __future__ import print_function
+"""A Python version of the Perforce "p4" client.
+This uses the Python type P4API.P4Adapter, which is a wrapper for the
+Perforce ClientApi object. 
+$Id: //depot/r10.1/p4-python/P4.py#1 $
+#*******************************************************************************
+# Copyright (c) 2007-2010, Perforce Software, Inc.  All rights reserved.
+# Portions Copyright (c) 1999, Mike Meyer. All rights reserved.
+# Portions Copyright (c) 2004-2007, Robert Cowham. All rights reserved.
+#
+# Redistribution and use in source and binary forms, with or without
+# modification, are permitted provided that the following conditions are met:
+#
+# 1.  Redistributions of source code must retain the above copyright
+# notice, this list of conditions and the following disclaimer.
+#
+# 2.  Redistributions in binary form must reproduce the above copyright
+# notice, this list of conditions and the following disclaimer in the
+# documentation and/or other materials provided with the distribution.
+#
+# THIS SOFTWARE IS PROVIDED BY THE COPYRIGHT HOLDERS AND CONTRIBUTORS "AS IS"
+# AND ANY EXPRESS OR IMPLIED WARRANTIES, INCLUDING, BUT NOT LIMITED TO, THE
+# IMPLIED WARRANTIES OF MERCHANTABILITY AND FITNESS FOR A PARTICULAR PURPOSE
+# ARE DISCLAIMED. IN NO EVENT SHALL PERFORCE SOFTWARE, INC. BE LIABLE FOR ANY
+# DIRECT, INDIRECT, INCIDENTAL, SPECIAL, EXEMPLARY, OR CONSEQUENTIAL DAMAGES
+# (INCLUDING, BUT NOT LIMITED TO, PROCUREMENT OF SUBSTITUTE GOODS OR SERVICES;
+# LOSS OF USE, DATA, OR PROFITS; OR BUSINESS INTERRUPTION) HOWEVER CAUSED AND
+# ON ANY THEORY OF LIABILITY, WHETHER IN CONTRACT, STRICT LIABILITY, OR TORT
+# (INCLUDING NEGLIGENCE OR OTHERWISE) ARISING IN ANY WAY OUT OF THE USE OF THIS
+# SOFTWARE, EVEN IF ADVISED OF THE POSSIBILITY OF SUCH DAMAGE.
+#*******************************************************************************
+Copyright 2007-2010 Perforce Software, Inc. All rights reserved
+License:
+See accompanying LICENSE.txt including for redistribution permission.
+"""
+import sys, os, string, datetime
+import types, re
+from contextlib import contextmanager
+# P4Error - some sort of error occurred
+class P4Exception(Exception):
+    """Exception thrown by a P4 instance in case of an Perforce error or warning"""
+    def __init__(self, value):
+        self.value = value
+    def __str__(self):
+        return str(self.value)
+# The following classes are not used yet.
+# Preparation for a future separation of P4Error and P4Warning
+class P4Error(P4Exception):
+    """Exception thrown by a P4 instance in case of an error"""
+    def __init__(self, value, errors):
+        P4Exception.__init__(self, value)
+        self.errors = errors
+    def __str__(self):
+        return "P4Error : " + str(self.value) + " " + str(self.errors)
+class P4Warning(P4Exception):
+    """Exception thrown by a P4 instance in case of a warning"""
+    def __init__(self, value, warnings):
+        P4Exception.__init__(self, value)
+        self.warnings = warnings
+    def __str__(self):
+        return "P4Warning : " + str(self.value) + " " + str(self.warnings)
+class Spec(dict):
+    """Subclass of dict, representing the fields of a spec definition.
+    Attributes can be accessed either with the conventional dict format,
+    spec['attribute'] or with shorthand spec._attribute.
+    Instances of this class will preventing any unknown keys.
+    """
+    def __init__(self, fieldmap=None):
+        self.__dict__['_Spec__fields'] = fieldmap
+    def permitted_fields(self):
+        return self.__fields
+    def __setitem__(self, key, value):
+        if key in self or self.__fields == None:
+            dict.__setitem__(self, key, value)
+        elif str(key).lower() in self.__fields:
+            dict.__setitem__(self, self.__fields[key.lower()], value)
+        else:
+            raise P4Error("Illegal field " + str(key), [])
+    def __getattr__(self, attr):
+        key = str(attr).lower()
+        if key[0] != '_':
+            raise AttributeError(attr)
+        key = key[1:]
+        if key in self:
+            return self[key]
+        elif key in self.__fields:
+            return self[self.__fields[key]]
+    def __setattr__(self, attr, value):
+        key = str(attr).lower()
+        if key[0] != '_':
+            raise AttributeError(attr)
+        key = key[1:]
+        self[key] = value
+#
+# P4Integration objects hold details about the integrations that have
+# been performed on a particular revision. Used primarily with the
+# P4Revision class
+#
+class Integration:
+    def __init__( self, how, file, srev, erev ):
+        self.how = how
+        self.file = file
+        self.srev = srev
+        self.erev = erev
+    def __repr__(self):
+        return "Integration (how = %s file = %s srev = %s erev = %s)" % (self.how, self.file, self.srev, self.erev)
+#
+# Each P4Revision object holds details about a particular revision
+# of a file. It may also contain the history of any integrations 
+# to/from the file
+#
+class Revision:
+    def __init__( self, depotFile ):
+        self.depotFile = depotFile
+        self.integrations = []
+        self.rev = None
+        self.change = None
+        self.action = None
+        self.type = None
+        self.time = None
+        self.user = None
+        self.client = None
+        self.desc = None
+        self.digest = None
+        self.fileSize = None
+    def integration( self, how, file, srev, erev ):
+        rec = Integration( how, file, srev, erev )
+        self.integrations.append( rec )
+        return rec
+    # iterator over the collection calling a provided function
+    # Python's poor man version of the Ruby block 
+    def each_integration(self):
+        for i in self.integrations:
+            yield i
+    def __repr__(self):
+        return "Revision (depotFile = %s rev = %s change = %s action = %s type = %s time = %s user = %s client = %s)" % \
+               (self.depotFile, self.rev, self.change, self.action, self.type, self.time, self.user, self.client)
+#
+# Each DepotFile entry contains details about one depot file. 
+# 
+class DepotFile:
+    def __init__( self, name ):
+        self.depotFile = name
+        self.revisions = []
+    def new_revision(self):
+        r = Revision( self.depotFile )
+        self.revisions.append( r )
+        return r
+    def each_revision(self):
+        for r in self.revisions:
+            yield r
+    def __repr__(self):
+        return "DepotFile (depotFile = %s, %s revisions)" % ( self.depotFile, len( self.revisions ) )
+#
+# Resolver class used in p4.run_resolve()
+#
+# The default simply checks that p4.input is set to sensible value
+# This class is meant to be subclassed for a custom resolver and
+# Resolver.resolve() overriden
+#
+class Resolver:
+    def __init__(self):
+        pass
+    def resolve(self, mergeInfo):
+        if mergeInfo.merge_hint == "e":
+            print("Standard resolver encountered merge conflict, skipping resolve")
+            return "s"
+        else:
+            return mergeInfo.merge_hint
+# This is where the C/C++ shared library is loaded
+# It has to be in this place because the library needs to access
+# the classes defined above. Accessing classes defined below this
+# entry would cause an endless loop
+import P4API
+class P4(P4API.P4Adapter):
+    """Use this class to communicate with a Perforce server
+    Instances of P4 will use the environment settings (including P4CONFIG)
+    to determine the connection parameters such as P4CLIENT and P4PORT.
+    This attributes can also be set separately before connecting.
+    To run any Perforce commands, users of this class first need to run
+    the connect() method.
+    It is good practice to disconnect() after the program is complete.
+    """
+    # Constants useful for exception_level
+    # RAISE_ALL:     Errors and Warnings are raised as exceptions (default)
+    # RAISE_ERROR:   Only Errors are raised as exceptions
+    # RAISE_NONE:    No exceptions are raised, instead False is returned
+    RAISE_ALL = 2
+    RAISE_ERROR = 1
+    RAISE_ERRORS = 1
+    RAISE_NONE = 0
+    def __init__(self, *args, **kwlist): 
+        P4API.P4Adapter.__init__(self, *args, **kwlist)
+    def __del__(self):
+        if self.debug > 3:
+            print("P4.__del__()", file=sys.stderr)
+    # store the references to the created lambdas as a weakref to allow Python 
+    # to clean up the garbage. |The lambda as a closure stores a reference to self
+    # which causes a circular reference problem without the weakref
+    def __getattr__(self, name):
+        if name.startswith("run_"):
+            cmd = name[len("run_"):]
+            return lambda *args: self.run(cmd, *args)
+        elif name.startswith("delete_"):
+            cmd = name[len("delete_"):]
+            return lambda *args: self.run(cmd, "-d", *args)
+        elif name.startswith("fetch_"):
+            cmd = name[len("fetch_"):]
+            return lambda *args: self.run(cmd, "-o", *args)[0]
+        elif name.startswith("save_"):
+            cmd = name[len("save_"):]
+            return lambda *args: self.__save(cmd, *args)
+        elif name.startswith("parse_"):
+            cmd = name[len("parse_"):]
+            return lambda *args: self.parse_spec(cmd, *args)
+        elif name.startswith("format_"):
+            cmd = name[len("format_"):]
+            return lambda *args: self.format_spec(cmd, *args)
+        else:
+            raise AttributeError(name)
+    def __save(self, cmd, *args):
+        self.input = args[0]
+        return self.run(cmd, "-i", args[1:])
+    def __repr__(self):
+        state = "disconnected"
+        if self.connected():
+            state = "connected"
+        return "P4 [%s@%s %s] %s" % \
+               (self.user, self.client, self.port, state)
+    def identify(cls):
+        return P4API.identify()
+    identify = classmethod(identify)
+    def run(self, *args):
+        "Generic run method"
+        return P4API.P4Adapter.run(self, *self.__flatten(args))
+    def run_submit(self, *args):
+        "Simplified submit - if any arguments is a dict, assume it to be the changeform"
+        nargs = list(args)
+        form = None
+        for n, arg in enumerate(nargs):
+            if isinstance( arg, dict):
+                self.input = arg
+                nargs.pop(n)
+                nargs.append("-i")
+                break
+        return self.run("submit", *nargs)
+    def run_login(self, *args):
+        "Simple interface to make login easier"
+        self.input = self.password
+        return self.run("login", *args)
+    def run_password( self, oldpass, newpass ):
+        "Simple interface to allow setting of the password"
+        if( oldpass and len(oldpass) > 0 ):
+            self.input = [ oldpass, newpass, newpass ]
+        else:
+            self.input = [ newpass, newpass ]
+        return self.run( "password" )
+    #
+    # run_filelog: convert "p4 filelog" responses into objects with useful
+    #              methods
+    #
+    # Requires tagged output to be of any real use. If tagged output it not 
+    # enabled then you just get the raw data back
+    #
+    def run_filelog( self, *args ):
+        raw = self.run( 'filelog', args )
+        if (not self.tagged): 
+            # untagged mode returns simple strings, which breaks the code below
+            return raw
+        result = []
+        for h in raw:
+            r = None
+            if isinstance( h, dict ):
+                df = DepotFile( h[ "depotFile" ] )
+                for n, rev in enumerate( h[ "rev" ]):
+                    # Create a new revision of this file ready for populating
+                    r = df.new_revision()
+                    # Populate the base attributes of each revision
+                    r.rev = int( rev )
+                    r.change = int( h[ "change" ][ n ] )
+                    r.action = h[ "action" ][ n ] 
+                    r.type = h[ "type" ][ n ]
+                    r.time = datetime.datetime.utcfromtimestamp( int( h[ "time" ][ n ]) )
+                    r.user = h[ "user" ][ n ]
+                    r.client = h[ "client" ][ n ]
+                    r.desc = h[ "desc" ][ n ]
+                    if "digest" in h:
+                        r.digest = h[ "digest" ][ n ]
+                    if "fileSize" in h and n < len(h[ "fileSize"]):
+                        r.fileSize = h[ "fileSize" ][ n ]
+                    # Now if there are any integration records for this revision,
+                    # add them in too
+                    if (not "how" in h) or (n >= len(h["how"]) or h["how"][n] == None):
+                        continue
+                    else:
+                        for m, how in enumerate( h[ "how" ][ n ] ):
+                            file = h[ "file" ][ n ][ m ]
+                            srev = h[ "srev" ][ n ][ m ].lstrip('#')
+                            erev = h[ "erev" ][ n ][ m ].lstrip('#')
+                            if srev == "none":
+                                srev = 0
+                            else:
+                                srev = int( srev )
+                            if erev == "none":
+                                erev = 0
+                            else:
+                                erev = int( erev )
+                            r.integration( how, file, srev, erev )
+            else:
+                r = h
+            result.append( df )
+        return result
+    def run_print(self, *args):
+        raw = self.run('print', args)
+        result = []
+        for line in raw:
+            if isinstance(line, dict):
+                result.append(line)
+                result.append("")
+            else:
+                result[-1] += line
+        return result
+    def run_resolve(self, *args, **kargs):
+        myResolver = Resolver()
+        if "resolver" in kargs: 
+            myResolver = kargs["resolver"]
+        savedResolver = self.resolver
+        self.resolver = myResolver
+        result = self.run("resolve", args)
+        self.resolver = savedResolver
+        return result
+    def __flatten(self, args):
+        result = []
+        if isinstance(args, tuple) or isinstance(args, list):
+            for i in args:
+                result.extend(self.__flatten(i))
+        else:
+            result.append(args)
+        return tuple(result)
+    def __enter__( self ):
+        return self
+    def __exit__( self, exc_type, exc_val, exc_tb ):
+        if self.connected():
+            self.disconnect()
+        return True
+    def connect( self ):
+        P4API.P4Adapter.connect( self )
+        return self
+    @contextmanager
+    def while_tagged( self, t ):
+        old = self.tagged
+        self.tagged = t
+        yield
+        self.tagged = old
+    @contextmanager
+    def at_exception_level( self, e ):
+        old = self.exception_level
+        self.exception_level = e
+        yield
+        self.exception_level = old
+    @contextmanager
+    def saved_context( self , **kargs):
+        """Saves the context of this p4 object and restores it again at the end of the block"""
+        saved_context = {}
+        for attr in self.__members__:
+            saved_context[attr] = getattr(self, attr)
+        for (k,v) in list(kargs.items()):
+            setattr( self, k, v)
+        yield
+        # now restore the context again. Ignore AttributeError exception
+        # Exception is expected because some attributes only have getters, no setters
+        for (k,v) in list(saved_context.items()):
+            try:
+                setattr( self, k, v )
+            except AttributeError:
+                pass # expected for server_level and p4config_file
+class Map(P4API.P4Map):
+    def __init__(self, *args):
+        P4API.P4Map.__init__(self, *args)
+        if len(args) > 0:
+            self.insert( *args )
+    LEFT2RIGHT = True
+    RIGHT2LEFT = False
+    def __str__( self ):
+        result = ""
+        for a in self.as_array():
+            result += a + "\n"
+        return result
+    def is_empty(self):
+        """Returns True if this map has no entries yet, otherwise False"""
+        return self.count() == 0
+    def includes(self, *args):
+        return self.translate(*args) != None
+    def insert(self, *args):
+        """Insert an argument to the map. The argument can be:
+        A String, 
+            Either of the form "[+-]//lhs/... //rhs/..." or "[+-]//lhs/..." 
+            for label style maps.
+        A List:
+            This is a list of strings of one of the single string formats 
+            described above.
+        A pair of Strings:
+            P4.Map.insert(lhs, rhs)
+        """
+        if len(args) == 1 :
+            arg = args[0]
+            if isinstance( arg, str ):
+                P4API.P4Map.insert( self, arg )
+            elif isinstance( arg, list ):
+                for s in arg:
+                    P4API.P4Map.insert( self, s )
+        else: # expecting 2 args in this case: left, right 
+            left = args[0].strip()
+            right = args[1].strip()
+            P4API.P4Map.insert(self, left, right )
+if __name__ == "__main__":
+    p4 = P4()
+    p4.connect()
+    try:
+        ret = p4.run(sys.argv[1:])
+        for line in ret:
+            if isinstance(line, dict):
+                print("-----")
+                for k in list(line.keys()):
+                    print(k, "=", line[k])
+            else:
+                print(line)
+    except:
+        for e in p4.errors:
+            print(e)
diff --git a/Engine/Extras/Maya_AnimationRiggingTools/ArtToolsOSX/MayaTools/General/Scripts/mayaSetup.py b/Engine/Extras/Maya_AnimationRiggingTools/ArtToolsOSX/MayaTools/General/Scripts/mayaSetup.py
new file mode 100644
--- /dev/null
+++ b/Engine/Extras/Maya_AnimationRiggingTools/ArtToolsOSX/MayaTools/General/Scripts/mayaSetup.py
@@ -0,0 +1 @@
+import maya.cmds as cmdsimport osimport sysfrom functools import partialimport maya.mel as meldef setupTools():    path = os.path.join(cmds.internalVar(usd = True), "mayaTools.txt")    f = open(path, 'r')    mayaToolsDir = f.readline()    path = os.path.join(mayaToolsDir, "General", "Scripts")    pluginPath = os.path.join(mayaToolsDir + "General", "Plugins")    #look in sys.path to see if path is in sys.path. if not, add it    if not path in sys.path:        sys.path.append(path)    #make sure MAYA_PLUG_IN_PATH has our plugin path    pluginPaths = os.environ["MAYA_PLUG_IN_PATH"]    pluginPaths = pluginPaths + ";" + pluginPath    os.environ["MAYA_PLUG_IN_PATH"] = pluginPaths    #setup menu item in main window    import customMayaMenu as cmm    cmm.customMayaMenu()setupTools()
\ No newline at end of file
diff --git a/Engine/Extras/Maya_AnimationRiggingTools/ArtToolsOSX/MayaTools/General/Scripts/setupScene.py b/Engine/Extras/Maya_AnimationRiggingTools/ArtToolsOSX/MayaTools/General/Scripts/setupScene.py
new file mode 100644
--- /dev/null
+++ b/Engine/Extras/Maya_AnimationRiggingTools/ArtToolsOSX/MayaTools/General/Scripts/setupScene.py
@@ -0,0 +1 @@
+import maya.cmds as cmdsdef setupScene():    #Setup scene to set timeline to 0-100 with current frame being 0 and the fps to be 30    cmds.currentUnit(time = 'ntsc')    cmds.playbackOptions(min = 0, max = 100, animationStartTime = 0, animationEndTime = 100)    cmds.currentTime(0)setupScene()
\ No newline at end of file
