[INTEGRATE] Content from 4.8 branch.

diff --git a/Engine/Extras/Maya_AnimationRiggingTools/MayaTools/General/Scripts/Modules/ART_CreateHydraulicRigs.py b/Engine/Extras/Maya_AnimationRiggingTools/MayaTools/General/Scripts/Modules/ART_CreateHydraulicRigs.py
new file mode 100644
--- /dev/null
+++ b/Engine/Extras/Maya_AnimationRiggingTools/MayaTools/General/Scripts/Modules/ART_CreateHydraulicRigs.py
@@ -0,0 +1,53 @@
+import maya.cmds as cmds
+
+# This tool only works on "leaf" joints created with A.R.T 
+# To use, select 3 objects, the start of your hydraulic system, the end, and the upvector.  It is assumed that you are selecting the _anim controls for the first 2
+# It is assumed that the upvector is 1 by default unless passed in otherwise.  Once set, you can always change this by simply manually reversing the value in the aimConstraint
+# It is also assumed that you are using X as the axis that points at each other and that the first control points at the 2nd, but the 2nd points away from the first.
+# Assuming Z as the up vector axis.
+
+def CreateHydraulicRig(upDir):
+
+    selection = cmds.ls(sl=True)
+    hyd_start_01_anim = selection[0]
+    hyd_start_02_anim = selection[1]
+    upObject = selection[2]
+    
+    hyd_start_01 = hyd_start_01_anim.rpartition("_")[0]
+    hyd_start_02 = hyd_start_02_anim.rpartition("_")[0]
+
+    # Create a rig for the lower arm extra hydraulic piston.
+    cmds.delete("driver_"+hyd_start_01+"_parentConstraint1")
+    robo_lowarm_pistonrod_01_anim_sub = cmds.group(em=True, name=hyd_start_01_anim+"_sub", parent=hyd_start_01_anim)
+    const = cmds.parentConstraint(hyd_start_01_anim, robo_lowarm_pistonrod_01_anim_sub, weight=1, mo=False)
+    cmds.delete(const)
+    const = cmds.parentConstraint(robo_lowarm_pistonrod_01_anim_sub, "driver_"+hyd_start_01, weight=1, mo=True)
+    
+    cmds.delete("driver_"+hyd_start_02+"_parentConstraint1")
+    robo_lowarm_pistonrod_02_anim_sub = cmds.group(em=True, name=hyd_start_02_anim+"_sub", parent=hyd_start_02_anim)
+    const = cmds.parentConstraint(hyd_start_02_anim, robo_lowarm_pistonrod_02_anim_sub, weight=1, mo=False)
+    cmds.delete(const)
+    const = cmds.parentConstraint(robo_lowarm_pistonrod_02_anim_sub, "driver_"+hyd_start_02, weight=1, mo=True)
+
+    # Hook up the hydraulics for the lowerarm piston.
+    const = cmds.aimConstraint(robo_lowarm_pistonrod_01_anim_sub, robo_lowarm_pistonrod_02_anim_sub, weight=1, mo=False, aimVector=(-1, 0, 0), upVector=(0, 0, upDir), worldUpType="object", worldUpVector=(0, 0, -1), worldUpObject=upObject)
+    const = cmds.aimConstraint(robo_lowarm_pistonrod_02_anim_sub, robo_lowarm_pistonrod_01_anim_sub, weight=1, mo=False, aimVector=(1, 0, 0), upVector=(0, 0, upDir), worldUpType="object", worldUpVector=(0, 0, -1), worldUpObject=upObject)
+
+
+def CreateHydraulicSlideRig(*args):
+    print "test"
+    # Create a slide mechanism on the forearm part of the lowerarm piston.
+    # CRA NOTE: this will come at a later date.
+    '''
+    cmds.addAttr("robo_lowarm_pistonrod_02_anim", ln = "Slide", dv = 0, min=0, max=10, keyable = True)
+    cmds.setDrivenKeyframe("robo_lowarm_pistonrod_02_anim_grp.translateX", cd="robo_lowarm_pistonrod_02_anim.Slide", itt="spline", ott="spline")
+    cmds.setDrivenKeyframe("robo_lowarm_pistonrod_02_anim_grp.translateZ", cd="robo_lowarm_pistonrod_02_anim.Slide", itt="spline", ott="spline")
+
+    cmds.setAttr("robo_lowarm_pistonrod_02_anim.Slide", 10)
+    cmds.setAttr("robo_lowarm_pistonrod_02_anim_grp.tx", -6.662443)
+    cmds.setAttr("robo_lowarm_pistonrod_02_anim_grp.tz", -5.510393)
+
+    cmds.setDrivenKeyframe("robo_lowarm_pistonrod_02_anim_grp.translateX", cd="robo_lowarm_pistonrod_02_anim.Slide", itt="spline", ott="spline")
+    cmds.setDrivenKeyframe("robo_lowarm_pistonrod_02_anim_grp.translateZ", cd="robo_lowarm_pistonrod_02_anim.Slide", itt="spline", ott="spline")
+    cmds.setAttr("robo_lowarm_pistonrod_02_anim.Slide", 0)
+    '''
diff --git a/Engine/Extras/Maya_AnimationRiggingTools/MayaTools/General/Scripts/Modules/ART_ForearmTwist.py b/Engine/Extras/Maya_AnimationRiggingTools/MayaTools/General/Scripts/Modules/ART_ForearmTwist.py
new file mode 100644
--- /dev/null
+++ b/Engine/Extras/Maya_AnimationRiggingTools/MayaTools/General/Scripts/Modules/ART_ForearmTwist.py
@@ -0,0 +1,145 @@
+import maya.cmds as cmds
+import maya.mel as mel
+import os
+import ART_rigUtils as utils
+reload(utils)
+
+#----------------------------------------------------------------------------------------
+def forearmTwistRig(color, rollGrpParent, fkArmJoint, ikArmJoint, suffix, name, prefix, lowerArm):
+    print "START Building the "+ suffix +" lower arm twists---------------------------------------"
+    
+    print "suffix = "+suffix
+    print "name = "+name
+    print "prefix = "+prefix
+    print "lowerArm = "+lowerArm
+    
+    # create a nice name
+    name = prefix + name + suffix
+    if name.find("_") == 0:
+        name = name.partition("_")[2]
+        lowerarm = prefix + lowerArm + suffix
+        driver_lowerarm = "driver_" + prefix + lowerArm + suffix
+        driver_hand = "driver_hand_"+name
+
+        numRolls = 0
+        for joint in ["_twist_01", "_twist_02", "_twist_03", "_twist_04", "_twist_05", "_twist_06"]:
+            if cmds.objExists("driver_" + prefix + lowerArm + joint + suffix):
+                numRolls = numRolls + 1
+
+        print "...There are a total of " + str(numRolls) + " to build."
+            
+        for i in range(int(numRolls)):
+            print "...Building lower arm twist_0" + str(i + 1)
+            driver_lowerarm_twist = "driver_" + prefix + lowerArm + "_twist_0" + str(i + 1) + suffix
+            
+            # create our roll group
+            rollGrp = cmds.group(empty = True, name = lowerarm + "_roll_grp_0" + str(i + 1))
+            cmds.parent(rollGrp, "arm_sys_grp")
+
+            # Make sure the twist joint is exactly at the hand and aimed at the lowerarm
+            #NOTE if the "side" method Jeremy is using in here(name) is intended to handle more than just "l" and "r" then the move part of this must be re-written to handle that accordingly.
+            '''if i == 0:
+                const = cmds.parentConstraint(driver_hand, driver_lowerarm_twist, weight=1, mo=False, sr=["x","y","z"])
+                cmds.delete(const)
+                cmds.select(driver_lowerarm_twist, r=True)
+                if name == "l":
+                    cmds.move(-0.995369, 0, 0, driver_lowerarm_twist, r = True)
+                else:
+                    cmds.move(0.995369, 0, 0, driver_lowerarm_twist, r = True)
+                cmds.makeIdentity(driver_lowerarm_twist, r = 1, t = 1, s = 1, apply =True)'''
+
+            # create a new heirarchy out of the driver skeleton
+            driver_lowerarm_offset = cmds.duplicate(driver_lowerarm, po=True, name=driver_lowerarm+"_Offset_0" + str(i + 1))[0]
+            cmds.parent(driver_lowerarm_offset, rollGrp)
+            cmds.setAttr(driver_lowerarm_offset + ".rotateOrder", 1)
+            cmds.setAttr(driver_lowerarm_offset + ".v", 1)
+
+            driver_lowerarm_twist_offset = cmds.duplicate(driver_lowerarm_twist, po=True, name=driver_lowerarm_twist+"_Offset")[0]
+            cmds.parent(driver_lowerarm_twist_offset, driver_lowerarm_offset)
+            cmds.setAttr(driver_lowerarm_twist_offset + ".rotateOrder", 1)
+            cmds.setAttr(driver_lowerarm_twist_offset + ".v", 1)
+
+            driver_hand_offset = cmds.duplicate(driver_lowerarm_twist_offset, po=True, name=driver_hand+"_Offset_0" + str(i + 1))[0]
+            const = cmds.parentConstraint(driver_hand, driver_hand_offset, weight=1, mo=False, sr=["x","y","z"])
+            cmds.delete(const)
+            cmds.setAttr(driver_hand_offset + ".rotateOrder", 1)
+            cmds.setAttr(driver_hand_offset + ".v", 1)
+
+            # create the manual twist control
+            if i == 0:
+                twistCtrl = utils.createControl("circle", 15, lowerarm + "_twist_anim")
+            else:
+                twistCtrl = utils.createControl("circle", 15, lowerarm + "_twist"+str(i + 1)+"_anim")
+            cmds.setAttr(twistCtrl + ".ry", -90)
+            cmds.setAttr(twistCtrl + ".sx", 0.8)
+            cmds.setAttr(twistCtrl + ".sy", 0.8)
+            cmds.setAttr(twistCtrl + ".sz", 0.8)
+            cmds.makeIdentity(twistCtrl, r = 1, s = 1, apply =True)
+            
+            # move the manual control to the correct location
+            constraint = cmds.parentConstraint(driver_lowerarm_twist_offset, twistCtrl)[0]
+            cmds.delete(constraint)
+            
+            # create a group for the manual control and parent the twist to it.
+            twistCtrlGrp = cmds.group(empty = True, name = twistCtrl + "_grp")
+            constraint = cmds.parentConstraint(driver_lowerarm_twist_offset, twistCtrlGrp)[0]
+            cmds.delete(constraint)
+            
+            cmds.parent(twistCtrl, twistCtrlGrp)
+            cmds.parent(twistCtrlGrp, rollGrp)
+            cmds.makeIdentity(twistCtrl, t = 1, r = 1, s = 1, apply = True)
+            cmds.parentConstraint(driver_lowerarm_twist_offset, twistCtrlGrp)
+
+            # set the manual controls visibility settings
+            cmds.setAttr(twistCtrl + ".overrideEnabled", 1)
+            cmds.setAttr(twistCtrl + ".overrideColor", color)
+            for attr in [".sx", ".sy", ".sz"]:
+                cmds.setAttr(twistCtrl + attr, lock = True, keyable = False)
+            cmds.setAttr(twistCtrl + ".v", keyable = False)	
+        
+            # add attr on rig settings for manual twist control visibility
+            cmds.select("Rig_Settings")
+            if i == 0:
+                cmds.addAttr(longName=(name + "twistCtrlVisLower"), at = 'bool', dv = 0, keyable = True)
+            cmds.connectAttr("Rig_Settings." + name + "twistCtrlVisLower", twistCtrl + ".v")
+            
+            # add attr to rig settings for the twist ammount values
+            cmds.select("Rig_Settings")
+            cmds.addAttr(longName= ( name + "ForearmTwist"+str(i + 1)+"Amount" ), defaultValue=0.5, minValue=0, maxValue=1, keyable = True)
+            for u in range(int(i+1)):
+                cmds.setAttr("Rig_Settings." + name + "ForearmTwist"+str(u + 1)+"Amount", (1.0/(i+2.0)*((2.0-u)+(i-1.0))) )
+            
+            # Create a locator and group for the up vector of the aim for the twist.
+            forearmTwistUp = cmds.spaceLocator(n="forearmTwistUp_0"+str(i + 1)+"_"+name)[0]
+            constraint = cmds.parentConstraint(driver_hand_offset, forearmTwistUp)[0]
+            cmds.delete(constraint)
+            forearmTwistUpGrp = cmds.duplicate(driver_hand_offset, po=True, name=forearmTwistUp+"_grp")[0]
+            cmds.parent(forearmTwistUpGrp, rollGrp)
+            cmds.parent(forearmTwistUp, forearmTwistUpGrp)
+            cmds.setAttr(forearmTwistUpGrp + ".rotateOrder", 1)
+            cmds.setAttr(forearmTwistUp + ".v", 0)	
+
+            # constrain the up group into the rig so that it takes the correct ammount of the twist
+            forearmGrpPrntConst = cmds.parentConstraint(driver_lowerarm_offset, driver_hand_offset, forearmTwistUpGrp, mo=True)[0]
+            cmds.setAttr(forearmGrpPrntConst+".interpType", 2)
+            cmds.move(0, 0, 6, forearmTwistUp, r = True)
+
+            twistAimConst = cmds.aimConstraint(driver_hand_offset, driver_lowerarm_twist_offset, weight=1, mo=True, aimVector=(1, 0, 0), upVector=(0, 0, 1), worldUpType="object", worldUpObject=forearmTwistUp)[0]
+
+            # constrain the offsets back to the original driver skeleton
+            cmds.parentConstraint(driver_hand, driver_hand_offset, mo=True)
+            cmds.parentConstraint(driver_lowerarm, driver_lowerarm_offset, mo=True)
+
+            # hook up the nodes to drive the twist based on the values on the rig settings node.
+            reverseNode = cmds.shadingNode("reverse", asUtility=True, name=forearmGrpPrntConst+"_reverse")
+            cmds.connectAttr("Rig_Settings."+name+"ForearmTwist"+str(i + 1)+"Amount", forearmGrpPrntConst+"."+driver_hand_offset+"W1", f=True)
+            cmds.connectAttr("Rig_Settings."+name+"ForearmTwist"+str(i + 1)+"Amount", reverseNode+".inputX", f=True)
+            cmds.connectAttr(reverseNode+".outputX", forearmGrpPrntConst+"."+driver_lowerarm_offset+"W0", f=True)
+            
+            #constrain driver joint to twist joint
+            cmds.parentConstraint(twistCtrl, driver_lowerarm_twist,  mo = True)	
+            
+    print ".END Building the "+ suffix +" lower arm twists---------------------------------------"
+    
+#forearmTwistRig(0, "", "", "", "_l", "", "", "lowerarm")
+    
\ No newline at end of file
diff --git a/Engine/Extras/Maya_AnimationRiggingTools/MayaTools/General/Scripts/Modules/ART_UpperArmIKTwist.py b/Engine/Extras/Maya_AnimationRiggingTools/MayaTools/General/Scripts/Modules/ART_UpperArmIKTwist.py
new file mode 100644
--- /dev/null
+++ b/Engine/Extras/Maya_AnimationRiggingTools/MayaTools/General/Scripts/Modules/ART_UpperArmIKTwist.py
@@ -0,0 +1,136 @@
+import maya.cmds as cmds
+import Modules.ART_rigUtils as utils
+reload(utils)
+
+#----------------------------------------------------------------------------------------
+# Currently this script is written to work with multiple twists, but if you have more than one the last one is the only one that will work.
+def upperIKTwist(color, suffix, name, prefix, upperJnt, lowerJnt, sysGroup):
+    print "START Building the "+suffix+" "+upperJnt+" UPPER IK twists---------------------------------------"
+
+    #create a nice name
+    name = prefix + name + suffix
+    if name.find("_") == 0:
+        name = name.partition("_")[2]
+        
+        # define the joints for this rig.
+        upperjnt = prefix + upperJnt + suffix
+        lowerjnt = prefix + lowerJnt + suffix
+        #OLD - upperjnt_parent = "driver_" + prefix + "clavicle" + suffix
+        upperjnt_parent = cmds.listRelatives(upperjnt, p=True)[0]
+        driver_upperjnt = "driver_" + prefix + upperJnt + suffix
+        driver_lowerjnt = "driver_" + prefix + lowerJnt + suffix
+
+        numRolls = 0
+        for joint in ["_twist_01", "_twist_02", "_twist_03", "_twist_04", "_twist_05", "_twist_06"]:
+            if cmds.objExists("driver_" + prefix + upperJnt + joint + suffix):
+                numRolls = numRolls + 1
+
+        print "...There are a total of " + str(numRolls) + " to build."
+
+        for i in range(int(numRolls)):
+            print "...Building upper twist_0" + str(i + 1)
+
+            upperjnt_twist = prefix+upperJnt+"_twist_0"+str(i+1)+suffix
+            driver_upperjnt_twist = "driver_"+prefix+upperJnt+"_twist_0"+str(i+1)+suffix
+
+            ####################################
+            # Create a master group for the roll controls to be parented to.
+            rollGrp = cmds.group(empty=True, name=upperjnt+"_roll_grp_0"+str(i+1))
+            cmds.parent(rollGrp, sysGroup)
+            cmds.parentConstraint(upperjnt_parent, rollGrp, mo=True)
+
+            # Create the manual twist control
+            if i == 0:
+                twistCtrl = utils.createControl("circle", 15, upperjnt + "_twist_anim")
+            else:
+                twistCtrl = utils.createControl("circle", 15, upperjnt + "_twist"+str(i+1)+"_anim")
+            cmds.setAttr(twistCtrl + ".ry", -90)
+            cmds.setAttr(twistCtrl + ".sx", 0.8)
+            cmds.setAttr(twistCtrl + ".sy", 0.8)
+            cmds.setAttr(twistCtrl + ".sz", 0.8)
+            cmds.makeIdentity(twistCtrl, r = 1, s = 1, apply =True)
+
+            # move the manual control to the correct location
+            constraint = cmds.parentConstraint(driver_upperjnt_twist, twistCtrl)[0]
+            cmds.delete(constraint)
+
+            # create a group for the manual control and parent the twist to it.
+            twistCtrlGrp = cmds.group(empty = True, name=twistCtrl+"_grp")
+            constraint = cmds.parentConstraint(driver_upperjnt_twist, twistCtrlGrp, mo=False)[0]
+            cmds.delete(constraint)
+            twistCtrlDeltoidOffsetGrp = cmds.group(empty = True, name=twistCtrl+"_deltOffset_grp")
+            constraint = cmds.parentConstraint(driver_upperjnt_twist, twistCtrlDeltoidOffsetGrp, mo=False)[0]
+            cmds.delete(constraint)
+            cmds.parent(twistCtrl, twistCtrlDeltoidOffsetGrp)
+            cmds.parent(twistCtrlDeltoidOffsetGrp, twistCtrlGrp)
+            cmds.parent(twistCtrlGrp, rollGrp)
+            cmds.makeIdentity(twistCtrl, t = 1, r = 1, s = 1, apply = True)
+
+            # set the manual controls visibility settings
+            cmds.setAttr(twistCtrl + ".overrideEnabled", 1)
+            cmds.setAttr(twistCtrl + ".overrideColor", color)
+            for attr in [".sx", ".sy", ".sz"]:
+                cmds.setAttr(twistCtrl + attr, lock = True, keyable = False)
+            cmds.setAttr(twistCtrl + ".v", keyable = False)	
+
+            # add attr on rig settings for manual twist control visibility
+            cmds.select("Rig_Settings")
+            if i == 0:
+                cmds.addAttr(longName=(upperjnt + "_twistCtrlVisUpper"), at = 'bool', dv = 0, keyable = True)
+            cmds.connectAttr("Rig_Settings." + upperjnt + "_twistCtrlVisUpper", twistCtrl + ".v")
+
+            # add attr to rig settings for the twist ammount values
+            cmds.select("Rig_Settings")
+            cmds.addAttr(longName= (upperjnt+"_Twist"+str(i+1)+"Amount"), defaultValue=0.5, minValue=0, maxValue=1, keyable = True)
+            for u in range(int(i+1)):
+                cmds.setAttr("Rig_Settings."+upperjnt+"_Twist"+str(u+1)+"Amount", (1.0/(i+2.0)*((2.0-u)+(i-1.0))) )
+
+            ####################################
+            # Create the twist rig on the shoulder without any extra controls.  Those are added later.
+            cmds.delete(upperjnt_twist+"_orientConstraint1")
+
+            twistIKHandle = cmds.ikHandle(name="twistIKHandle_0"+str(i+1)+"_"+upperjnt, sol="ikRPsolver", sj=upperjnt, ee=upperjnt_twist)[0]
+            print "TWIST IK HANDLE IS: "+twistIKHandle
+            cmds.setAttr(twistIKHandle+".poleVectorX", 0)
+            cmds.setAttr(twistIKHandle+".poleVectorY", 0)
+            cmds.setAttr(twistIKHandle+".poleVectorZ", 0)
+            cmds.parent(twistIKHandle, driver_upperjnt)
+
+            driver_upperjnt_offset = cmds.duplicate(driver_upperjnt, po=True, name="driver_"+upperJnt+"_twistOffset_l")[0]
+            cmds.parent(driver_upperjnt_offset, driver_upperjnt)
+
+            twistOrntConst = cmds.orientConstraint(driver_upperjnt, upperjnt, driver_upperjnt_offset, mo=True)[0]
+            cmds.setAttr(twistOrntConst+".interpType", 2)
+
+            twistMDNode = cmds.shadingNode("multiplyDivide", asUtility=True, name=driver_upperjnt_twist+"_MDnode")
+            cmds.setAttr(twistMDNode+".input2X", 1)
+            cmds.setAttr(twistMDNode+".input2Y", -1)
+
+            cmds.connectAttr(driver_upperjnt_offset+".rotateX", twistMDNode+".input1X", force=True)
+            cmds.connectAttr(driver_upperjnt+".rotateX", twistMDNode+".input1Y", force=True)
+            cmds.connectAttr(twistMDNode+".outputX", driver_upperjnt_twist+".rx", force=True)            
+            cmds.connectAttr(twistMDNode+".outputY", twistCtrlDeltoidOffsetGrp+".rx", force=True)            
+
+            cmds.pointConstraint(driver_upperjnt_twist, twistCtrlGrp, mo=True)
+            cmds.orientConstraint(driver_upperjnt_twist, twistCtrlGrp, mo=True)
+
+            cmds.orientConstraint(twistCtrl, upperjnt_twist, mo=True)
+            
+            if cmds.objExists("deltoid_anim"+suffix):
+                if suffix == "_l":
+                    cmds.connectAttr("deltoid_anim"+suffix+".rotateX", twistIKHandle+".twist", force=True)
+                else:
+                    deltMDNode = cmds.shadingNode("multiplyDivide", asUtility=True, name=twistIKHandle+"_MDNode")
+                    cmds.setAttr(deltMDNode+".input2X", -1)
+                    cmds.connectAttr("deltoid_anim"+suffix+".rotateX", deltMDNode+".input1X", force=True)
+                    cmds.connectAttr(deltMDNode+".outputX", twistIKHandle+".twist", force=True)
+                    
+            reverseMDNode = cmds.shadingNode("reverse", asUtility=True, name=driver_upperjnt_offset+"_Reversenode")
+            cmds.connectAttr("Rig_Settings."+upperjnt+"_Twist"+str(i+1)+"Amount", reverseMDNode+".inputX", force=True)
+            cmds.connectAttr(reverseMDNode+".outputX", twistOrntConst+"."+upperjnt+"W1", force=True)            
+            cmds.connectAttr("Rig_Settings."+upperjnt+"_Twist"+str(i+1)+"Amount", twistOrntConst+"."+driver_upperjnt+"W0", force=True)
+                          
+
+    print ".END Building the "+suffix+" "+upperJnt+" UPPER twists---------------------------------------"            
+
+#upperIKTwist(6, "_l", "", "", "upperarm", "lowerarm", "arm_sys_grp")
diff --git a/Engine/Extras/Maya_AnimationRiggingTools/MayaTools/General/Scripts/Modules/ART_UpperArmTwist.py b/Engine/Extras/Maya_AnimationRiggingTools/MayaTools/General/Scripts/Modules/ART_UpperArmTwist.py
new file mode 100644
--- /dev/null
+++ b/Engine/Extras/Maya_AnimationRiggingTools/MayaTools/General/Scripts/Modules/ART_UpperArmTwist.py
@@ -0,0 +1,176 @@
+import maya.cmds as cmds
+import maya.mel as mel
+import os
+import ART_rigUtils as utils
+reload(utils)
+
+#----------------------------------------------------------------------------------------
+# Currently this script is written to work with multiple twists, but if you have more than one the last one is the only one that will work.
+def upperArmTwist(color, rollGrpParent, fkArmJoint, ikArmJoint, suffix, name, prefix, upperArm, lowerArm):
+    print "START Building the "+ suffix +" UPPER arm twists---------------------------------------"
+
+    #create a nice name
+    name = prefix + name + suffix
+    if name.find("_") == 0:
+        name = name.partition("_")[2]
+        
+        # define the joints for this rig.
+        upperarm = prefix + upperArm + suffix
+        lowerarm = prefix + lowerArm + suffix
+        driver_clavicle = "driver_" + prefix + "clavicle" + suffix
+        driver_upperarm = "driver_deltoid_" + prefix + upperArm + suffix
+        driver_lowerarm = "driver_" + prefix + lowerArm + suffix
+
+        numRolls = 0
+        for joint in ["_twist_01", "_twist_02", "_twist_03", "_twist_04", "_twist_05", "_twist_06"]:
+            if cmds.objExists("driver_" + prefix + upperArm + joint + suffix):
+                numRolls = numRolls + 1
+
+        print "...There are a total of " + str(numRolls) + " to build."
+            
+        for i in range(int(numRolls)):
+            print "...Building upper arm twist_0" + str(i + 1)
+            driver_upperarm_twist = "driver_"+prefix+upperArm+"_twist_0"+str(i+1)+suffix
+
+            rollGrp = cmds.group(empty=True, name=upperarm+"_roll_grp_0"+str(i+1))
+            cmds.parent(rollGrp, "arm_sys_grp")
+
+            # Move the driver_upperarm_twist joint to its correct position between the driver_upperarm and its driver_lowerarm using a point constraint.
+            const = cmds.parentConstraint(driver_upperarm, driver_upperarm_twist, n=driver_upperarm+"_temp_parentConst", weight=1)
+            cmds.delete(const)
+            const = cmds.pointConstraint(driver_upperarm, driver_lowerarm, driver_upperarm_twist, n=driver_upperarm+"_temp_PointConst", weight=1)
+            cmds.delete(const)
+
+            # Duplicate the driver_upperarm, driver_lowerarm, and driver_upperarm_twist joints
+            blendJnts=[]
+            blend_driver_upperarmParent = cmds.duplicate(driver_upperarm, n="Blend_"+driver_upperarm+"_parent_ignore_0"+str(i+1), ic=True, po=True)[0]
+            blendJnts.append(blend_driver_upperarmParent)
+
+            blend_driver_lowerarm = cmds.duplicate(driver_lowerarm, n="Blend_"+driver_lowerarm+"_ignore_0"+str(i+1), ic=True, po =True)[0]
+            blendJnts.append(blend_driver_lowerarm)
+
+            blend_driver_upperarm = cmds.duplicate(driver_upperarm, n="Blend_"+driver_upperarm+"_ignore_0"+str(i+1), po=True)[0]
+            blendJnts.append(blend_driver_upperarm)
+
+            blend_driver_upperarm_twist = cmds.duplicate(driver_upperarm_twist, n="Blend_"+driver_upperarm_twist+"_ignore_0"+str(i+1), ic=True, po= True)[0]
+            blendJnts.append(blend_driver_upperarm_twist)
+
+            cmds.parent(blend_driver_lowerarm, blend_driver_upperarm_twist, blend_driver_upperarmParent)
+            cmds.parent(blend_driver_upperarm, blend_driver_upperarm_twist)
+
+            # Set the parent bone to have zyx rotation order
+            cmds.setAttr(blend_driver_upperarmParent+".rotateOrder", 2)
+
+            originalJnts = [driver_upperarm, driver_lowerarm]
+            #print originalJnts
+            
+            # Now we are going to disconnect the constraints that were driving the original driver_upperarm 
+            # and driver_lowerarm joints and re-connect them to the parent_ignore and ignore joints that were duped from them.
+            # This is so that we can create new constraints on the driver_upperarm and driver_lowerarm joints for the rig.
+            j = 0
+            a = 0
+
+            #while a < len(originalJnts):
+            for a in range(0,len(originalJnts)):
+                # Go through the inputs of the original joints and disconnect them
+                origConn = cmds.listConnections(originalJnts[a], d=False, c=+True, p=True)
+
+                for j in range(0,len(origConn),2):
+                    cmds.disconnectAttr(origConn[j+1], origConn[j])
+                    nameList = origConn[j+1].split(".")[0]
+                    destConn = cmds.listConnections(nameList, s=False)
+                    # Find out if the node that was an input to the original joint is also an output
+                    if originalJnts[a] in destConn:
+                        sourceConn = cmds.listConnections(nameList, d=False, c=True, p=True)
+                        for k in range(0,len(sourceConn),2):
+                            # Get the input connections to the node that are connected to the original joint
+                            nameConnList = sourceConn[k+1].split(".")
+                            if (nameConnList[0] == originalJnts[a]):
+                                # Disconnect from the original joint and connect to the blend joint
+                                cmds.disconnectAttr(sourceConn[k+1], sourceConn[k])
+                                nameConnList[0] = blendJnts[a]
+                                cmds.connectAttr(nameConnList[0] + "." + nameConnList[-1], sourceConn[k])
+
+            # create the manual twist control
+            if i == 0:
+                twistCtrl = utils.createControl("circle", 15, upperarm + "_twist_anim")
+            else:
+                twistCtrl = utils.createControl("circle", 15, upperarm + "_twist"+str(i+1)+"_anim")
+            cmds.setAttr(twistCtrl + ".ry", -90)
+            cmds.setAttr(twistCtrl + ".sx", 0.8)
+            cmds.setAttr(twistCtrl + ".sy", 0.8)
+            cmds.setAttr(twistCtrl + ".sz", 0.8)
+            cmds.makeIdentity(twistCtrl, r = 1, s = 1, apply =True)
+
+            # move the manual control to the correct location
+            constraint = cmds.parentConstraint(blend_driver_upperarm_twist, twistCtrl)[0]
+            cmds.delete(constraint)
+
+            # create a group for the manual control and parent the twist to it.
+            twistCtrlGrp = cmds.group(empty = True, name = twistCtrl + "_grp")
+            constraint = cmds.parentConstraint(blend_driver_upperarm_twist, twistCtrlGrp)[0]
+            cmds.delete(constraint)
+            cmds.parent(twistCtrl, twistCtrlGrp)
+            cmds.parent(twistCtrlGrp, rollGrp)
+            cmds.makeIdentity(twistCtrl, t = 1, r = 1, s = 1, apply = True)
+            cmds.parentConstraint(blend_driver_upperarm_twist, twistCtrlGrp)
+
+            # set the manual controls visibility settings
+            cmds.setAttr(twistCtrl + ".overrideEnabled", 1)
+            cmds.setAttr(twistCtrl + ".overrideColor", color)
+            for attr in [".sx", ".sy", ".sz"]:
+                cmds.setAttr(twistCtrl + attr, lock = True, keyable = False)
+            cmds.setAttr(twistCtrl + ".v", keyable = False)	
+            
+            # add attr on rig settings for manual twist control visibility
+            cmds.select("Rig_Settings")
+            if i == 0:
+                cmds.addAttr(longName=(name + "twistCtrlVisUpper"), at = 'bool', dv = 0, keyable = True)
+            cmds.connectAttr("Rig_Settings." + name + "twistCtrlVisUpper", twistCtrl + ".v")
+                        
+            # add attr to rig settings for the twist ammount values
+            cmds.select("Rig_Settings")
+            cmds.addAttr(longName= (name+"UpperarmTwist"+str(i+1)+"Amount"), defaultValue=0.5, minValue=0, maxValue=1, keyable = True)
+            for u in range(int(i+1)):
+                cmds.setAttr("Rig_Settings."+name+"UpperarmTwist"+str(u+1)+"Amount", (1.0/(i+2.0)*((2.0-u)+(i-1.0))) )
+
+            cmds.parent(blend_driver_upperarmParent, rollGrp)
+
+            # Constrain the original joints to the blend joints
+            cmds.orientConstraint(blend_driver_upperarm, driver_upperarm, n=driver_upperarm+"_0"+str(i+1)+"_OrntCnst", mo=True, w=1)
+            cmds.pointConstraint(blend_driver_upperarm, driver_upperarm, n=driver_upperarm+"_0"+str(i+1)+"_PtCnst", mo=True, w=1)
+
+            cmds.orientConstraint(blend_driver_lowerarm, driver_lowerarm, n=driver_lowerarm+"_0"+str(i+1)+"_OrntCnst", mo=True, w=1 )
+            cmds.pointConstraint(blend_driver_lowerarm, driver_lowerarm, n=driver_lowerarm+"_0"+str(i+1)+"_PtCnst", mo=True, w=1)
+
+            cmds.orientConstraint(twistCtrl, driver_upperarm_twist, n=driver_upperarm_twist+"_0"+str(i+1)+"_OrntCnst", mo=True, w=1)
+            cmds.pointConstraint(twistCtrl, driver_upperarm_twist, n=driver_upperarm_twist+"_0"+str(i+1)+"_PtCnst", mo=True, w=1)
+
+            # Create a driver_upperarm_twist multiplier to multiply the driver_upperarm_twist values by -.5 to counter rotate the driver_upperarm_twist and driver_upperarm back from the overall limb's driver_upperarm_twist amount.
+            twistMult = cmds.shadingNode("multiplyDivide", asUtility=True, name=blend_driver_upperarm_twist+"_0"+str(i+1)+"_MDnode")
+            twistMultDrive = cmds.shadingNode("multiplyDivide", asUtility=True, name=blend_driver_upperarm_twist+"_0"+str(i+1)+"_Drive_MDnode")
+            twistMultDriveInv = cmds.shadingNode("multiplyDivide", asUtility=True, name=blend_driver_upperarm_twist+"_0"+str(i+1)+"_DriveInv_MDnode")
+            twistMultMultiplier = cmds.shadingNode("multiplyDivide", asUtility=True, name=blend_driver_upperarm_twist+"_0"+str(i+1)+"_Multiplier_MDnode")
+            twistMultMultiplierInv = cmds.shadingNode("multiplyDivide", asUtility=True, name=blend_driver_upperarm_twist+"_0"+str(i+1)+"_MultiplierInv_MDnode")
+            twistReverse = cmds.shadingNode("reverse", asUtility=True, name=blend_driver_upperarm_twist+"_0"+str(i+1)+"_REVnode")
+
+            cmds.connectAttr(blend_driver_upperarmParent+".rotateX", twistMult+".input1X", force=True)
+            cmds.connectAttr("Rig_Settings."+name+"UpperarmTwist"+str(i+1)+"Amount", twistReverse+".inputX")
+            cmds.connectAttr("Rig_Settings."+name+"UpperarmTwist"+str(i+1)+"Amount", twistMultMultiplierInv+".input1X")
+            cmds.connectAttr(twistReverse+".outputX", twistMultMultiplier+".input1X", force=True)
+            cmds.connectAttr(twistMult+".outputX", twistMultDrive+".input1X", force=True)
+            cmds.connectAttr(twistMultMultiplier+".outputX", twistMultDrive+".input2X", force=True)
+            cmds.connectAttr(twistMult+".outputX", twistMultDriveInv+".input1X", force=True)
+            cmds.connectAttr(twistMultMultiplierInv+".outputX", twistMultDriveInv+".input2X", force=True)
+            cmds.connectAttr(twistMultDrive+".outputX", blend_driver_upperarm_twist+".rotateX", force=True)
+            cmds.connectAttr(twistMultDriveInv+".outputX", blend_driver_upperarm+".rotateX", force=True)
+            cmds.setAttr(twistMult+".input2X", -0.5)
+            cmds.setAttr(twistMult+".input2X", l=True)
+
+            cmds.setAttr(twistMultMultiplier+".input2X", 2)
+            cmds.setAttr(twistMultMultiplierInv+".input2X", 2)
+            
+            cmds.select(blend_driver_upperarmParent)
+            
+            cmds.parentConstraint(driver_clavicle, rollGrp, mo=True)
+    print ".END Building the "+ suffix +" UPPER arm twists---------------------------------------"
\ No newline at end of file
diff --git a/Engine/Extras/Maya_AnimationRiggingTools/MayaTools/General/Scripts/Modules/ART_UpperLegIKTwist.py b/Engine/Extras/Maya_AnimationRiggingTools/MayaTools/General/Scripts/Modules/ART_UpperLegIKTwist.py
new file mode 100644
--- /dev/null
+++ b/Engine/Extras/Maya_AnimationRiggingTools/MayaTools/General/Scripts/Modules/ART_UpperLegIKTwist.py
@@ -0,0 +1,129 @@
+import maya.cmds as cmds
+import Modules.ART_rigUtils as utils
+reload(utils)
+
+#----------------------------------------------------------------------------------------
+# Currently this script is written to work with multiple twists, but if you have more than one the last one is the only one that will work.
+def upperIKTwist(color, suffix, name, prefix, upperJnt, lowerJnt, sysGroup):
+    print "START Building the "+suffix+" "+upperJnt+" UPPER IK twists---------------------------------------"
+
+    #create a nice name
+    name = prefix + name + suffix
+    if name.find("_") == 0:
+        name = name.partition("_")[2]
+        
+        # define the joints for this rig.
+        upperjnt = prefix + upperJnt + suffix
+        lowerjnt = prefix + lowerJnt + suffix
+        #OLD - upperjnt_parent = "driver_" + prefix + "clavicle" + suffix
+        upperjnt_parent = cmds.listRelatives(upperjnt, p=True)[0]
+        driver_upperjnt = "driver_" + prefix + upperJnt + suffix
+        driver_lowerjnt = "driver_" + prefix + lowerJnt + suffix
+
+        numRolls = 0
+        for joint in ["_twist_01", "_twist_02", "_twist_03", "_twist_04", "_twist_05", "_twist_06"]:
+            if cmds.objExists("driver_" + prefix + upperJnt + joint + suffix):
+                numRolls = numRolls + 1
+
+        print "...There are a total of " + str(numRolls) + " to build."
+
+        for i in range(int(numRolls)):
+            print "...Building upper twist_0" + str(i + 1)
+
+            upperjnt_twist = prefix+upperJnt+"_twist_0"+str(i+1)+suffix
+            driver_upperjnt_twist = "driver_"+prefix+upperJnt+"_twist_0"+str(i+1)+suffix
+
+            ####################################
+            # Create a master group for the roll controls to be parented to.
+            rollGrp = cmds.group(empty=True, name=upperjnt+"_roll_grp_0"+str(i+1))
+            cmds.parent(rollGrp, sysGroup)
+            cmds.parentConstraint(upperjnt_parent, rollGrp, mo=True)
+
+            # Create the manual twist control
+            if i == 0:
+                twistCtrl = utils.createControl("circle", 15, upperjnt + "_twist_anim")
+            else:
+                twistCtrl = utils.createControl("circle", 15, upperjnt + "_twist"+str(i+1)+"_anim")
+            cmds.setAttr(twistCtrl + ".ry", -90)
+            cmds.setAttr(twistCtrl + ".sx", 0.8)
+            cmds.setAttr(twistCtrl + ".sy", 0.8)
+            cmds.setAttr(twistCtrl + ".sz", 0.8)
+            cmds.makeIdentity(twistCtrl, r = 1, s = 1, apply =True)
+
+            # move the manual control to the correct location
+            constraint = cmds.parentConstraint(driver_upperjnt_twist, twistCtrl)[0]
+            cmds.delete(constraint)
+
+            # create a group for the manual control and parent the twist to it.
+            twistCtrlGrp = cmds.group(empty = True, name=twistCtrl+"_grp")
+            constraint = cmds.parentConstraint(driver_upperjnt_twist, twistCtrlGrp)[0]
+            cmds.delete(constraint)
+            cmds.parent(twistCtrl, twistCtrlGrp)
+            cmds.parent(twistCtrlGrp, rollGrp)
+            cmds.makeIdentity(twistCtrl, t = 1, r = 1, s = 1, apply = True)
+
+            # set the manual controls visibility settings
+            cmds.setAttr(twistCtrl + ".overrideEnabled", 1)
+            cmds.setAttr(twistCtrl + ".overrideColor", color)
+            for attr in [".sx", ".sy", ".sz"]:
+                cmds.setAttr(twistCtrl + attr, lock = True, keyable = False)
+            cmds.setAttr(twistCtrl + ".v", keyable = False)	
+
+            # add attr on rig settings for manual twist control visibility
+            cmds.select("Rig_Settings")
+            if i == 0:
+                cmds.addAttr(longName=(upperjnt + "_twistCtrlVisUpper"), at = 'bool', dv = 0, keyable = True)
+            cmds.connectAttr("Rig_Settings." + upperjnt + "_twistCtrlVisUpper", twistCtrl + ".v")
+
+            # add attr to rig settings for the twist ammount values
+            cmds.select("Rig_Settings")
+            cmds.addAttr(longName= (upperjnt+"_Twist"+str(i+1)+"Amount"), defaultValue=0.5, minValue=0, maxValue=1, keyable = True)
+            for u in range(int(i+1)):
+                cmds.setAttr("Rig_Settings."+upperjnt+"_Twist"+str(u+1)+"Amount", (1.0/(i+2.0)*((2.0-u)+(i-1.0))) )
+
+            ####################################
+            # Create the twist rig on the shoulder without any extra controls.  Those are added later.
+            cmds.delete(upperjnt_twist+"_orientConstraint1")
+
+            twistIKHandle = cmds.ikHandle(name="twistIKHandle_0"+str(i+1)+"_"+upperjnt, sol="ikRPsolver", sj=upperjnt, ee=upperjnt_twist)[0]
+            print "TWIST IK HANDLE IS: "+twistIKHandle
+            cmds.setAttr(twistIKHandle+".poleVectorX", 0)
+            cmds.setAttr(twistIKHandle+".poleVectorY", 0)
+            cmds.setAttr(twistIKHandle+".poleVectorZ", 0)
+            cmds.parent(twistIKHandle, driver_upperjnt)
+
+            driver_upperjnt_offset = cmds.duplicate(driver_upperjnt, po=True, name="driver_"+upperJnt+"_twistOffset_l")[0]
+            cmds.parent(driver_upperjnt_offset, driver_upperjnt)
+
+            twistOrntConst = cmds.orientConstraint(driver_upperjnt, upperjnt, driver_upperjnt_offset, mo=True)[0]
+            cmds.setAttr(twistOrntConst+".interpType", 2)
+
+            twistMDNode = cmds.shadingNode("multiplyDivide", asUtility=True, name=driver_upperjnt_twist+"_MDnode")
+            cmds.setAttr(twistMDNode+".input2X", 1)
+
+            cmds.connectAttr(driver_upperjnt_offset+".rotateX", twistMDNode+".input1X", force=True)
+            cmds.connectAttr(twistMDNode+".outputX", driver_upperjnt_twist+".rx", force=True)            
+
+            cmds.pointConstraint(driver_upperjnt_twist, twistCtrlGrp, mo=True)
+            cmds.orientConstraint(driver_upperjnt_twist, twistCtrlGrp, mo=True)
+
+            cmds.orientConstraint(twistCtrl, upperjnt_twist, mo=True)
+            
+            if cmds.objExists("deltoid_anim"+suffix):
+                if suffix == "_l":
+                    cmds.connectAttr("deltoid_anim"+suffix+".rotateX", twistIKHandle+".twist", force=True)
+                else:
+                    deltMDNode = cmds.shadingNode("multiplyDivide", asUtility=True, name=twistIKHandle+"_MDNode")
+                    cmds.setAttr(deltMDNode+".input2X", -1)
+                    cmds.connectAttr("deltoid_anim"+suffix+".rotateX", deltMDNode+".input1X", force=True)
+                    cmds.connectAttr(deltMDNode+".outputX", twistIKHandle+".twist", force=True)
+
+            reverseMDNode = cmds.shadingNode("reverse", asUtility=True, name=driver_upperjnt_offset+"_Reversenode")
+            cmds.connectAttr("Rig_Settings."+upperjnt+"_Twist"+str(i+1)+"Amount", reverseMDNode+".inputX", force=True)
+            cmds.connectAttr(reverseMDNode+".outputX", twistOrntConst+"."+upperjnt+"W1", force=True)            
+            cmds.connectAttr("Rig_Settings."+upperjnt+"_Twist"+str(i+1)+"Amount", twistOrntConst+"."+driver_upperjnt+"W0", force=True)
+
+
+    print ".END Building the "+suffix+" "+upperJnt+" UPPER twists---------------------------------------"            
+
+#upperIKTwist(6, "_l", "", "", "upperarm", "lowerarm", "arm_sys_grp")
