Branching

diff --git a/Engine/Extras/Maya_AnimationRiggingTools/ARTv1/General/ART/Projects/Orion/Scripts/Coil_Post_Script.py b/Engine/Extras/Maya_AnimationRiggingTools/ARTv1/General/ART/Projects/Orion/Scripts/Coil_Post_Script.py
new file mode 100644
--- /dev/null
+++ b/Engine/Extras/Maya_AnimationRiggingTools/ARTv1/General/ART/Projects/Orion/Scripts/Coil_Post_Script.py
@@ -0,0 +1,452 @@
+import maya.mel as mel
+import maya.cmds as cmds
+
+
+# NOTE: for building this rig, because of the spine_03_alt joint being oriented to match exactly spine_03, the children of spine_03_alt when mirrored have to have their mirror settings set to 1, -1, 1, 1, -1, 1.  Otherwise they wont mirror.  When you go back to the Skeleton Creation step these values likely wont persist and will have to be manually added.  The following script can be run to set these values.
+'''
+cmds.setAttr("big_wire_l_mover.mirrorTX", 1)
+cmds.setAttr("big_wire_l_mover.mirrorTY", -1)
+cmds.setAttr("big_wire_l_mover.mirrorTZ", 1)
+cmds.setAttr("big_wire_l_mover.mirrorRX", 1)
+cmds.setAttr("big_wire_l_mover.mirrorRY", -1)
+cmds.setAttr("big_wire_l_mover.mirrorRZ", 1)
+cmds.setAttr("shoulderlamp_l_mover.mirrorTX", 1)
+cmds.setAttr("shoulderlamp_l_mover.mirrorTY", -1)
+cmds.setAttr("shoulderlamp_l_mover.mirrorTZ", 1)
+cmds.setAttr("shoulderlamp_l_mover.mirrorRX", -1)
+cmds.setAttr("shoulderlamp_l_mover.mirrorRY", 1)
+cmds.setAttr("shoulderlamp_l_mover.mirrorRZ", -1)
+cmds.setAttr("shoulderpad_l_mover.mirrorTX", 1)
+cmds.setAttr("shoulderpad_l_mover.mirrorTY", -1)
+cmds.setAttr("shoulderpad_l_mover.mirrorTZ", 1)
+cmds.setAttr("shoulderpad_l_mover.mirrorRX", 1)
+cmds.setAttr("shoulderpad_l_mover.mirrorRY", -1)
+cmds.setAttr("shoulderpad_l_mover.mirrorRZ", 1)
+cmds.setAttr("sideplate_l_mover.mirrorTX", 1)
+cmds.setAttr("sideplate_l_mover.mirrorTY", -1)
+cmds.setAttr("sideplate_l_mover.mirrorTZ", 1)
+cmds.setAttr("sideplate_l_mover.mirrorRX", 1)
+cmds.setAttr("sideplate_l_mover.mirrorRY", -1)
+cmds.setAttr("sideplate_l_mover.mirrorRZ", 1)
+cmds.setAttr("small_wire_l_mover.mirrorTX", 1)
+cmds.setAttr("small_wire_l_mover.mirrorTY", -1)
+cmds.setAttr("small_wire_l_mover.mirrorTZ", 1)
+cmds.setAttr("small_wire_l_mover.mirrorRX", 1)
+cmds.setAttr("small_wire_l_mover.mirrorRY", -1)
+cmds.setAttr("small_wire_l_mover.mirrorRZ", 1)
+cmds.setAttr("turbineplate_l_mover.mirrorTX", 1)
+cmds.setAttr("turbineplate_l_mover.mirrorTY", -1)
+cmds.setAttr("turbineplate_l_mover.mirrorTZ", 1)
+cmds.setAttr("turbineplate_l_mover.mirrorRX", -1)
+cmds.setAttr("turbineplate_l_mover.mirrorRY", 1)
+cmds.setAttr("turbineplate_l_mover.mirrorRZ", -1)
+'''
+
+
+# ------------------------------------------------------------------
+# scale all of the controls to be more visible with the mesh.
+def scaleControls():
+    for side in ["_l", "_r"]:
+        # feet
+        cmds.select(["ik_foot_anim"+side+".cv[*]"])
+        cmds.move(0, 8.796488, 5.948312, r = True)
+
+        cmds.select(["ik_wrist"+side+"_anim.cv[*]"])
+        cmds.scale(1.6, 1.6, 1.6, r = True, ocp = False)
+        
+        cmds.select(["weapon"+side+"_anim.cv[*]"])
+        cmds.scale(0.48, 0.48, 0.48, r = True, ocp = False)
+        cmds.move(0, -30.39, 0, r = True)            
+
+    cmds.select(["hip_l_anim.cv[*]"])
+    cmds.scale(0.48, 0.48, 0.48, r = True, ocp = False)
+    cmds.move(16, 0, 15, r = True)            
+
+    cmds.select(["hip_r_anim.cv[*]"])
+    cmds.scale(0.48, 0.48, 0.48, r = True, ocp = False)
+    cmds.move(-16, 0, 15, r = True)            
+
+    cmds.select(["shoulderpad_l_anim.cv[*]"])
+    cmds.scale(0.48, 0.48, 0.48, r = True, ocp = False)
+    cmds.move(5, 0, 21, r = True)            
+
+    cmds.select(["shoulderpad_r_anim.cv[*]"])
+    cmds.scale(0.48, 0.48, 0.48, r = True, ocp = False)
+    cmds.move(-5, 0, 21, r = True)            
+
+    cmds.select(["sideplate_l_anim.cv[*]"])
+    cmds.scale(0.48, 0.48, 0.48, r = True, ocp = False)
+    cmds.move(15, 0, 0, r = True)            
+
+    cmds.select(["sideplate_r_anim.cv[*]"])
+    cmds.scale(0.48, 0.48, 0.48, r = True, ocp = False)
+    cmds.move(-15, 0, 0, r = True)            
+
+    cmds.select(["spine_03_alt_anim.cv[*]"])
+    cmds.scale(1, 5, 5, r = True, ocp = False)
+    cmds.move(0, -19, 12, r = True)            
+
+    cmds.select(["weapon_stock_anim.cv[*]"])
+    cmds.scale(0.48, 0.48, 0.48, r = True, ocp = False)
+    cmds.move(25, 0, 0, r = True)            
+
+    cmds.select(["weapon_zapper_top_anim.cv[*]"])
+    cmds.scale(0.48, 0.48, 0.48, r = True, ocp = False)
+    cmds.move(21, 0, 0, r = True)            
+
+    cmds.select(["weapon_zapper_bott_anim.cv[*]"])
+    cmds.scale(0.48, 0.48, 0.48, r = True, ocp = False)
+    cmds.move(-21, 0, 0, r = True)            
+
+    cmds.select(["neckbackplate_anim.cv[*]"])
+    cmds.scale(0.48, 0.48, 0.48, r = True, ocp = False)
+    cmds.move(0, 8, 4, r = True)            
+
+    cmds.select(["chestplate_anim.cv[*]"])
+    cmds.scale(0.48, 0.48, 0.48, r = True, ocp = False)
+    cmds.move(0, -6, 11, r = True)            
+
+    cmds.select(["turbine_anim.cv[*]"])
+    cmds.scale(0.48, 0.48, 0.48, r = True, ocp = False)
+    cmds.move(0, 9, 0, r = True)            
+
+    cmds.select(["weapon_wires_anim.cv[*]"])
+    cmds.scale(0.48, 0.48, 0.48, r = True, ocp = False)
+   
+    objects = ["sixpack_bott_anim", "sixpack_mid_anim", "sixpack_top_anim", "small_wire_l_anim", "small_wire_r_anim", "big_wire_l_anim", "big_wire_r_anim", "shoulderlamp_l_anim", "shoulderlamp_r_anim", "turbineplate_l_anim", "turbineplate_r_anim"]
+    for one in objects:
+        cmds.select([one+".cv[*]"])
+        cmds.scale(0.4, 0.4, 0.4, r = True, ocp = True)
+        cmds.move(0, 0, 0, r = True, os=True, wd=True)
+
+    # Other
+    cmds.select(["jaw_anim.cv[*]"])
+    cmds.scale(0.2, 0.2, 0.2, r = True, ocp = True)
+    cmds.move(0, -7, -2, r = True)
+
+    cmds.select(None)
+scaleControls()
+
+# ------------------------------------------------------------------
+# Move the panels on the torso and hook up the clav to move with them.
+def panelsRig():
+
+    cmds.setAttr("Rig_Settings.lClavMode", 0)
+    cmds.setAttr("Rig_Settings.rClavMode", 0)
+    
+    cmds.parentConstraint("sideplate_l_anim", "fk_clavicle_l_anim_grp", mo=True)
+    cmds.parentConstraint("sideplate_r_anim", "fk_clavicle_r_anim_grp", mo=True)
+
+panelsRig()
+
+# ------------------------------------------------------------------
+def shoulderRig():
+    for side in ["_l", "_r"]:
+        cmds.addAttr("shoulderpad"+side+"_anim", ln = "FollowClav", dv = 1, min = 0, max = 1, keyable = True)
+
+        shoulder_loc = cmds.duplicate("shoulderpad"+side+"_anim_space_switcher", n="shoulder_loc"+side+"", po=True)[0]
+        cmds.parent(shoulder_loc, "ctrl_rig")
+        cmds.makeIdentity(shoulder_loc, t=True, r=True, a=True)
+        cmds.parentConstraint("driver_upperarm"+side+"", shoulder_loc, mo=True)
+
+        const = cmds.pointConstraint(shoulder_loc, "shoulderpad"+side+"_anim_space_switcher", "shoulderpad"+side+"_anim_grp", mo=True)[0]
+
+        reverseNode = cmds.shadingNode("reverse", asUtility=True, name="shoulderpad"+side+"_reverse")
+        cmds.connectAttr("shoulderpad"+side+"_anim.FollowClav", const+"."+shoulder_loc+"W0", f=True)
+        cmds.connectAttr("shoulderpad"+side+"_anim.FollowClav", reverseNode+".inputX", f=True)
+        cmds.connectAttr(reverseNode+".outputX", const+".shoulderpad"+side+"_anim_space_switcherW1", f=True)
+
+        if side == "_l":
+            cmds.transformLimits("shoulderpad"+side+"_anim_grp", tz=[0, 1], etz=[1, 0])
+        else:
+            cmds.transformLimits("shoulderpad"+side+"_anim_grp", tz=[-1, 0], etz=[0, 1])
+
+    cmds.select(None)
+    
+shoulderRig()
+
+# ------------------------------------------------------------------
+# Modify the hips so that the legs follow the hip_*_anim control.
+def hipsRig():
+    cmds.parentConstraint("hip_l_anim", "hip_redirect_l", mo=True)        
+    cmds.parentConstraint("hip_r_anim", "hip_redirect_r", mo=True)        
+hipsRig()
+
+# ------------------------------------------------------------------
+def spineRig():
+    cmds.setAttr("Rig_Settings.spine_ik", 0)
+    cmds.setAttr("Rig_Settings.spine_fk", 1)
+
+    # Create the rig for the sixpack stuff.
+    for one in ["sixpack_top", "sixpack_mid", "sixpack_bott"]:
+        cmds.select("neck_spine_geo", r=True)
+        cmds.ClosestPointOn(1, 1)
+        surfaceLoc = cmds.rename("cpConstraintPos", one+"_surfaceLoc")
+        worldLoc = cmds.rename("cpConstraintIn", one+"worldLoc")
+        cmds.parent(surfaceLoc, worldLoc, "rig_grp")
+        const = cmds.parentConstraint(one, worldLoc, mo=False)[0]
+        cmds.delete(const)
+        
+        cmds.parentConstraint("spine_03", worldLoc, mo=True)[0]
+        
+
+        if one == "sixpack_top":
+            cmds.normalConstraint("neck_spine_geo", surfaceLoc, aimVector=[0, 0, 1], upVector=[0, 1, 0], worldUpType="object", worldUpObject="spine_03_alt")
+        if one == "sixpack_mid":
+            cmds.normalConstraint("neck_spine_geo", surfaceLoc, aimVector=[0, 0, 1], upVector=[0, 1, 0], worldUpType="object", worldUpObject="sixpack_top")
+        if one == "sixpack_bott":
+            cmds.normalConstraint("neck_spine_geo", surfaceLoc, aimVector=[0, 0, 1], upVector=[0, 1, 0], worldUpType="object", worldUpObject="sixpack_mid")
+
+        cmds.parentConstraint(surfaceLoc, one+"_anim_grp", mo=True)[0]
+
+spineRig()
+
+# ------------------------------------------------------------------
+def createGameCam():
+    gameCam = cmds.camera(centerOfInterest=5, focalLength=18, lensSqueezeRatio=1, cameraScale=1, horizontalFilmAperture=1.41732, horizontalFilmOffset=0, verticalFilmAperture=0.94488, verticalFilmOffset=0, filmFit="Overscan", overscan=1.05, motionBlur=0, shutterAngle=144, nearClipPlane=1.0, farClipPlane=10000, orthographic=0, orthographicWidth=30, panZoomEnabled=0, horizontalPan=0, verticalPan=0, zoom=1, displayFilmGate=True)[0]
+    print gameCam
+    gameCam = cmds.rename(gameCam, "GameCam")
+    cmds.camera(gameCam, e=True, horizontalFilmAperture=1.6724376, focalLength=21.2399857011)
+    cmds.setAttr(gameCam+".displayGateMaskOpacity", 1)
+    cmds.setAttr(gameCam+".displayGateMaskColor", 0, 0, 0, type="double3")
+    cmds.setAttr(gameCam+".r", 90, 0, 180, type="double3")
+    
+    # Create the control
+    control = cmds.circle(name=gameCam+"_Ctrl", ch=True, o=True, r=15)[0]
+    cmds.makeIdentity(control, a=True, t=True, r=True, s=True)
+    # Create groups for the sdk and the constraints
+    const_grp = cmds.group(control, name=control+"_grp")
+    if cmds.objExists("rig_grp"):
+        cmds.parent(const_grp, "rig_grp")
+    # Parent Constrain the cam to the control
+    cmds.parentConstraint(control, gameCam, mo=True)
+    
+    cmds.connectAttr(control+".visibility", gameCam+".visibility", f=True)
+    cmds.setAttr(const_grp+".v", 0)
+    
+    cmds.setAttr(const_grp+".t", 0, 500, 200, type="double3")
+createGameCam()
+
+# ------------------------------------------------------------------
+def addCapsule(radius, halfHeight):
+    trueHalfHeight = (halfHeight*2)-(radius*2)
+    capsule = cmds.polyCylinder(name="collision_capsule", r=radius, h=trueHalfHeight, sx=18, sy=1, sz=5, ax=[0,0,1], rcp=True, cuv=3, ch=True)
+    cmds.setAttr(capsule[0]+".tz", halfHeight+2.2)
+    cmds.setAttr(capsule[0]+".rz", 90)
+
+    cmds.addAttr(capsule[0], ln="Radius", dv=radius, keyable=True)
+    cmds.addAttr(capsule[0], ln="HalfHeight", dv=halfHeight, keyable=True)
+
+    expr = capsule[1]+".radius = "+capsule[0]+".Radius;\n"+capsule[1]+".height = ("+capsule[0]+".HalfHeight*2) - ("+capsule[0]+".Radius*2);"
+    cmds.expression(s=expr, o=capsule[1], ae=1, uc="all")
+
+    cmds.setAttr(capsule[0]+".overrideEnabled", 1)
+    cmds.setAttr(capsule[0]+".overrideColor", 18)
+    cmds.setAttr(capsule[0]+".overrideShading", 0)
+
+    cmds.addAttr(capsule[0], ln="Shaded", dv=0, min=0, max=1, keyable=True)
+    cmds.connectAttr(capsule[0]+".Shaded", capsule[0]+".overrideShading")
+
+    cmds.parentConstraint("root", capsule[0], mo=True)
+    cmds.select(capsule[0], r=True)
+    capsuleDisplayLayer = cmds.createDisplayLayer(nr=True, name="Collision Capsule")
+    cmds.setAttr(capsuleDisplayLayer+".displayType", 2)
+    cmds.setAttr(capsuleDisplayLayer+".v", 0)
+    
+    for attr in ["tx", "ty", "tz", "rx", "ry", "rz", "sx", "sy", "sz"]:
+        cmds.setAttr(capsule[0]+"."+attr, lock=True, keyable=False, channelBox=False)
+
+    cmds.select(clear=True)
+addCapsule(72, 108)
+
+# ------------------------------------------------------------------
+# add ik bones
+def addIKBones():
+    try:
+		# create the joints
+		cmds.select(clear = True)
+		ikFootRoot = cmds.joint(name = "ik_foot_root")
+		cmds.select(clear = True)
+		
+		cmds.select(clear = True)
+		ikFootLeft = cmds.joint(name = "ik_foot_l")
+		cmds.select(clear = True)
+		
+		cmds.select(clear = True)
+		ikFootRight = cmds.joint(name = "ik_foot_r")
+		cmds.select(clear = True)
+		
+		cmds.select(clear = True)
+		ikHandRoot = cmds.joint(name = "ik_hand_root")
+		cmds.select(clear = True)
+		
+		cmds.select(clear = True)
+		ikHandGun = cmds.joint(name = "ik_hand_gun")
+		cmds.select(clear = True)
+		
+		cmds.select(clear = True)
+		ikHandLeft = cmds.joint(name = "ik_hand_l")
+		cmds.select(clear = True)
+		
+		cmds.select(clear = True)
+		ikHandRight = cmds.joint(name = "ik_hand_r")
+		cmds.select(clear = True)
+		
+		# create hierarhcy
+		cmds.parent(ikFootRoot, "root")
+		cmds.parent(ikHandRoot, "root")
+		cmds.parent(ikFootLeft, ikFootRoot)
+		cmds.parent(ikFootRight, ikFootRoot)
+		cmds.parent(ikHandGun, ikHandRoot)
+		cmds.parent(ikHandLeft, ikHandGun)
+		cmds.parent(ikHandRight, ikHandGun)
+		
+		# constrain the joints
+		leftHandConstraint = cmds.parentConstraint("hand_l", ikHandLeft)[0]
+		rightHandGunConstraint = cmds.parentConstraint("hand_r", ikHandGun)[0]
+		rightHandConstraint = cmds.parentConstraint("hand_r", ikHandRight)[0]
+		leftFootConstraint = cmds.parentConstraint("foot_l", ikFootLeft)[0]
+		rightFootConstraint = cmds.parentConstraint("foot_r", ikFootRight)[0]
+	
+    except:
+		cmds.warning("Something went wrong")
+
+# ------------------------------------------------------------------
+# ------------------------------------------------------------------
+# Import in and hook up the Goblin
+def hookUpGoblin():
+    import Tools.ART_addSpaceSwitching as space
+    reload(space)
+    
+    # import in goblin
+    cmds.file("D:/Build/usr/jeremy_ernst/MayaTools/General/ART/Projects/Orion/AnimRigs/Coil_Goblin.mb", i=True, options='v=0', type="mayaBinary", namespace="Goblin", mergeNamespacesOnClash=False)
+
+    cmds.setAttr("Goblin:Rig_Settings.spine_fk", 1)
+    cmds.setAttr("Goblin:Rig_Settings.spine_ik", 0)
+
+    # set elbows to body space
+    cmds.setAttr("Goblin:ik_elbow_l_anim_space_switcher.space_body_anim", 1)
+    cmds.setAttr("Goblin:ik_elbow_r_anim_space_switcher.space_body_anim", 1)
+
+    '''
+    REMOVING THIS SO THAT THE HANDS FOLLOW THE MECHANISM THING!
+    cmds.setAttr("Goblin:ik_wrist_l_anim_space_switcher.space_body_anim", 1)
+    cmds.setAttr("Goblin:ik_wrist_r_anim_space_switcher.space_body_anim", 1)
+    '''
+    
+    cmds.setAttr("Goblin:ik_foot_anim_l_space_switcher.space_body_anim", 1)
+    cmds.setAttr("Goblin:ik_foot_anim_r_space_switcher.space_body_anim", 1)
+
+    # turn on auto clavicles
+    #cmds.setAttr("Goblin:clavicle_l_anim.autoShoulders", 1)
+    #cmds.setAttr("Goblin:clavicle_r_anim.autoShoulders", 1)
+
+    # set head to world space
+    cmds.setAttr("Goblin:head_fk_anim.fkOrientation", 3)
+
+    # Add space switching to the feet on the foot pedals
+    #space.addSpaceSwitching("Goblin:ik_foot_anim_l", "cockpit_foot_pedal_l_anim", True)
+    #space.addSpaceSwitching("Goblin:ik_foot_anim_r", "cockpit_foot_pedal_r_anim", True)
+
+    #position goblin
+    cmds.setAttr("Goblin:body_anim.tx", 219.57371)
+    cmds.setAttr("Goblin:body_anim.ty", 92.45555)
+    cmds.setAttr("Goblin:body_anim.tz", -67.72273)
+
+    # join skeletons
+    cmds.select("Goblin:root", hi = True)
+    joints = cmds.ls(sl = True)
+    for joint in joints:
+        newName = joint.partition(":")[2]
+        cmds.rename(joint, "goblin_" + newName)
+    cmds.parent("goblin_root", "root")    
+
+    # Disconnect the attrs driving the goblin pelvis joint so that we can re-drive it with a parent constraint.
+    attrToDisconnect = "goblin_pelvis.translate"
+    connection = cmds.connectionInfo(attrToDisconnect, sfd=True)
+    cmds.disconnectAttr(connection, attrToDisconnect)
+    cmds.delete("goblin_pelvis_orientConstraint1")
+
+    # Parent constrain the deform pelvis to follow the driver pelvis
+    cmds.parentConstraint("Goblin:driver_pelvis", "goblin_pelvis", mo=True)
+
+    cmds.select(None)
+hookUpGoblin()
+
+# ------------------------------------------------------------------
+# save the AnimRig file
+cmds.file(save = True, type = "mayaBinary", force = True)
+
+# add IK Bones here so that they are included in the export file.  But dont save becasue we dont want them in the rig file.  Thats why we saved in the line before.
+addIKBones()
+
+# ------------------------------------------------------------------
+# Creeate the export file that goes into source so it can be exported to the game
+def createExportFile():
+    # Import the export file from the reference so that it can be edited.
+    cmds.file("D:/Build/usr/jeremy_ernst/MayaTools/General/ART/Projects/Orion/ExportFiles/Coil_Export.mb", importReference=True)
+
+    # select the entire skeleton and make a list of it.
+    cmds.select("root", hi=True)
+    selection = cmds.ls(sl=True)
+    cmds.showHidden(a=True)
+
+    # for each object in the selection, disconnect translation and scale.
+    for each in selection:
+        attrs = ["translate", "scale"]
+        cmds.lockNode(each, lock=False)
+
+        for attr in attrs:
+            try:
+                conn = cmds.listConnections(each + "." + attr, s=True, p=True)[0]
+                cmds.disconnectAttr(conn, each + "." + attr)
+            except:
+                pass
+
+    # find all of the constraints in the skeleton and delete them
+    cmds.select("root", hi=True)
+    constraints = cmds.ls(sl=True, type="constraint")
+    cmds.select(constraints)
+    cmds.delete(constraints)
+
+    # delete the control rig so that it severs any connections between it and the geo or skeleton.
+    cmds.select("Goblin:rig_grp", "rig_grp", hi=True)
+    rigs = cmds.ls(sl=True)
+    for one in rigs:
+        cmds.lockNode(one, lock=False)
+    cmds.delete("Goblin:rig_grp", "rig_grp")
+
+    # Remove the Goblin namespace and set the bindposes so that Creating LODs will work later on.
+    cmds.namespace(rm="Goblin", mnp=True)
+    cmds.dagPose("root", bp=True, r=True) 
+    cmds.dagPose("goblin_root", bp=True, r=True)     
+
+    # select the skeleton and all of the geo and export it to the export file.  Remember to go export an fbx from this.
+    cmds.select("Goblin:Coil_Goblin_Geo", "Coil_Bot_Geo", "root", hi=True)
+    cmds.file("D:/Build/ArtSource/Orion/Characters/Heroes/Coil/Export/Coil_Export.mb", force=True, typ="mayaBinary", pr=True, es=True)
+createExportFile()
+
+# ------------------------------------------------------------------
+# ADD IK BONES to the export file
+sceneName = cmds.file(q = True, sceneName = True)
+exportFile = sceneName.partition("AnimRigs")[0] + "ExportFiles" +  sceneName.partition("AnimRigs")[2].partition(".mb")[0] + "_Export.mb"
+cmds.file(exportFile, open = True, force = True)
+
+
+addIKBones()
+
+# ------------------------------------------------------------------
+def addCustomAttrs():
+    attrs = cmds.listAttr("root")
+    exists = False
+    for one in attrs:
+        if one.find("DistanceCurve") != -1:
+            exists = True
+    if exists == False:
+        cmds.addAttr("root", ln="DistanceCurve", keyable=True)
+addCustomAttrs()
+
+# save the export file
+cmds.file(save = True, type = "mayaBinary", force = True)
+
+
diff --git a/Engine/Extras/Maya_AnimationRiggingTools/ARTv1/General/ART/Projects/Orion/Scripts/Coil_Pre_Script.py b/Engine/Extras/Maya_AnimationRiggingTools/ARTv1/General/ART/Projects/Orion/Scripts/Coil_Pre_Script.py
new file mode 100644
--- /dev/null
+++ b/Engine/Extras/Maya_AnimationRiggingTools/ARTv1/General/ART/Projects/Orion/Scripts/Coil_Pre_Script.py
@@ -0,0 +1,7 @@
+import maya.cmds as cmds
+
+if cmds.objExists("ik_foot_root"):
+	cmds.delete("ik_foot_root")
+	
+if cmds.objExists("ik_hand_root"):
+	cmds.delete("ik_hand_root")
\ No newline at end of file
diff --git a/Engine/Extras/Maya_AnimationRiggingTools/ARTv1/General/ART/Projects/Orion/Scripts/Gadget_CreateLODs.py b/Engine/Extras/Maya_AnimationRiggingTools/ARTv1/General/ART/Projects/Orion/Scripts/Gadget_CreateLODs.py
new file mode 100644
--- /dev/null
+++ b/Engine/Extras/Maya_AnimationRiggingTools/ARTv1/General/ART/Projects/Orion/Scripts/Gadget_CreateLODs.py
@@ -0,0 +1,555 @@
+import maya.cmds as cmds
+import customMayaMenu as menu
+import Tools.ART_ImportMultipleWeights as impw
+import Tools.ART_TransferWeights as transfer
+import Tools.ART_ExportLODtoFBX as export
+reload(export)
+
+################################################## GADGET #########################################################################
+# Warn the user that this operation will save their file.
+saveResult = cmds.confirmDialog(title = "Save Warning", message="In order to continue your file must be saved.  Would you like to save it?  If yes it will be saved and after the operation is complete your file will be re-opened.", button = ["Yes", "Cancel"], defaultButton='Yes', cancelButton='Cancel', dismissString='Cancel')
+
+# If they are ok with saving the file, warn them that it will also delete all of the current files in their skin weights directory.
+if saveResult == "Yes":
+    deleteSkinFilesResult = cmds.confirmDialog(title = "SkinWeights Warning", message="In order to continue we must delete all of the files in your ..\General\ART\SkinWeights\ directory.  Would you like to continue?.", button = ["Yes", "Cancel"], defaultButton='Yes', cancelButton='Cancel', dismissString='Cancel')
+    if deleteSkinFilesResult == "Yes":
+        toolsPath = cmds.internalVar(usd = True) + "mayaTools.txt"
+        if os.path.exists(toolsPath):
+            f = open(toolsPath, 'r')
+            mayaToolsDir = f.readline()
+            f.close()
+
+        path = mayaToolsDir + "\General\ART\SkinWeights\\"
+        files = os.listdir(path)
+        #delete all weight cache files in the system folder
+        for file in files:
+            os.remove(path + file)
+            #print path + file
+
+
+
+# Gadget LOD1 -------------------------------------------------------------------------------------------------------------
+def GadgetLOD1():
+    if saveResult == "Yes":
+
+        cmds.file(save=True)
+
+
+        meshes = ["Corneas_Geo_LOD1", "Eyeballs_Geo_LOD1", "LowerBody_Ankle_Geo_LOD1", "LowerBody_Boots_Geo_LOD1", "LowerBody_Gas_Geo_LOD1", "LowerBody_Gren_Geo_LOD1", "LowerBody_KneePad_Geo_LOD1", "LowerBody_Pak_Geo_LOD1", "LowerBody_Pants_Geo_LOD1", "LowerBody_Torso_Geo_LOD1", "Machine_Arm_DogBone_Geo_LOD1", "Machine_Arm_Elbow1_Bott_Geo_LOD1", "Machine_Arm_Elbow1_Top_Geo_LOD1", "Machine_Arm_ElbowPiston_Inner_Geo_LOD1", "Machine_Arm_ElbowPiston_Outer_Geo_LOD1", "Machine_Arm_ForeArm_Geo_LOD1", "Machine_Arm_Hand_Geo_LOD1", "Machine_Arm_HandFiller_Geo_LOD1", "Machine_Arm_KNmid_Geo_LOD1", "Machine_Arm_Pinky_Geo_LOD1", "Machine_Arm_PistonsMid_Geo_LOD1", "Machine_Arm_PistonUpper_Geo_LOD1", "Machine_Arm_PistonUpper_Inner_Geo_LOD1", "Machine_Arm_PistonUpperSmall_Geo_LOD1", "Machine_Arm_PistonUpperSmall_Inner_Geo_LOD1", "Machine_Arm_PistonUpperSmallBoot_Geo_LOD1", "Machine_Arm_Pointer_Geo_LOD1", "Machine_Arm_Thumb_Geo_LOD1", "Machine_Arm_Wires_Geo_LOD1", "Machine_Arm_WristPiston_Inner_Geo_LOD1", "Machine_Arm_WristPiston_Outer_Geo_LOD1", "Machine_Base_Ball_Geo_LOD1", "Machine_Base_Base_Geo_LOD1", "Machine_Base_Pipes_Geo_LOD1", "Machine_Base_ShoulderPadBottom_Geo_LOD1", "Machine_Base_ShoulderPadTop_Geo_LOD1", "Machine_Base_Wire_Big_Geo_LOD1", "Machine_Base_Wire_Small_Geo_LOD1", "Machine_Shoulder_Geo_LOD1", "spring_Geo_LOD1", "Teeth_Lower_Geo_LOD1", "TeethUpper_Geo_LOD1", "Tongue_Geo_LOD1", "turkey_Geo_LOD1", "UpperBody_Arms_Geo_LOD1", "UpperBody_EyeLashes_Geo_LOD1", "UpperBody_ShoulderPad_Geo_LOD1", "UpperBody_Torso_Geo_LOD1"]
+
+        jointsToDelete = []
+
+        # remove unneeded joints
+
+        # Delete the joints we no longer want.
+        for one in jointsToDelete:
+            if cmds.objExists(one):
+                cmds.delete(one)
+                
+        # Export the fbx file.
+        export.exportLOD("Gadget_LOD1.fbx", "LOD1")
+        
+        # Re-open the file without saving.
+        fullPath = cmds.file(q = True, sceneName = True)
+        cmds.file(fullPath, open=True, f=True)
+
+        print "LOD1 Exported Successfully!!!!!!!!!!!!!!!"
+
+GadgetLOD1()
+
+# Gadget LOD2 -------------------------------------------------------------------------------------------------------------
+def GadgetLOD2():
+    if saveResult == "Yes":
+
+        cmds.file(save=True)
+        
+        meshes = ["Corneas_Geo_LOD2", "Eyeballs_Geo_LOD2", "LowerBody_Ankle_Geo_LOD2", "LowerBody_Boots_Geo_LOD2", "LowerBody_Gas_Geo_LOD2", "LowerBody_Gren_Geo_LOD2", "LowerBody_KneePad_Geo_LOD2", "LowerBody_Pak_Geo_LOD2", "LowerBody_Pants_Geo_LOD2", "LowerBody_Torso_Geo_LOD2", "Machine_Arm_DogBone_Geo_LOD2", "Machine_Arm_Elbow1_Bott_Geo_LOD2", "Machine_Arm_Elbow1_Top_Geo_LOD2", "Machine_Arm_ElbowPiston_Inner_Geo_LOD2", "Machine_Arm_ElbowPiston_Outer_Geo_LOD2", "Machine_Arm_ForeArm_Geo_LOD2", "Machine_Arm_Hand_Geo_LOD2", "Machine_Arm_HandFiller_Geo_LOD2", "Machine_Arm_KNmid_Geo_LOD2", "Machine_Arm_Pinky_Geo_LOD2", "Machine_Arm_PistonsMid_Geo_LOD2", "Machine_Arm_PistonUpper_Geo_LOD2", "Machine_Arm_PistonUpper_Inner_Geo_LOD2", "Machine_Arm_PistonUpperSmall_Geo_LOD2", "Machine_Arm_PistonUpperSmall_Inner_Geo_LOD2", "Machine_Arm_PistonUpperSmallBoot_Geo_LOD2", "Machine_Arm_Pointer_Geo_LOD2", "Machine_Arm_Thumb_Geo_LOD2", "Machine_Arm_Wires_Geo_LOD2", "Machine_Arm_WristPiston_Inner_Geo_LOD2", "Machine_Arm_WristPiston_Outer_Geo_LOD2", "Machine_Base_Ball_Geo_LOD2", "Machine_Base_Base_Geo_LOD2", "Machine_Base_Pipes_Geo_LOD2", "Machine_Base_ShoulderPadBottom_Geo_LOD2", "Machine_Base_ShoulderPadTop_Geo_LOD2", "Machine_Base_Wire_Big_Geo_LOD2", "Machine_Base_Wire_Small_Geo_LOD2", "Machine_Shoulder_Geo_LOD2", "spring_Geo_LOD2", "Teeth_Lower_Geo_LOD2", "TeethUpper_Geo_LOD2", "Tongue_Geo_LOD2", "turkey_Geo_LOD2", "UpperBody_Arms_Geo_LOD2", "UpperBody_EyeLashes_Geo_LOD2", "UpperBody_ShoulderPad_Geo_LOD2", "UpperBody_Torso_Geo_LOD2"]
+
+        # Export a skin weights file for each mesh we plan to operate on.
+        cmds.select(meshes)
+        menu.exportSelectedSkinWeights(True)
+
+        # Set the pose if needed
+        cmds.setAttr("thumb_03_l.rotateZ", 23.02290762)
+        cmds.setAttr("ring_03_l.rotateZ", 29.40776524)
+        cmds.setAttr("pinky_03_l.rotateZ", 39.74899767)
+        cmds.setAttr("middle_03_l.rotateZ", 23.93256607)
+        cmds.setAttr("index_03_l.rotateZ", 16.68034218)
+
+        cmds.setAttr("thumb_03_r.rotateZ", 23.02290762)
+        cmds.setAttr("ring_03_r.rotateZ", 29.40776524)
+        cmds.setAttr("pinky_03_r.rotateZ", 39.74899767)
+        cmds.setAttr("middle_03_r.rotateZ", 23.93256607)
+        cmds.setAttr("index_03_r.rotateZ", 16.68034218)
+
+        # Delete the history on the meshes so that they get baked down into their posed positions
+        cmds.delete(meshes, ch=True)
+
+        # Move the skeleton back to its bindpose (This script REQUIRES the bind pose to be the pose that we want to export the skeleton to.  This is the case for all known characters created with A.R.T.)
+        cmds.dagPose("root", bp=True, r=True) 
+
+        # Re-import the skinning we exported before onto the mesh so thats its bound in its posed position.
+        cmds.select(meshes)
+        impw.importMultipleSkinWeights("", True)
+
+        jointsToDelete = []
+
+         # remove unneeded joints
+        jointsToRemove = ["arm_wire_01", "backpack_gas_01", "backpack_gas_02", "shield_bott", "shield_top", "spring"]
+        jointToTransferTo = "backpack_base"
+        transfer.LOD_transferWeights(meshes, jointsToRemove, jointToTransferTo)
+        jointsToDelete.extend(jointsToRemove)
+
+        jointsToRemove = ["robo_pinky_02"]
+        jointToTransferTo = "robo_pinky_01"
+        transfer.LOD_transferWeights(meshes, jointsToRemove, jointToTransferTo)
+        jointsToDelete.extend(jointsToRemove)
+
+        jointsToRemove = ["robo_pointer_02"]
+        jointToTransferTo = "robo_pointer_01"
+        transfer.LOD_transferWeights(meshes, jointsToRemove, jointToTransferTo)
+        jointsToDelete.extend(jointsToRemove)
+
+        jointsToRemove = ["robo_thumb_02"]
+        jointToTransferTo = "robo_thumb_01"
+        transfer.LOD_transferWeights(meshes, jointsToRemove, jointToTransferTo)
+        jointsToDelete.extend(jointsToRemove)
+
+        jointsToRemove = ["index_03_l"]
+        jointToTransferTo = "index_02_l"
+        transfer.LOD_transferWeights(meshes, jointsToRemove, jointToTransferTo)
+        jointsToDelete.extend(jointsToRemove)
+
+        jointsToRemove = ["middle_03_l"]
+        jointToTransferTo = "middle_02_l"
+        transfer.LOD_transferWeights(meshes, jointsToRemove, jointToTransferTo)
+        jointsToDelete.extend(jointsToRemove)
+
+        jointsToRemove = ["pinky_03_l"]
+        jointToTransferTo = "pinky_02_l"
+        transfer.LOD_transferWeights(meshes, jointsToRemove, jointToTransferTo)
+        jointsToDelete.extend(jointsToRemove)
+
+        jointsToRemove = ["ring_03_l"]
+        jointToTransferTo = "ring_02_l"
+        transfer.LOD_transferWeights(meshes, jointsToRemove, jointToTransferTo)
+        jointsToDelete.extend(jointsToRemove)
+
+        jointsToRemove = ["thumb_03_l"]
+        jointToTransferTo = "thumb_02_l"
+        transfer.LOD_transferWeights(meshes, jointsToRemove, jointToTransferTo)
+        jointsToDelete.extend(jointsToRemove)
+
+        jointsToRemove = ["index_03_r"]
+        jointToTransferTo = "index_02_r"
+        transfer.LOD_transferWeights(meshes, jointsToRemove, jointToTransferTo)
+        jointsToDelete.extend(jointsToRemove)
+
+        jointsToRemove = ["middle_03_r"]
+        jointToTransferTo = "middle_02_r"
+        transfer.LOD_transferWeights(meshes, jointsToRemove, jointToTransferTo)
+        jointsToDelete.extend(jointsToRemove)
+
+        jointsToRemove = ["pinky_03_r"]
+        jointToTransferTo = "pinky_02_r"
+        transfer.LOD_transferWeights(meshes, jointsToRemove, jointToTransferTo)
+        jointsToDelete.extend(jointsToRemove)
+
+        jointsToRemove = ["ring_03_r"]
+        jointToTransferTo = "ring_02_r"
+        transfer.LOD_transferWeights(meshes, jointsToRemove, jointToTransferTo)
+        jointsToDelete.extend(jointsToRemove)
+
+        jointsToRemove = ["thumb_03_r"]
+        jointToTransferTo = "thumb_02_r"
+        transfer.LOD_transferWeights(meshes, jointsToRemove, jointToTransferTo)
+        jointsToDelete.extend(jointsToRemove)
+
+        jointsToRemove = ["eyeball_l", "eyeball_r", "eyebrow_mid", "eyebrow_mid_l", "eyebrow_mid_r", "eyelid_bott_l", "eyelid_bott_r", "eyelid_top_l", "eyelid_top_r", "helmet_straps_l_01", "helmet_straps_l_02", "helmet_straps_r_01", "helmet_straps_r_02", "jaw"]
+        jointToTransferTo = "head"
+        transfer.LOD_transferWeights(meshes, jointsToRemove, jointToTransferTo)
+        jointsToDelete.extend(jointsToRemove)
+
+        jointsToRemove = ["thigh_pouch_big"]
+        jointToTransferTo = "thigh_l"
+        transfer.LOD_transferWeights(meshes, jointsToRemove, jointToTransferTo)
+        jointsToDelete.extend(jointsToRemove)
+
+        jointsToRemove = ["thigh_pouch_sml"]
+        jointToTransferTo = "thigh_r"
+        transfer.LOD_transferWeights(meshes, jointsToRemove, jointToTransferTo)
+        jointsToDelete.extend(jointsToRemove)
+
+        jointsToRemove = ["calf_twist_01_l"]
+        jointToTransferTo = "calf_l"
+        transfer.LOD_transferWeights(meshes, jointsToRemove, jointToTransferTo)
+        jointsToDelete.extend(jointsToRemove)
+
+        jointsToRemove = ["calf_twist_01_r"]
+        jointToTransferTo = "calf_r"
+        transfer.LOD_transferWeights(meshes, jointsToRemove, jointToTransferTo)
+        jointsToDelete.extend(jointsToRemove)
+        
+        # Delete the joints we no longer want.
+        for one in jointsToDelete:
+            if cmds.objExists(one):
+                cmds.delete(one)
+                
+        # Export the fbx file.
+        export.exportLOD("Gadget_LOD2.fbx", "LOD2")
+
+        # Re-open the file without saving.
+        fullPath = cmds.file(q = True, sceneName = True)
+        cmds.file(fullPath, open=True, f=True)
+
+        print "LOD2 Exported Successfully!!!!!!!!!!!!!!!"
+
+GadgetLOD2()
+
+# Gadget LOD3 -------------------------------------------------------------------------------------------------------------
+def GadgetLOD3():
+    if saveResult == "Yes":
+
+        cmds.file(save=True)
+    
+        meshes = ["Corneas_Geo_LOD3", "Eyeballs_Geo_LOD3", "LowerBody_Ankle_Geo_LOD3", "LowerBody_Boots_Geo_LOD3", "LowerBody_Gas_Geo_LOD3", "LowerBody_Gren_Geo_LOD3", "LowerBody_KneePad_Geo_LOD3", "LowerBody_Pak_Geo_LOD3", "LowerBody_Pants_Geo_LOD3", "LowerBody_Torso_Geo_LOD3", "Machine_Arm_DogBone_Geo_LOD3", "Machine_Arm_Elbow1_Bott_Geo_LOD3", "Machine_Arm_Elbow1_Top_Geo_LOD3", "Machine_Arm_ElbowPiston_Inner_Geo_LOD3", "Machine_Arm_ElbowPiston_Outer_Geo_LOD3", "Machine_Arm_ForeArm_Geo_LOD3", "Machine_Arm_Hand_Geo_LOD3", "Machine_Arm_HandFiller_Geo_LOD3", "Machine_Arm_KNmid_Geo_LOD3", "Machine_Arm_Pinky_Geo_LOD3", "Machine_Arm_PistonsMid_Geo_LOD3", "Machine_Arm_PistonUpper_Geo_LOD3", "Machine_Arm_PistonUpper_Inner_Geo_LOD3", "Machine_Arm_PistonUpperSmall_Geo_LOD3", "Machine_Arm_PistonUpperSmall_Inner_Geo_LOD3", "Machine_Arm_PistonUpperSmallBoot_Geo_LOD3", "Machine_Arm_Pointer_Geo_LOD3", "Machine_Arm_Thumb_Geo_LOD3", "Machine_Arm_Wires_Geo_LOD3", "Machine_Arm_WristPiston_Inner_Geo_LOD3", "Machine_Arm_WristPiston_Outer_Geo_LOD3", "Machine_Base_Ball_Geo_LOD3", "Machine_Base_Base_Geo_LOD3", "Machine_Base_Pipes_Geo_LOD3", "Machine_Base_ShoulderPadBottom_Geo_LOD3", "Machine_Base_ShoulderPadTop_Geo_LOD3", "Machine_Base_Wire_Big_Geo_LOD3", "Machine_Base_Wire_Small_Geo_LOD3", "Machine_Shoulder_Geo_LOD3", "spring_Geo_LOD3", "Teeth_Lower_Geo_LOD3", "TeethUpper_Geo_LOD3", "Tongue_Geo_LOD3", "turkey_Geo_LOD3", "UpperBody_Arms_Geo_LOD3", "UpperBody_EyeLashes_Geo_LOD3", "UpperBody_ShoulderPad_Geo_LOD3", "UpperBody_Torso_Geo_LOD3"]
+    
+        # Export a skin weights file for each mesh we plan to operate on.
+        cmds.select(meshes)
+        menu.exportSelectedSkinWeights(True)
+
+        # Set the pose if needed
+        cmds.setAttr("thumb_03_l.rotateZ", 23.02290762)
+        cmds.setAttr("thumb_02_l.rotateZ", 10.06960964)
+        cmds.setAttr("thumb_01_l.rotateZ", 8.338113921)
+        cmds.setAttr("thumb_01_l.rotateY", 3.002087702)
+        cmds.setAttr("thumb_01_l.rotateX", -3.231699475)
+        cmds.setAttr("ring_03_l.rotateZ", 29.40776524)
+        cmds.setAttr("ring_02_l.rotateZ", 29.40776524)
+        cmds.setAttr("ring_01_l.rotateZ", 29.40776524)
+        cmds.setAttr("pinky_03_l.rotateZ", 39.74899767)
+        cmds.setAttr("pinky_02_l.rotateZ", 39.74899767)
+        cmds.setAttr("pinky_01_l.rotateZ", 39.74899767)
+        cmds.setAttr("middle_03_l.rotateZ", 23.93256607)
+        cmds.setAttr("middle_02_l.rotateZ", 23.93256607)
+        cmds.setAttr("middle_01_l.rotateZ", 23.93256607)
+        cmds.setAttr("index_03_l.rotateZ", 16.68034218)
+        cmds.setAttr("index_02_l.rotateZ", 16.68034218)
+        cmds.setAttr("index_01_l.rotateZ", 16.68034218)
+
+        cmds.setAttr("thumb_03_r.rotateZ", 23.02290762)
+        cmds.setAttr("thumb_02_r.rotateZ", 10.06960964)
+        cmds.setAttr("thumb_01_r.rotateZ", 8.338113921)
+        cmds.setAttr("thumb_01_r.rotateY", 3.002087702)
+        cmds.setAttr("thumb_01_r.rotateX", -3.231699475)
+        cmds.setAttr("ring_03_r.rotateZ", 29.40776524)
+        cmds.setAttr("ring_02_r.rotateZ", 29.40776524)
+        cmds.setAttr("ring_01_r.rotateZ", 29.40776524)
+        cmds.setAttr("pinky_03_r.rotateZ", 39.74899767)
+        cmds.setAttr("pinky_02_r.rotateZ", 39.74899767)
+        cmds.setAttr("pinky_01_r.rotateZ", 39.74899767)
+        cmds.setAttr("middle_03_r.rotateZ", 23.93256607)
+        cmds.setAttr("middle_02_r.rotateZ", 23.93256607)
+        cmds.setAttr("middle_01_r.rotateZ", 23.93256607)
+        cmds.setAttr("index_03_r.rotateZ", 16.68034218)
+        cmds.setAttr("index_02_r.rotateZ", 16.68034218)
+        cmds.setAttr("index_01_r.rotateZ", 16.68034218)
+
+        # Delete the history on the meshes so that they get baked down into their posed positions
+        cmds.delete(meshes, ch=True)
+
+        # Move the skeleton back to its bindpose (This script REQUIRES the bind pose to be the pose that we want to export the skeleton to.  This is the case for all known characters created with A.R.T.)
+        cmds.dagPose("root", bp=True, r=True) 
+
+        # Re-import the skinning we exported before onto the mesh so thats its bound in its posed position.
+        cmds.select(meshes)
+        impw.importMultipleSkinWeights("", True)
+        
+        jointsToDelete = []
+
+        # remove unneeded joints
+        jointsToRemove = ["arm_wire_01", "backpack_gas_01", "backpack_gas_02", "shield_bott", "shield_top", "spring"]
+        jointToTransferTo = "backpack_base"
+        transfer.LOD_transferWeights(meshes, jointsToRemove, jointToTransferTo)
+        jointsToDelete.extend(jointsToRemove)
+
+        jointsToRemove = ["robo_uparm_pistonrod_01"]
+        jointToTransferTo = "robo_arm_01"
+        transfer.LOD_transferWeights(meshes, jointsToRemove, jointToTransferTo)
+        jointsToDelete.extend(jointsToRemove)
+
+        jointsToRemove = ["robo_uparm_pistonrod"]
+        jointToTransferTo = "robo_arm_02"
+        transfer.LOD_transferWeights(meshes, jointsToRemove, jointToTransferTo)
+        jointsToDelete.extend(jointsToRemove)
+
+        jointsToRemove = ["robo_lowarm_pistonrod_01"]
+        jointToTransferTo = "robo_arm_03"
+        transfer.LOD_transferWeights(meshes, jointsToRemove, jointToTransferTo)
+        jointsToDelete.extend(jointsToRemove)
+
+        jointsToRemove = ["robo_lowarm_pistonrod_02"]
+        jointToTransferTo = "robo_arm_05"
+        transfer.LOD_transferWeights(meshes, jointsToRemove, jointToTransferTo)
+        jointsToDelete.extend(jointsToRemove)
+
+        jointsToRemove = ["robo_pinky_02"]
+        jointToTransferTo = "robo_pinky_01"
+        transfer.LOD_transferWeights(meshes, jointsToRemove, jointToTransferTo)
+        jointsToDelete.extend(jointsToRemove)
+
+        jointsToRemove = ["robo_pointer_02"]
+        jointToTransferTo = "robo_pointer_01"
+        transfer.LOD_transferWeights(meshes, jointsToRemove, jointToTransferTo)
+        jointsToDelete.extend(jointsToRemove)
+
+        jointsToRemove = ["robo_thumb_02"]
+        jointToTransferTo = "robo_thumb_01"
+        transfer.LOD_transferWeights(meshes, jointsToRemove, jointToTransferTo)
+        jointsToDelete.extend(jointsToRemove)
+
+        jointsToRemove = ["upperarm_twist_01_l", "upperarm_twist_02_l"]
+        jointToTransferTo = "upperarm_l"
+        transfer.LOD_transferWeights(meshes, jointsToRemove, jointToTransferTo)
+        jointsToDelete.extend(jointsToRemove)
+
+        jointsToRemove = ["upperarm_twist_02_r", "upperarm_twist_01_r"]
+        jointToTransferTo = "upperarm_r"
+        transfer.LOD_transferWeights(meshes, jointsToRemove, jointToTransferTo)
+        jointsToDelete.extend(jointsToRemove)
+
+        jointsToRemove = ["lowerarm_twist_01_l"]
+        jointToTransferTo = "lowerarm_l"
+        transfer.LOD_transferWeights(meshes, jointsToRemove, jointToTransferTo)
+        jointsToDelete.extend(jointsToRemove)
+
+        jointsToRemove = ["lowerarm_twist_01_r"]
+        jointToTransferTo = "lowerarm_r"
+        transfer.LOD_transferWeights(meshes, jointsToRemove, jointToTransferTo)
+        jointsToDelete.extend(jointsToRemove)
+
+        jointsToRemove = ["index_01_l", "index_02_l", "index_03_l", "middle_01_l", "middle_02_l", "middle_03_l", "pinky_01_l", "pinky_02_l", "pinky_03_l", "ring_01_l", "ring_02_l", "ring_03_l", "thumb_01_l", "thumb_02_l", "thumb_03_l"]
+        jointToTransferTo = "hand_l"
+        transfer.LOD_transferWeights(meshes, jointsToRemove, jointToTransferTo)
+        jointsToDelete.extend(jointsToRemove)
+
+        jointsToRemove = ["index_01_r", "index_02_r", "index_03_r", "middle_01_r", "middle_02_r", "middle_03_r", "pinky_01_r", "pinky_02_r", "pinky_03_r", "ring_01_r", "ring_02_r", "ring_03_r", "thumb_01_r", "thumb_02_r", "thumb_03_r"]
+        jointToTransferTo = "hand_r"
+        transfer.LOD_transferWeights(meshes, jointsToRemove, jointToTransferTo)
+        jointsToDelete.extend(jointsToRemove)
+
+        jointsToRemove = ["eyeball_l", "eyeball_r", "eyebrow_mid", "eyebrow_mid_l", "eyebrow_mid_r", "eyelid_bott_l", "eyelid_bott_r", "eyelid_top_l", "eyelid_top_r", "helmet_straps_l_01", "helmet_straps_l_02", "helmet_straps_r_01", "helmet_straps_r_02", "jaw"]
+        jointToTransferTo = "head"
+        transfer.LOD_transferWeights(meshes, jointsToRemove, jointToTransferTo)
+        jointsToDelete.extend(jointsToRemove)
+
+        jointsToRemove = ["thigh_pouch_big", "thigh_twist_01_l", "thigh_twist_02_l"]
+        jointToTransferTo = "thigh_l"
+        transfer.LOD_transferWeights(meshes, jointsToRemove, jointToTransferTo)
+        jointsToDelete.extend(jointsToRemove)
+
+        jointsToRemove = ["thigh_pouch_sml", "thigh_twist_01_r", "thigh_twist_02_r"]
+        jointToTransferTo = "thigh_r"
+        transfer.LOD_transferWeights(meshes, jointsToRemove, jointToTransferTo)
+        jointsToDelete.extend(jointsToRemove)
+
+        jointsToRemove = ["calf_twist_01_l"]
+        jointToTransferTo = "calf_l"
+        transfer.LOD_transferWeights(meshes, jointsToRemove, jointToTransferTo)
+        jointsToDelete.extend(jointsToRemove)
+
+        jointsToRemove = ["calf_twist_01_r"]
+        jointToTransferTo = "calf_r"
+        transfer.LOD_transferWeights(meshes, jointsToRemove, jointToTransferTo)
+        jointsToDelete.extend(jointsToRemove)
+
+        # Delete the joints we no longer want.
+        for one in jointsToDelete:
+            if cmds.objExists(one):
+                cmds.delete(one)
+                
+        # Export the fbx file.
+        export.exportLOD("Gadget_LOD3.fbx", "LOD3")
+        
+        # Re-open the file without saving.
+        fullPath = cmds.file(q = True, sceneName = True)
+        cmds.file(fullPath, open=True, f=True)
+
+        print "LOD3 Exported Successfully!!!!!!!!!!!!!!!"
+
+GadgetLOD3()
+
+# Gadget LOD4 -------------------------------------------------------------------------------------------------------------
+def GadgetLOD4():
+    if saveResult == "Yes":
+
+        cmds.file(save=True)
+    
+        meshes = ["Corneas_Geo_LOD4", "Eyeballs_Geo_LOD4", "LowerBody_Ankle_Geo_LOD4", "LowerBody_Boots_Geo_LOD4", "LowerBody_Gas_Geo_LOD4", "LowerBody_Gren_Geo_LOD4", "LowerBody_KneePad_Geo_LOD4", "LowerBody_Pak_Geo_LOD4", "LowerBody_Pants_Geo_LOD4", "LowerBody_Torso_Geo_LOD4", "Machine_Arm_DogBone_Geo_LOD4", "Machine_Arm_Elbow1_Bott_Geo_LOD4", "Machine_Arm_Elbow1_Top_Geo_LOD4", "Machine_Arm_ElbowPiston_Inner_Geo_LOD4", "Machine_Arm_ElbowPiston_Outer_Geo_LOD4", "Machine_Arm_ForeArm_Geo_LOD4", "Machine_Arm_Hand_Geo_LOD4", "Machine_Arm_HandFiller_Geo_LOD4", "Machine_Arm_KNmid_Geo_LOD4", "Machine_Arm_Pinky_Geo_LOD4", "Machine_Arm_PistonsMid_Geo_LOD4", "Machine_Arm_PistonUpper_Geo_LOD4", "Machine_Arm_PistonUpper_Inner_Geo_LOD4", "Machine_Arm_PistonUpperSmall_Geo_LOD4", "Machine_Arm_PistonUpperSmall_Inner_Geo_LOD4", "Machine_Arm_PistonUpperSmallBoot_Geo_LOD4", "Machine_Arm_Pointer_Geo_LOD4", "Machine_Arm_Thumb_Geo_LOD4", "Machine_Arm_Wires_Geo_LOD4", "Machine_Arm_WristPiston_Inner_Geo_LOD4", "Machine_Arm_WristPiston_Outer_Geo_LOD4", "Machine_Base_Ball_Geo_LOD4", "Machine_Base_Base_Geo_LOD4", "Machine_Base_Pipes_Geo_LOD4", "Machine_Base_ShoulderPadBottom_Geo_LOD4", "Machine_Base_ShoulderPadTop_Geo_LOD4", "Machine_Base_Wire_Big_Geo_LOD4", "Machine_Base_Wire_Small_Geo_LOD4", "Machine_Shoulder_Geo_LOD4", "spring_Geo_LOD4", "Teeth_Lower_Geo_LOD4", "TeethUpper_Geo_LOD4", "Tongue_Geo_LOD4", "turkey_Geo_LOD4", "UpperBody_Arms_Geo_LOD4", "UpperBody_EyeLashes_Geo_LOD4", "UpperBody_ShoulderPad_Geo_LOD4", "UpperBody_Torso_Geo_LOD4"]
+
+        # Export a skin weights file for each mesh we plan to operate on.
+        cmds.select(meshes)
+        menu.exportSelectedSkinWeights(True)
+
+        # Set the pose if needed
+        cmds.setAttr("thumb_03_l.rotateZ", 23.02290762)
+        cmds.setAttr("thumb_02_l.rotateZ", 10.06960964)
+        cmds.setAttr("thumb_01_l.rotateZ", 8.338113921)
+        cmds.setAttr("thumb_01_l.rotateY", 3.002087702)
+        cmds.setAttr("thumb_01_l.rotateX", -3.231699475)
+        cmds.setAttr("ring_03_l.rotateZ", 29.40776524)
+        cmds.setAttr("ring_02_l.rotateZ", 29.40776524)
+        cmds.setAttr("ring_01_l.rotateZ", 29.40776524)
+        cmds.setAttr("pinky_03_l.rotateZ", 39.74899767)
+        cmds.setAttr("pinky_02_l.rotateZ", 39.74899767)
+        cmds.setAttr("pinky_01_l.rotateZ", 39.74899767)
+        cmds.setAttr("middle_03_l.rotateZ", 23.93256607)
+        cmds.setAttr("middle_02_l.rotateZ", 23.93256607)
+        cmds.setAttr("middle_01_l.rotateZ", 23.93256607)
+        cmds.setAttr("index_03_l.rotateZ", 16.68034218)
+        cmds.setAttr("index_02_l.rotateZ", 16.68034218)
+        cmds.setAttr("index_01_l.rotateZ", 16.68034218)
+
+        cmds.setAttr("thumb_03_r.rotateZ", 23.02290762)
+        cmds.setAttr("thumb_02_r.rotateZ", 10.06960964)
+        cmds.setAttr("thumb_01_r.rotateZ", 8.338113921)
+        cmds.setAttr("thumb_01_r.rotateY", 3.002087702)
+        cmds.setAttr("thumb_01_r.rotateX", -3.231699475)
+        cmds.setAttr("ring_03_r.rotateZ", 29.40776524)
+        cmds.setAttr("ring_02_r.rotateZ", 29.40776524)
+        cmds.setAttr("ring_01_r.rotateZ", 29.40776524)
+        cmds.setAttr("pinky_03_r.rotateZ", 39.74899767)
+        cmds.setAttr("pinky_02_r.rotateZ", 39.74899767)
+        cmds.setAttr("pinky_01_r.rotateZ", 39.74899767)
+        cmds.setAttr("middle_03_r.rotateZ", 23.93256607)
+        cmds.setAttr("middle_02_r.rotateZ", 23.93256607)
+        cmds.setAttr("middle_01_r.rotateZ", 23.93256607)
+        cmds.setAttr("index_03_r.rotateZ", 16.68034218)
+        cmds.setAttr("index_02_r.rotateZ", 16.68034218)
+        cmds.setAttr("index_01_r.rotateZ", 16.68034218)
+
+        # Delete the history on the meshes so that they get baked down into their posed positions
+        cmds.delete(meshes, ch=True)
+
+        # Move the skeleton back to its bindpose (This script REQUIRES the bind pose to be the pose that we want to export the skeleton to.  This is the case for all known characters created with A.R.T.)
+        cmds.dagPose("root", bp=True, r=True) 
+
+        # Re-import the skinning we exported before onto the mesh so thats its bound in its posed position.
+        cmds.select(meshes)
+        impw.importMultipleSkinWeights("", True)
+        
+        jointsToDelete = []
+
+        # remove unneeded joints
+        jointsToRemove = ["arm_wire_01", "backpack_gas_01", "backpack_gas_02", "shield_bott", "shield_top", "spring"]
+        jointToTransferTo = "backpack_base"
+        transfer.LOD_transferWeights(meshes, jointsToRemove, jointToTransferTo)
+        jointsToDelete.extend(jointsToRemove)
+
+        jointsToRemove = ["robo_uparm_pistonrod_01"]
+        jointToTransferTo = "robo_arm_01"
+        transfer.LOD_transferWeights(meshes, jointsToRemove, jointToTransferTo)
+        jointsToDelete.extend(jointsToRemove)
+
+        jointsToRemove = ["robo_uparm_pistonrod"]
+        jointToTransferTo = "robo_arm_02"
+        transfer.LOD_transferWeights(meshes, jointsToRemove, jointToTransferTo)
+        jointsToDelete.extend(jointsToRemove)
+
+        jointsToRemove = ["robo_lowarm_pistonrod_01"]
+        jointToTransferTo = "robo_arm_03"
+        transfer.LOD_transferWeights(meshes, jointsToRemove, jointToTransferTo)
+        jointsToDelete.extend(jointsToRemove)
+
+        jointsToRemove = ["robo_lowarm_pistonrod_02"]
+        jointToTransferTo = "robo_arm_05"
+        transfer.LOD_transferWeights(meshes, jointsToRemove, jointToTransferTo)
+        jointsToDelete.extend(jointsToRemove)
+
+        jointsToRemove = ["robo_pinky_01", "robo_pinky_02", "robo_pointer_01", "robo_pointer_02", "robo_thumb_01", "robo_thumb_02"]
+        jointToTransferTo = "robo_arm_07"
+        transfer.LOD_transferWeights(meshes, jointsToRemove, jointToTransferTo)
+        jointsToDelete.extend(jointsToRemove)
+
+        jointsToRemove = ["upperarm_twist_01_l", "upperarm_twist_02_l"]
+        jointToTransferTo = "upperarm_l"
+        transfer.LOD_transferWeights(meshes, jointsToRemove, jointToTransferTo)
+        jointsToDelete.extend(jointsToRemove)
+
+        jointsToRemove = ["upperarm_twist_02_r", "upperarm_twist_01_r"]
+        jointToTransferTo = "upperarm_r"
+        transfer.LOD_transferWeights(meshes, jointsToRemove, jointToTransferTo)
+        jointsToDelete.extend(jointsToRemove)
+
+        jointsToRemove = ["lowerarm_twist_01_l"]
+        jointToTransferTo = "lowerarm_l"
+        transfer.LOD_transferWeights(meshes, jointsToRemove, jointToTransferTo)
+        jointsToDelete.extend(jointsToRemove)
+
+        jointsToRemove = ["lowerarm_twist_01_r"]
+        jointToTransferTo = "lowerarm_r"
+        transfer.LOD_transferWeights(meshes, jointsToRemove, jointToTransferTo)
+        jointsToDelete.extend(jointsToRemove)
+
+        jointsToRemove = ["index_01_l", "index_02_l", "index_03_l", "middle_01_l", "middle_02_l", "middle_03_l", "pinky_01_l", "pinky_02_l", "pinky_03_l", "ring_01_l", "ring_02_l", "ring_03_l", "thumb_01_l", "thumb_02_l", "thumb_03_l"]
+        jointToTransferTo = "hand_l"
+        transfer.LOD_transferWeights(meshes, jointsToRemove, jointToTransferTo)
+        jointsToDelete.extend(jointsToRemove)
+
+        jointsToRemove = ["index_01_r", "index_02_r", "index_03_r", "middle_01_r", "middle_02_r", "middle_03_r", "pinky_01_r", "pinky_02_r", "pinky_03_r", "ring_01_r", "ring_02_r", "ring_03_r", "thumb_01_r", "thumb_02_r", "thumb_03_r"]
+        jointToTransferTo = "hand_r"
+        transfer.LOD_transferWeights(meshes, jointsToRemove, jointToTransferTo)
+        jointsToDelete.extend(jointsToRemove)
+
+        jointsToRemove = ["eyeball_l", "eyeball_r", "eyebrow_mid", "eyebrow_mid_l", "eyebrow_mid_r", "eyelid_bott_l", "eyelid_bott_r", "eyelid_top_l", "eyelid_top_r", "helmet_straps_l_01", "helmet_straps_l_02", "helmet_straps_r_01", "helmet_straps_r_02", "jaw"]
+        jointToTransferTo = "head"
+        transfer.LOD_transferWeights(meshes, jointsToRemove, jointToTransferTo)
+        jointsToDelete.extend(jointsToRemove)
+
+        jointsToRemove = ["thigh_pouch_big", "thigh_twist_01_l", "thigh_twist_02_l"]
+        jointToTransferTo = "thigh_l"
+        transfer.LOD_transferWeights(meshes, jointsToRemove, jointToTransferTo)
+        jointsToDelete.extend(jointsToRemove)
+
+        jointsToRemove = ["thigh_pouch_sml", "thigh_twist_01_r", "thigh_twist_02_r"]
+        jointToTransferTo = "thigh_r"
+        transfer.LOD_transferWeights(meshes, jointsToRemove, jointToTransferTo)
+        jointsToDelete.extend(jointsToRemove)
+
+        jointsToRemove = ["calf_twist_01_l"]
+        jointToTransferTo = "calf_l"
+        transfer.LOD_transferWeights(meshes, jointsToRemove, jointToTransferTo)
+        jointsToDelete.extend(jointsToRemove)
+
+        jointsToRemove = ["calf_twist_01_r"]
+        jointToTransferTo = "calf_r"
+        transfer.LOD_transferWeights(meshes, jointsToRemove, jointToTransferTo)
+        jointsToDelete.extend(jointsToRemove)
+
+        jointsToRemove = ["ball_l"]
+        jointToTransferTo = "foot_l"
+        transfer.LOD_transferWeights(meshes, jointsToRemove, jointToTransferTo)
+        jointsToDelete.extend(jointsToRemove)
+
+        jointsToRemove = ["ball_r"]
+        jointToTransferTo = "foot_r"
+        transfer.LOD_transferWeights(meshes, jointsToRemove, jointToTransferTo)
+        jointsToDelete.extend(jointsToRemove)
+        
+        # Delete the joints we no longer want.
+        for one in jointsToDelete:
+            if cmds.objExists(one):
+                cmds.delete(one)
+                
+        # Export the fbx file.
+        export.exportLOD("Gadget_LOD4.fbx", "LOD4")
+        
+        # Re-open the file without saving.
+        fullPath = cmds.file(q = True, sceneName = True)
+        cmds.file(fullPath, open=True, f=True)
+
+        print "LOD4 Exported Successfully!!!!!!!!!!!!!!!"
+
+GadgetLOD4()
+
+
diff --git a/Engine/Extras/Maya_AnimationRiggingTools/ARTv1/General/ART/Projects/Orion/Scripts/Gadget_Post_Script.py b/Engine/Extras/Maya_AnimationRiggingTools/ARTv1/General/ART/Projects/Orion/Scripts/Gadget_Post_Script.py
new file mode 100644
--- /dev/null
+++ b/Engine/Extras/Maya_AnimationRiggingTools/ARTv1/General/ART/Projects/Orion/Scripts/Gadget_Post_Script.py
@@ -0,0 +1,503 @@
+import maya.cmds as cmds
+import Tools.ART_CreateHydraulicRigs as hyd
+reload(hyd)
+import Tools.ART_addSpaceSwitching as space
+reload(space)
+
+# ------------------------------------------------------------------
+# scale all of the controls to be more visible with the mesh.
+def scaleControls():
+    sides = ["_l", "_r"]
+    for side in sides:
+        # feet
+        cmds.select(["ik_foot_anim"+side+".cv[0:17]"])
+        cmds.scale(1.0, 1.0, 1.0, r = True, ocp = False)
+        cmds.move(0, 0, -8.4, r = True)
+
+        cmds.select(["eyebrow_mid"+side+"_anim.cv[0:32]"])
+        cmds.scale(0.117744, 0.117744, 0.117744, r = True, ocp = False)
+        cmds.move(0, 0, 0, r = True)
+
+        cmds.select(["eyelid_bott"+side+"_anim.cv[0:32]"])
+        cmds.scale(0.1, 0.1, 0.01, r = True, ocp = False)
+        cmds.move(0, -3.757386, 0.654023, r = True)
+
+        cmds.select(["eyelid_top"+side+"_anim.cv[0:32]"])
+        cmds.scale(0.1, 0.1, 0.01, r = True, ocp = False)
+        cmds.move(0, -3.757386, 0.654023, r = True)
+
+
+        cmds.select(["fk_helmet_straps"+side+"_01_anim.cv[0:7]"])
+        cmds.scale(0.204714, 0.204714, 0.204714, r = True, ocp = False)
+        cmds.move(0, 0, 0, r = True)
+
+        cmds.select(["fk_helmet_straps"+side+"_02_anim.cv[0:7]"])
+        cmds.scale(0.204714, 0.204714, 0.204714, r = True, ocp = False)
+        cmds.move(0, 0, 0, r = True)
+
+    # Other
+    cmds.select(["eyeball_l_anim.cv[0:32]"])
+    cmds.scale(0.1, 0.1, 0.05, r = True, ocp = False)
+    cmds.move(0, -3.757386, 0, r = True)
+
+    cmds.select(["eyeball_r_anim.cv[0:32]"])
+    cmds.scale(0.1, 0.1, 0.05, r = True, ocp = False)
+    cmds.move(0, 3.757386, 0, r = True)
+
+    cmds.select(["jaw_anim.cv[0:32]"])
+    cmds.scale(0.2, 0.2, 0.2, r = True, ocp = True)
+    cmds.move(0, -8, -6, r = True)
+
+    cmds.select(["eyebrow_mid_anim.cv[0:32]"])
+    cmds.scale(0.117744, 0.117744, 0.117744, r = True, ocp = False)
+    cmds.move(0, 0, 0, r = True)
+
+    cmds.select(["backpack_gas_01_anim.cv[0:32]"])
+    cmds.scale(0.391111, 0.391111, 0.391111, r = True, ocp = True)
+    cmds.move(-7.1, 0, 0, r = True, os=True, wd=True)
+
+    cmds.select(["backpack_gas_02_anim.cv[0:32]"])
+    cmds.scale(0.391111, 0.391111, 0.391111, r = True, ocp = True)
+    cmds.move(-7.1, 0, 0, r = True, os=True, wd=True)
+
+    cmds.select(["thigh_pouch_sml_anim.cv[0:32]"])
+    cmds.scale(0.391111, 0.391111, 0.391111, r = True, ocp = True)
+    cmds.move(0, -6.6, 0, r = True, os=True, wd=True)
+
+    cmds.select(["thigh_pouch_big_anim.cv[0:32]"])
+    cmds.scale(0.391111, 0.391111, 0.391111, r = True, ocp = True)
+    cmds.move(0, -6.6, 0, r = True, os=True, wd=True)
+
+    cmds.select(["backpack_base_anim.cv[0:32]"])
+    cmds.scale(0.306684, 1, 0.306684, r = True, ocp = True)
+    cmds.move(0, 24.6, 0, r = True, os=True, wd=True)
+
+    cmds.select(["shield_bott_anim.cv[0:32]"])
+    cmds.scale(0.301673, 0.301673, 0.301673, r = True, ocp = True)
+    cmds.move(7.937765, 12.982729, -6.212037, r = True, os=True, wd=True)
+
+    cmds.select(["shield_top_anim.cv[0:32]"])
+    cmds.scale(0.301673, 0.301673, 0.301673, r = True, ocp = True)
+    cmds.move(7.788894, 16.519037, 10.864125, r = True, os=True, wd=True)
+
+    cmds.select(["spring_anim.cv[0:32]"])
+    cmds.scale(1.0, 1.0, 0.001, r = True, ocp = True)
+
+    cmds.select(["canister_main_anim.cv[0:32]"])
+    cmds.scale(1.252256, 1.252256, 1.252256, r = True, ocp = True)
+
+    cmds.select(["robo_uparm_pistonrod_anim.cv[0:32]"])
+    cmds.scale(0.1, 0.447259, 0.447259, r = True, ocp = True)
+
+    cmds.select(["robo_uparm_pistonrod_01_anim.cv[0:32]"])
+    cmds.scale(0.1, 0.35, 0.35, r = True, ocp = True)
+
+    cmds.select(["robo_lowarm_pistonrod_01_anim.cv[0:32]"])
+    cmds.scale(0.1, 0.45, 0.45, r = True, ocp = True)
+
+    cmds.select(["robo_lowarm_pistonrod_02_anim.cv[0:32]"])
+    cmds.scale(0.1, 0.677192, 0.677192, r = True, ocp = True)
+
+    cmds.select(["fk_robo_arm_01_anim.cv[0:7]"])
+    cmds.scale(1.772257, 1.772257, 1.772257, r = True, ocp = False)
+
+    cmds.select(["fk_robo_arm_02_anim.cv[0:7]"])
+    cmds.scale(0.699603, 10.699603, 0.699603, r = True, ocp = False)
+
+    cmds.select(["fk_robo_arm_03_anim.cv[0:7]"])
+    cmds.scale(1.772257, 1.772257, 1.772257, r = True, ocp = False)
+
+    cmds.select(["fk_robo_arm_04_anim.cv[0:7]"])
+    cmds.scale(0.699603, 0.699603, 0.699603, r = True, ocp = False)
+
+    cmds.select(["fk_robo_arm_05_anim.cv[0:7]"])
+    cmds.scale(1.772257, 1.772257, 1.772257, r = True, ocp = False)
+
+    cmds.select(["fk_robo_arm_06_anim.cv[0:7]"])
+    cmds.scale(0.699603, 0.699603, 0.699603, r = True, ocp = False)
+
+    cmds.select(["fk_robo_arm_07_anim.cv[0:7]"])
+    cmds.scale(1.772257, 1.772257, 1.772257, r = True, ocp = False)
+
+    cmds.select(["fk_robo_thumb_01_anim.cv[0:7]"])
+    cmds.scale(1.0, 1.0, 1.0, r = True, ocp = True)
+
+    cmds.select(["fk_robo_thumb_02_anim.cv[0:7]"])
+    cmds.scale(1.0, 1.0, 1.0, r = True, ocp = True)
+
+    cmds.select(["fk_robo_pointer_01_anim.cv[0:7]"])
+    cmds.scale(1.0, 1.0, 1.0, r = True, ocp = True)
+
+    cmds.select(["fk_robo_pointer_02_anim.cv[0:7]"])
+    cmds.scale(1.0, 1.0, 1.0, r = True, ocp = True)
+
+    cmds.select(["fk_robo_pinky_01_anim.cv[0:7]"])
+    cmds.scale(1.0, 1.0, 1.0, r = True, ocp = True)
+
+    cmds.select(["fk_robo_pinky_02_anim.cv[0:7]"])
+    cmds.scale(1.0, 1.0, 1.0, r = True, ocp = True)
+
+    cmds.select(["arm_wire_01_anim.cv[0:32]"])
+    cmds.scale(0.5, 0.5, 0.5, r = True, ocp = False)
+
+    cmds.select(["arm_wire_02_anim.cv[0:32]"])
+    cmds.scale(0.5, 0.5, 0.5, r = True, ocp = False)
+
+    cmds.select(None)
+
+scaleControls()
+
+# ------------------------------------------------------------------
+def buildRoboArmRigFK():
+    import Tools.ART_CreateHydraulicRigs as hyd
+    reload(hyd)
+
+    # Delete all of the dynamics for the arms as we dont need them.
+    cmds.delete("rig_dyn_robo_arm_01_HairControls","rig_dyn_robo_pinky_01_HairControls","rig_dyn_robo_pointer_01_HairControls","rig_dyn_robo_thumb_01_HairControls")
+
+    # Lock the attrs on the piston controls for the robo arm so they cant accidentally be animated.
+    for one in ["fk_robo_arm_06_anim", "fk_robo_arm_04_anim", "fk_robo_arm_02_anim"]:
+        cmds.setAttr(one+".tx", lock=True)
+        cmds.setAttr(one+".tz", lock=True)
+        cmds.setAttr(one+".rx", lock=True)
+        cmds.setAttr(one+".rz", lock=True)
+
+    for one in ["fk_robo_arm_03_anim", "fk_robo_arm_05_anim"]:
+        cmds.setAttr(one+".tx", lock=True)
+        cmds.setAttr(one+".ty", lock=True)
+        cmds.setAttr(one+".tz", lock=True)
+        cmds.setAttr(one+".ry", lock=True)
+        cmds.setAttr(one+".rz", lock=True)
+
+    cmds.select("robo_lowarm_pistonrod_01_anim", "robo_lowarm_pistonrod_02_anim", "driver_robo_arm_05")
+    hyd.CreateHydraulicRig(-1)
+
+    cmds.select("robo_uparm_pistonrod_01_anim", "robo_uparm_pistonrod_anim", "driver_robo_arm_02")
+    hyd.CreateHydraulicRig(-1)
+
+    # Create a slide mechanism on the forearm part of the lowerarm piston.
+    cmds.addAttr("robo_lowarm_pistonrod_02_anim", ln = "Slide", dv = 0, min=0, max=10, keyable = True)
+    cmds.setDrivenKeyframe("robo_lowarm_pistonrod_02_anim_grp.translateX", cd="robo_lowarm_pistonrod_02_anim.Slide", itt="spline", ott="spline")
+    cmds.setDrivenKeyframe("robo_lowarm_pistonrod_02_anim_grp.translateZ", cd="robo_lowarm_pistonrod_02_anim.Slide", itt="spline", ott="spline")
+
+    cmds.setAttr("robo_lowarm_pistonrod_02_anim.Slide", 10)
+    cmds.setAttr("robo_lowarm_pistonrod_02_anim_grp.tx", -6.662443)
+    cmds.setAttr("robo_lowarm_pistonrod_02_anim_grp.tz", -5.510393)
+
+    cmds.setDrivenKeyframe("robo_lowarm_pistonrod_02_anim_grp.translateX", cd="robo_lowarm_pistonrod_02_anim.Slide", itt="spline", ott="spline")
+    cmds.setDrivenKeyframe("robo_lowarm_pistonrod_02_anim_grp.translateZ", cd="robo_lowarm_pistonrod_02_anim.Slide", itt="spline", ott="spline")
+    cmds.setAttr("robo_lowarm_pistonrod_02_anim.Slide", 0)
+    
+    cmds.setAttr("robo_lowarm_pistonrod_02_anim_grp.ry", -28)
+buildRoboArmRigFK()
+
+
+# ------------------------------------------------------------------
+def buildRoboArmRigIK():
+    import Modules.ART_rigUtils as utils
+    reload(utils)
+    import Tools.ART_addSpaceSwitching as space
+    reload(space)
+
+    # create an empty group for the Roboarm rig 
+    RoboArmGrp = cmds.group(em=True, name="RoboArm_grp")
+
+    # Create an IK arm control
+    ik_robowrist_anim = utils.createControl("sphere", 15, "ik_robowrist_anim")
+    constraint = cmds.parentConstraint("driver_robo_arm_07", ik_robowrist_anim)[0]
+    cmds.delete(constraint)
+    cmds.makeIdentity(ik_robowrist_anim, r = 1, s = 1, apply =True)
+    cmds.setAttr(ik_robowrist_anim + ".sx", 1.8)
+    cmds.setAttr(ik_robowrist_anim + ".sy", 1.8)
+    cmds.setAttr(ik_robowrist_anim + ".sz", 0.3)
+    cmds.makeIdentity(ik_robowrist_anim, r = 1, s = 1, apply =True)
+
+    ik_robowrist_anim_grp = cmds.group(empty=True, name=ik_robowrist_anim + "_grp")
+    constraint = cmds.pointConstraint("driver_robo_arm_07", ik_robowrist_anim_grp)[0]
+    cmds.delete(constraint)
+
+    ik_robowrist_animSpaceSwitchFollow = cmds.duplicate(ik_robowrist_anim_grp, po=True, n=ik_robowrist_anim+"_space_switcher_follow")[0]
+    ik_robowrist_animSpaceSwitch = cmds.duplicate(ik_robowrist_anim_grp, po=True, n=ik_robowrist_anim+"_space_switcher")[0]
+
+    cmds.parent(ik_robowrist_animSpaceSwitchFollow, RoboArmGrp)
+    cmds.parent(ik_robowrist_animSpaceSwitch, ik_robowrist_animSpaceSwitchFollow)
+    cmds.parent(ik_robowrist_anim_grp, ik_robowrist_animSpaceSwitch)
+
+    cmds.parent(ik_robowrist_anim, ik_robowrist_anim_grp)
+    cmds.makeIdentity(ik_robowrist_anim_grp, t=1, r=1, s=1, apply=True)
+
+    # create a pole vector control for the elbow
+    ik_roboelbow_anim = utils.createControl("sphere", 6, "ik_roboelbow_anim")
+    constraint = cmds.parentConstraint("driver_robo_arm_03", ik_roboelbow_anim)[0]
+    cmds.delete(constraint)
+    cmds.makeIdentity(ik_roboelbow_anim, t=1, r=1, s=1, apply=True)
+    cmds.setAttr(ik_roboelbow_anim + ".tz", -30)
+    cmds.makeIdentity(ik_roboelbow_anim, t=1, r=1, s=1, apply=True)	    
+    #cmds.parent(ik_roboelbowx_anim, RoboArmGrp)
+
+    ik_roboelbow_anim_grp = cmds.group(empty=True, name=ik_roboelbow_anim + "_grp")
+    constraint = cmds.pointConstraint(ik_roboelbow_anim, ik_roboelbow_anim_grp)[0]
+    cmds.delete(constraint)
+
+    ik_roboelbow_animSpaceSwitchFollow = cmds.duplicate(ik_roboelbow_anim_grp, po=True, n=ik_roboelbow_anim+"_space_switcher_follow")[0]
+    ik_roboelbow_animSpaceSwitch = cmds.duplicate(ik_roboelbow_anim_grp, po=True, n=ik_roboelbow_anim+"_space_switcher")[0]
+
+    cmds.parent(ik_roboelbow_animSpaceSwitchFollow, RoboArmGrp)
+    cmds.parent(ik_roboelbow_animSpaceSwitch, ik_roboelbow_animSpaceSwitchFollow)
+    cmds.parent(ik_roboelbow_anim_grp, ik_roboelbow_animSpaceSwitch)
+
+    cmds.parent(ik_roboelbow_anim, ik_roboelbow_anim_grp)
+    cmds.makeIdentity(ik_roboelbow_anim_grp, t=1, r=1, s=1, apply=True)
+
+    # Duplicate the driver Robo Arm to make a new IK robo arm
+    roboarm_01_joint = cmds.duplicate("driver_robo_arm_01", name="rig_robo_arm_01_ik", po=True)[0]
+    cmds.parent(roboarm_01_joint, RoboArmGrp)
+    roboarm_03_joint = cmds.duplicate("driver_robo_arm_03", name="rig_robo_arm_03_ik", po=True)[0]
+    cmds.parent(roboarm_03_joint, roboarm_01_joint)
+    roboarm_05_joint = cmds.duplicate("driver_robo_arm_05", name="rig_robo_arm_05_ik", po=True)[0]
+    cmds.parent(roboarm_05_joint, roboarm_03_joint)
+    roboarm_07_joint = cmds.duplicate("driver_robo_arm_07", name="rig_robo_arm_07_ik", po=True)[0]
+    cmds.parent(roboarm_07_joint, roboarm_05_joint)
+
+    # create an ikHandle for the arm and parent it under the IK arm control
+    roboarmIKHandle = cmds.ikHandle(name="roboarmIKHandle", sol="ikRPsolver", sj=roboarm_01_joint, ee=roboarm_07_joint)[0]
+    cmds.parent(roboarmIKHandle, ik_robowrist_anim)
+
+    # create a pole vector constraint for the ikhandle to the pv control
+    cmds.poleVectorConstraint(ik_roboelbow_anim, roboarmIKHandle)
+
+    # Orient Constrain the last joint in the Roboarm IK chain to the IK control for hand control.
+    cmds.orientConstraint(ik_robowrist_anim, roboarm_07_joint, mo=True)
+    
+    # Parent constrain the first joint of the IK chain to the driver skeleton.
+    cmds.parentConstraint("driver_backpack_base", roboarm_01_joint, mo=True)
+
+    # parent the roboarm group under the ctrl_rig node
+    cmds.parent(RoboArmGrp, "ctrl_rig")
+
+    # create new parent constraints with MO on for each of the original driver joints to be following the new ik joints
+    for one in ["01", "03", "05", "07"]:
+        cmds.parentConstraint("rig_robo_arm_"+one+"_ik", "driver_robo_arm_"+one, mo=True)
+        
+    Rig_SettingsConnection = cmds.listConnections("Rig_Settings.robo_arm_ik", p=True)
+    for one in Rig_SettingsConnection:
+        cmds.disconnectAttr("Rig_Settings.robo_arm_ik", one)
+
+    cmds.connectAttr("Rig_Settings.robo_arm_ik", "driver_robo_arm_01_parentConstraint1.rig_robo_arm_01_ikW3")
+    cmds.connectAttr("Rig_Settings.robo_arm_ik", "driver_robo_arm_03_parentConstraint1.rig_robo_arm_03_ikW3")
+    cmds.connectAttr("Rig_Settings.robo_arm_ik", "driver_robo_arm_05_parentConstraint1.rig_robo_arm_05_ikW3")
+    cmds.connectAttr("Rig_Settings.robo_arm_ik", "driver_robo_arm_07_parentConstraint1.rig_robo_arm_07_ikW3")
+    cmds.connectAttr("Rig_Settings.robo_arm_ik", "RoboArm_grp.visibility")
+
+    space.addSpaceSwitching(ik_robowrist_anim, "master_anim", False)
+    space.addSpaceSwitching(ik_robowrist_anim, "body_anim", False)
+    space.addSpaceSwitching(ik_robowrist_anim, "backpack_base_anim", False)
+
+    space.addSpaceSwitching(ik_roboelbow_anim, "master_anim", False)
+    space.addSpaceSwitching(ik_roboelbow_anim, "body_anim", False)
+    space.addSpaceSwitching(ik_roboelbow_anim, "backpack_base_anim", False)
+
+    cmds.setAttr("Rig_Settings.robo_arm_ik", 1)
+    cmds.setAttr("Rig_Settings.robo_arm_fk", 0)
+
+buildRoboArmRigIK()
+
+# ------------------------------------------------------------------
+def buildBackpackRig():
+    # Move the spring into position.
+    cmds.setAttr("spring_anim.ty", 0.058)
+    cmds.setAttr("spring_anim.tz", -28.7)
+    cmds.setAttr("spring_anim.sx", 0.749)
+    cmds.setAttr("spring_anim.sy", 0.749)
+    cmds.setAttr("spring_anim.sz", 0.081)
+    
+    # Move the canister into position
+    cmds.setAttr("canister_main_anim.ty", 1)
+    cmds.setAttr("canister_main_anim.tz", -49.4)
+    cmds.setAttr("canister_main_anim.sx", 0.585)
+    cmds.setAttr("canister_main_anim.sy", 0.585)
+    cmds.setAttr("canister_main_anim.sz", 0.314)
+    
+    # Rig the gas cans to be staying between the pelvis and the backpack base.
+    gascan_02_const = cmds.parentConstraint("driver_pelvis", "backpack_gas_02_parent_grp", mo=True)[0]
+    cmds.setAttr(gascan_02_const+".interpType", 2)
+    gascan_01_const = cmds.parentConstraint("driver_backpack_gas_02", "backpack_gas_01_parent_grp", mo=True)[0]
+    cmds.setAttr(gascan_01_const+".interpType", 2)
+
+buildBackpackRig()
+
+# ------------------------------------------------------------------
+def dynamicsOff():
+    for one in ["thigh_pouch_big_anim", "thigh_pouch_sml_anim", "helmet_straps_l_dyn_anim", "helmet_straps_r_dyn_anim"]:
+        cmds.setAttr(one+".chainStartEnvelope", 0)
+        cmds.setAttr(one+".chainStartEnvelope", l=True)
+dynamicsOff()
+
+# ------------------------------------------------------------------
+def helmetStraps():
+    const = cmds.parentConstraint("driver_jaw", "helmet_straps_l_ik_base_anim_grp", mo=True)[0]
+    cmds.setAttr(const+".interpType", 2)
+    const = cmds.parentConstraint("driver_jaw", "helmet_straps_r_ik_base_anim_grp", mo=True)[0]
+    cmds.setAttr(const+".interpType", 2)
+    const = cmds.parentConstraint("driver_jaw", "fk_helmet_straps_l_01_grp", mo=True)[0]
+    cmds.setAttr(const+".interpType", 2)
+    const = cmds.parentConstraint("driver_jaw", "fk_helmet_straps_r_01_grp", mo=True)[0]
+    cmds.setAttr(const+".interpType", 2)
+helmetStraps()
+
+# ------------------------------------------------------------------
+def invertEyelids():
+    sides = ["_l", "_r"]
+    for side in sides:
+        cmds.delete("driver_eyelid_bott"+side+"_parentConstraint1")
+        cmds.rotate(0, -180, 0, "eyelid_bott"+side+"_anim_grp", r=True, os=True)
+        cmds.parentConstraint("eyelid_bott"+side+"_anim", "driver_eyelid_bott"+side, mo=True)
+invertEyelids()
+
+def invertRightEyeball():
+    cmds.delete("driver_eyeball_r_parentConstraint1")
+    cmds.rotate(0, -180, -180, "eyeball_r_anim_grp", r=True, os=True)
+    cmds.parentConstraint("eyeball_r_anim", "driver_eyeball_r", mo=True)
+invertRightEyeball()
+
+
+##################################################################################
+##################################################################################
+
+def addCapsule(radius, halfHeight):
+    trueHalfHeight = (halfHeight*2)-(radius*2)
+    capsule = cmds.polyCylinder(name="collision_capsule", r=radius, h=trueHalfHeight, sx=18, sy=1, sz=5, ax=[0,0,1], rcp=True, cuv=3, ch=True)
+    cmds.setAttr(capsule[0]+".tz", halfHeight+2.2)
+    cmds.setAttr(capsule[0]+".rz", 90)
+
+    cmds.addAttr(capsule[0], ln="Radius", dv=radius, keyable=True)
+    cmds.addAttr(capsule[0], ln="HalfHeight", dv=halfHeight, keyable=True)
+
+    expr = capsule[1]+".radius = "+capsule[0]+".Radius;\n"+capsule[1]+".height = ("+capsule[0]+".HalfHeight*2) - ("+capsule[0]+".Radius*2);"
+    cmds.expression(s=expr, o=capsule[1], ae=1, uc="all")
+
+    cmds.setAttr(capsule[0]+".overrideEnabled", 1)
+    cmds.setAttr(capsule[0]+".overrideColor", 18)
+    cmds.setAttr(capsule[0]+".overrideShading", 0)
+
+    cmds.addAttr(capsule[0], ln="Shaded", dv=0, min=0, max=1, keyable=True)
+    cmds.connectAttr(capsule[0]+".Shaded", capsule[0]+".overrideShading")
+
+    cmds.parentConstraint("root", capsule[0], mo=True)
+    cmds.select(capsule[0], r=True)
+    capsuleDisplayLayer = cmds.createDisplayLayer(nr=True, name="Collision Capsule")
+    cmds.setAttr(capsuleDisplayLayer+".displayType", 2)
+    cmds.setAttr(capsuleDisplayLayer+".v", 0)
+    
+    for attr in ["tx", "ty", "tz", "rx", "ry", "rz", "sx", "sy", "sz"]:
+        cmds.setAttr(capsule[0]+"."+attr, lock=True, keyable=False, channelBox=False)
+
+    cmds.select(clear=True)
+addCapsule(42, 85)
+
+# ------------------------------------------------------------------
+def createGameCam():
+    gameCam = cmds.camera(centerOfInterest=5, focalLength=18, lensSqueezeRatio=1, cameraScale=1, horizontalFilmAperture=1.41732, horizontalFilmOffset=0, verticalFilmAperture=0.94488, verticalFilmOffset=0, filmFit="Overscan", overscan=1.05, motionBlur=0, shutterAngle=144, nearClipPlane=1.0, farClipPlane=10000, orthographic=0, orthographicWidth=30, panZoomEnabled=0, horizontalPan=0, verticalPan=0, zoom=1, displayFilmGate=True)[0]
+    print gameCam
+    gameCam = cmds.rename(gameCam, "GameCam")
+    cmds.camera(gameCam, e=True, horizontalFilmAperture=1.6724376, focalLength=21.2399857011)
+    cmds.setAttr(gameCam+".displayGateMaskOpacity", 1)
+    cmds.setAttr(gameCam+".displayGateMaskColor", 0, 0, 0, type="double3")
+    cmds.setAttr(gameCam+".r", 90, 0, 180, type="double3")
+    
+    # Create the control
+    control = cmds.circle(name=gameCam+"_Ctrl", ch=True, o=True, r=15)[0]
+    cmds.makeIdentity(control, a=True, t=True, r=True, s=True)
+    # Create groups for the sdk and the constraints
+    const_grp = cmds.group(control, name=control+"_grp")
+    if cmds.objExists("rig_grp"):
+        cmds.parent(const_grp, "rig_grp")
+    # Parent Constrain the cam to the control
+    cmds.parentConstraint(control, gameCam, mo=True)
+    
+    cmds.connectAttr(control+".visibility", gameCam+".visibility", f=True)
+    cmds.setAttr(const_grp+".v", 0)
+    
+    cmds.setAttr(const_grp+".t", 0, 500, 200, type="double3")
+createGameCam()
+# ------------------------------------------------------------------
+
+# ------------------------------------------------------------------
+# add ik bones Function that is called later after we open the export file.
+def addIKBones():
+    try:
+		# create the joints
+		cmds.select(clear = True)
+		ikFootRoot = cmds.joint(name = "ik_foot_root")
+		cmds.select(clear = True)
+		
+		cmds.select(clear = True)
+		ikFootLeft = cmds.joint(name = "ik_foot_l")
+		cmds.select(clear = True)
+		
+		cmds.select(clear = True)
+		ikFootRight = cmds.joint(name = "ik_foot_r")
+		cmds.select(clear = True)
+		
+		cmds.select(clear = True)
+		ikHandRoot = cmds.joint(name = "ik_hand_root")
+		cmds.select(clear = True)
+		
+		cmds.select(clear = True)
+		ikHandGun = cmds.joint(name = "ik_hand_gun")
+		cmds.select(clear = True)
+		
+		cmds.select(clear = True)
+		ikHandLeft = cmds.joint(name = "ik_hand_l")
+		cmds.select(clear = True)
+		
+		cmds.select(clear = True)
+		ikHandRight = cmds.joint(name = "ik_hand_r")
+		cmds.select(clear = True)
+		
+		# create hierarhcy
+		cmds.parent(ikFootRoot, "root")
+		cmds.parent(ikHandRoot, "root")
+		cmds.parent(ikFootLeft, ikFootRoot)
+		cmds.parent(ikFootRight, ikFootRoot)
+		cmds.parent(ikHandGun, ikHandRoot)
+		cmds.parent(ikHandLeft, ikHandGun)
+		cmds.parent(ikHandRight, ikHandGun)
+		
+		# constrain the joints
+		leftHandConstraint = cmds.parentConstraint("hand_l", ikHandLeft)[0]
+		rightHandGunConstraint = cmds.parentConstraint("hand_r", ikHandGun)[0]
+		rightHandConstraint = cmds.parentConstraint("hand_r", ikHandRight)[0]
+		leftFootConstraint = cmds.parentConstraint("foot_l", ikFootLeft)[0]
+		rightFootConstraint = cmds.parentConstraint("foot_r", ikFootRight)[0]
+	
+    except:
+		cmds.warning("Something went wrong")
+
+# ------------------------------------------------------------------
+# ADD the IK bones to the Export file.
+#save the AnimRig file before opening the export file.
+cmds.file(save = True, type = "mayaBinary", force = True)
+
+#ADD IK BONES
+#open export file
+sceneName = cmds.file(q = True, sceneName = True)
+exportFile = sceneName.partition("AnimRigs")[0] + "ExportFiles" +  sceneName.partition("AnimRigs")[2].partition(".mb")[0] + "_Export.mb"
+cmds.file(exportFile, open = True, force = True)
+
+#add ik bones
+addIKBones()
+
+# ------------------------------------------------------------------
+def addCustomAttrs():
+    attrs = cmds.listAttr("root")
+    exists = False
+    for one in attrs:
+        if one.find("DistanceCurve") != -1:
+            exists = True
+    if exists == False:
+        cmds.addAttr("root", ln="DistanceCurve", keyable=True)
+addCustomAttrs()
+
+#save the export file.
+cmds.file(save = True, type = "mayaBinary", force = True)
\ No newline at end of file
