Branching

diff --git a/Engine/Extras/Maya_AnimationRiggingTools/ARTv1/General/ART/Projects/Orion/Scripts/Grux_CreateLODs.py b/Engine/Extras/Maya_AnimationRiggingTools/ARTv1/General/ART/Projects/Orion/Scripts/Grux_CreateLODs.py
new file mode 100644
--- /dev/null
+++ b/Engine/Extras/Maya_AnimationRiggingTools/ARTv1/General/ART/Projects/Orion/Scripts/Grux_CreateLODs.py
@@ -0,0 +1,510 @@
+import maya.cmds as cmds
+import customMayaMenu as menu
+import Tools.ART_ImportMultipleWeights as impw
+import Tools.ART_TransferWeights as transfer
+import Tools.ART_ExportLODtoFBX as export
+reload(export)
+
+################################################## GRUX #########################################################################
+# Warn the user that this operation will save their file.
+saveResult = cmds.confirmDialog(title = "Save Warning", message="In order to continue your file must be saved.  Would you like to save it?  If yes it will be saved and after the operation is complete your file will be re-opened.", button = ["Yes", "Cancel"], defaultButton='Yes', cancelButton='Cancel', dismissString='Cancel')
+
+# If they are ok with saving the file, warn them that it will also delete all of the current files in their skin weights directory.
+if saveResult == "Yes":
+    deleteSkinFilesResult = cmds.confirmDialog(title = "SkinWeights Warning", message="In order to continue we must delete all of the files in your ..\General\ART\SkinWeights\ directory.  Would you like to continue?.", button = ["Yes", "Cancel"], defaultButton='Yes', cancelButton='Cancel', dismissString='Cancel')
+    if deleteSkinFilesResult == "Yes":
+        toolsPath = cmds.internalVar(usd = True) + "mayaTools.txt"
+        if os.path.exists(toolsPath):
+            f = open(toolsPath, 'r')
+            mayaToolsDir = f.readline()
+            f.close()
+
+        path = mayaToolsDir + "\General\ART\SkinWeights\\"
+        files = os.listdir(path)
+        #delete all weight cache files in the system folder
+        for file in files:
+            os.remove(path + file)
+            #print path + file
+
+
+
+# Grux LOD1 -------------------------------------------------------------------------------------------------------------
+def GruxLOD1():
+    if saveResult == "Yes":
+
+        cmds.file(save=True)
+        
+        meshes = ["belt_geo_LOD1", "big_packs_geo_LOD1", "body_geo_LOD1", "eye_tear_duct_geo_LOD1", "eye_wetness_geo_LOD1", "eyeballs_geo_LOD1", "eyeshadow_geo_LOD1", "horn_ring_geo_LOD1", "hornwrap_01_geo_LOD1", "hornwrap_02_geo_LOD1", "loin_cloth_middle_geo_LOD1", "loin_cloth_sides_geo_LOD1", "small_packs_geo_LOD1", "teeth_geo_LOD1", "underwear_geo_LOD1", "weapon_left_geo_LOD1", "weapon_right_geo_LOD1", "wristwrap_geo_LOD1"]
+
+        jointsToDelete = []
+
+        # remove unneeded joints
+        jointsToRemove = ["tongue_01", "tongue_02", "tongue_03"]
+        jointToTransferTo = "head"
+        transfer.LOD_transferWeights(meshes, jointsToRemove, jointToTransferTo)
+        jointsToDelete.extend(jointsToRemove)
+
+        
+        # Delete the joints we no longer want.
+        for one in jointsToDelete:
+            if cmds.objExists(one):
+                cmds.delete(one)
+                
+        # Export the fbx file.
+        export.exportLOD("Grux_LOD1.fbx", "LOD1")
+        
+        # Re-open the file without saving.
+        fullPath = cmds.file(q = True, sceneName = True)
+        cmds.file(fullPath, open=True, f=True)
+
+        print "LOD1 Exported Successfully!!!!!!!!!!!!!!!"
+
+GruxLOD1()
+
+# Grux LOD2 -------------------------------------------------------------------------------------------------------------
+def GruxLOD2():
+    if saveResult == "Yes":
+        
+        cmds.file(save=True)
+        
+        meshes = ["belt_geo_LOD2", "big_packs_geo_LOD2", "body_geo_LOD2", "eye_tear_duct_geo_LOD2", "eye_wetness_geo_LOD2", "eyeballs_geo_LOD2", "eyeshadow_geo_LOD2", "horn_ring_geo_LOD2", "hornwrap_01_geo_LOD2", "hornwrap_02_geo_LOD2", "loin_cloth_middle_geo_LOD2", "loin_cloth_sides_geo_LOD2", "small_packs_geo_LOD2", "teeth_geo_LOD2", "underwear_geo_LOD2", "weapon_left_geo_LOD2", "weapon_right_geo_LOD2", "wristwrap_geo_LOD2"]
+
+        # Export a skin weights file for each mesh we plan to operate on.
+        cmds.select(meshes)
+        menu.exportSelectedSkinWeights(True)
+
+        # Set the pose if needed
+        cmds.setAttr("thumb_03_l.rotateZ", -1.887)
+        cmds.setAttr("thumb_03_l.rotateY", -1.789)
+        cmds.setAttr("thumb_03_l.rotateX", -2.92)
+        cmds.setAttr("thumb_02_l.rotateZ", 8.435338581)
+        cmds.setAttr("middle_03_l.rotateZ", 56.68)
+        cmds.setAttr("middle_02_l.rotateZ", 70.719)
+        cmds.setAttr("index_03_l.rotateZ", 56.687)
+        cmds.setAttr("index_02_l.rotateZ", 56.687)
+
+        cmds.setAttr("thumb_03_r.rotateZ", -1.887)
+        cmds.setAttr("thumb_03_r.rotateY", -1.789)
+        cmds.setAttr("thumb_03_r.rotateX", -2.92)
+        cmds.setAttr("thumb_02_r.rotateZ", 8.435338581)
+        cmds.setAttr("middle_03_r.rotateZ", 56.68)
+        cmds.setAttr("middle_02_r.rotateZ", 70.719)
+        cmds.setAttr("index_03_r.rotateZ", 56.687)
+        cmds.setAttr("index_02_r.rotateZ", 56.687)
+
+        # Delete the history on the meshes so that they get baked down into their posed positions
+        cmds.delete(meshes, ch=True)
+
+        # Move the skeleton back to its bindpose (This script REQUIRES the bind pose to be the pose that we want to export the skeleton to.  This is the case for all known characters created with A.R.T.)
+        cmds.dagPose("root", bp=True, r=True) 
+
+        # Re-import the skinning we exported before onto the mesh so thats its bound in its posed position.
+        cmds.select(meshes)
+        impw.importMultipleSkinWeights("", True)
+        
+        jointsToDelete = []
+
+         # remove unneeded joints
+        jointsToRemove = ["belt", "big_pack_l", "big_pack_r", "loin_cloth_bk_l_01", "loin_cloth_bk_l_02", "loin_cloth_bk_l_03", "loin_cloth_bk_r_01", "loin_cloth_bk_r_02", "loin_cloth_bk_r_03", "loin_cloth_fr_l_01", "loin_cloth_fr_l_02", "loin_cloth_fr_l_03", "loin_cloth_fr_r_01", "loin_cloth_fr_r_02", "loin_cloth_fr_r_03", "small_pack_l", "small_pack_r"]
+        jointToTransferTo = "pelvis"
+        transfer.LOD_transferWeights(meshes, jointsToRemove, jointToTransferTo)
+        jointsToDelete.extend(jointsToRemove)
+
+        jointsToRemove = ["muscle_gut"]
+        jointToTransferTo = "spine_02"
+        transfer.LOD_transferWeights(meshes, jointsToRemove, jointToTransferTo)
+        jointsToDelete.extend(jointsToRemove)
+
+        jointsToRemove = ["muscle_pec_l", "muscle_pec_r"]
+        jointToTransferTo = "spine_03"
+        transfer.LOD_transferWeights(meshes, jointsToRemove, jointToTransferTo)
+        jointsToDelete.extend(jointsToRemove)
+
+        jointsToRemove = ["scapula_l"]
+        jointToTransferTo = "clavicle_l"
+        transfer.LOD_transferWeights(meshes, jointsToRemove, jointToTransferTo)
+        jointsToDelete.extend(jointsToRemove)
+
+        jointsToRemove = ["scapula_r"]
+        jointToTransferTo = "clavicle_r"
+        transfer.LOD_transferWeights(meshes, jointsToRemove, jointToTransferTo)
+        jointsToDelete.extend(jointsToRemove)
+
+        jointsToRemove = ["upperarm_twist_03_l"]
+        jointToTransferTo = "upperarm_l"
+        transfer.LOD_transferWeights(meshes, jointsToRemove, jointToTransferTo)
+        jointsToDelete.extend(jointsToRemove)
+
+        jointsToRemove = ["upperarm_twist_03_r"]
+        jointToTransferTo = "upperarm_r"
+        transfer.LOD_transferWeights(meshes, jointsToRemove, jointToTransferTo)
+        jointsToDelete.extend(jointsToRemove)
+
+        jointsToRemove = ["index_02_l", "index_03_l"]
+        jointToTransferTo = "index_01_l"
+        transfer.LOD_transferWeights(meshes, jointsToRemove, jointToTransferTo)
+        jointsToDelete.extend(jointsToRemove)
+
+        jointsToRemove = ["middle_02_l", "middle_03_l"]
+        jointToTransferTo = "middle_01_l"
+        transfer.LOD_transferWeights(meshes, jointsToRemove, jointToTransferTo)
+        jointsToDelete.extend(jointsToRemove)
+
+        jointsToRemove = ["thumb_02_l", "thumb_03_l"]
+        jointToTransferTo = "thumb_01_l"
+        transfer.LOD_transferWeights(meshes, jointsToRemove, jointToTransferTo)
+        jointsToDelete.extend(jointsToRemove)
+
+        jointsToRemove = ["index_02_r", "index_03_r"]
+        jointToTransferTo = "index_01_r"
+        transfer.LOD_transferWeights(meshes, jointsToRemove, jointToTransferTo)
+        jointsToDelete.extend(jointsToRemove)
+
+        jointsToRemove = ["middle_02_r", "middle_03_r"]
+        jointToTransferTo = "middle_01_r"
+        transfer.LOD_transferWeights(meshes, jointsToRemove, jointToTransferTo)
+        jointsToDelete.extend(jointsToRemove)
+
+        jointsToRemove = ["thumb_02_r", "thumb_03_r"]
+        jointToTransferTo = "thumb_01_r"
+        transfer.LOD_transferWeights(meshes, jointsToRemove, jointToTransferTo)
+        jointsToDelete.extend(jointsToRemove)
+
+        jointsToRemove = ["ear_l", "horn_ring", "ear_r", "eyeball_l", "eyeball_r", "eyelid_bott_l", "eyelid_bott_r", "eyelid_top_l", "eyelid_top_r", "tongue_01", "tongue_02", "tongue_03"]
+        jointToTransferTo = "head"
+        transfer.LOD_transferWeights(meshes, jointsToRemove, jointToTransferTo)
+        jointsToDelete.extend(jointsToRemove)
+
+        # Delete the joints we no longer want.
+        for one in jointsToDelete:
+            if cmds.objExists(one):
+                cmds.delete(one)
+                
+        # Export the fbx file.
+        export.exportLOD("Grux_LOD2.fbx", "LOD2")
+
+        # Re-open the file without saving.
+        fullPath = cmds.file(q = True, sceneName = True)
+        cmds.file(fullPath, open=True, f=True)
+
+        print "LOD2 Exported Successfully!!!!!!!!!!!!!!!"
+
+GruxLOD2()
+
+# Grux LOD3 -------------------------------------------------------------------------------------------------------------
+def GruxLOD3():
+    if saveResult == "Yes":
+
+        cmds.file(save=True)
+    
+        meshes = ["belt_geo_LOD3", "big_packs_geo_LOD3", "body_geo_LOD3", "eye_tear_duct_geo_LOD3", "eye_wetness_geo_LOD3", "eyeballs_geo_LOD3", "eyeshadow_geo_LOD3", "horn_ring_geo_LOD3", "hornwrap_01_geo_LOD3", "hornwrap_02_geo_LOD3", "loin_cloth_middle_geo_LOD3", "loin_cloth_sides_geo_LOD3", "small_packs_geo_LOD3", "teeth_geo_LOD3", "underwear_geo_LOD3", "weapon_left_geo_LOD3", "weapon_right_geo_LOD3", "wristwrap_geo_LOD3"]
+    
+        # Export a skin weights file for each mesh we plan to operate on.
+        cmds.select(meshes)
+        menu.exportSelectedSkinWeights(True)
+
+        # Set the pose if needed
+        cmds.setAttr("thumb_03_l.rotateZ", -1.887)
+        cmds.setAttr("thumb_03_l.rotateY", -1.789)
+        cmds.setAttr("thumb_03_l.rotateX", -2.92)
+        cmds.setAttr("thumb_02_l.rotateZ", 8.435338581)
+        cmds.setAttr("thumb_01_l.rotateZ", 3.76)
+        cmds.setAttr("thumb_01_l.rotateY", -4.45)
+        cmds.setAttr("thumb_01_l.rotateX", 0.74)
+        cmds.setAttr("middle_03_l.rotateZ", 56.68)
+        cmds.setAttr("middle_02_l.rotateZ", 70.719)
+        cmds.setAttr("middle_01_l.rotateZ", 56.687)
+        cmds.setAttr("middle_metacarpal_l.rotateY", 2.497)
+        cmds.setAttr("index_03_l.rotateZ", 56.687)
+        cmds.setAttr("index_02_l.rotateZ", 56.687)
+        cmds.setAttr("index_01_l.rotateZ", 56.687)
+        cmds.setAttr("index_metacarpal_l.rotateZ", 1.0784)
+        cmds.setAttr("index_metacarpal_l.rotateY", -1.0436)
+        cmds.setAttr("index_metacarpal_l.rotateX", 1.5033)
+
+        cmds.setAttr("thumb_03_r.rotateZ", -1.887)
+        cmds.setAttr("thumb_03_r.rotateY", -1.789)
+        cmds.setAttr("thumb_03_r.rotateX", -2.92)
+        cmds.setAttr("thumb_02_r.rotateZ", 8.435338581)
+        cmds.setAttr("thumb_01_r.rotateZ", 3.76)
+        cmds.setAttr("thumb_01_r.rotateY", -4.45)
+        cmds.setAttr("thumb_01_r.rotateX", 0.74)
+        cmds.setAttr("middle_03_r.rotateZ", 56.68)
+        cmds.setAttr("middle_02_r.rotateZ", 70.719)
+        cmds.setAttr("middle_01_r.rotateZ", 56.687)
+        cmds.setAttr("middle_metacarpal_r.rotateY", 2.497)
+        cmds.setAttr("index_03_r.rotateZ", 56.687)
+        cmds.setAttr("index_02_r.rotateZ", 56.687)
+        cmds.setAttr("index_01_r.rotateZ", 56.687)
+        cmds.setAttr("index_metacarpal_r.rotateZ", 1.0784)
+        cmds.setAttr("index_metacarpal_r.rotateY", -1.0436)
+        cmds.setAttr("index_metacarpal_r.rotateX", 1.5033)
+
+        # Delete the history on the meshes so that they get baked down into their posed positions
+        cmds.delete(meshes, ch=True)
+
+        # Move the skeleton back to its bindpose (This script REQUIRES the bind pose to be the pose that we want to export the skeleton to.  This is the case for all known characters created with A.R.T.)
+        cmds.dagPose("root", bp=True, r=True) 
+
+        # Re-import the skinning we exported before onto the mesh so thats its bound in its posed position.
+        cmds.select(meshes)
+        impw.importMultipleSkinWeights("", True)
+        
+        jointsToDelete = []
+
+        # remove unneeded joints
+        jointsToRemove = ["belt", "big_pack_l", "big_pack_r", "loin_cloth_bk_l_01", "loin_cloth_bk_l_02", "loin_cloth_bk_l_03", "loin_cloth_bk_r_01", "loin_cloth_bk_r_02", "loin_cloth_bk_r_03", "loin_cloth_fr_l_01", "loin_cloth_fr_l_02", "loin_cloth_fr_l_03", "loin_cloth_fr_r_01", "loin_cloth_fr_r_02", "loin_cloth_fr_r_03", "small_pack_l", "small_pack_r"]
+        jointToTransferTo = "pelvis"
+        transfer.LOD_transferWeights(meshes, jointsToRemove, jointToTransferTo)
+        jointsToDelete.extend(jointsToRemove)
+
+        jointsToRemove = ["muscle_gut"]
+        jointToTransferTo = "spine_02"
+        transfer.LOD_transferWeights(meshes, jointsToRemove, jointToTransferTo)
+        jointsToDelete.extend(jointsToRemove)
+
+        jointsToRemove = ["muscle_pec_l", "muscle_pec_r"]
+        jointToTransferTo = "spine_03"
+        transfer.LOD_transferWeights(meshes, jointsToRemove, jointToTransferTo)
+        jointsToDelete.extend(jointsToRemove)
+
+        jointsToRemove = ["scapula_l"]
+        jointToTransferTo = "clavicle_l"
+        transfer.LOD_transferWeights(meshes, jointsToRemove, jointToTransferTo)
+        jointsToDelete.extend(jointsToRemove)
+
+        jointsToRemove = ["scapula_r"]
+        jointToTransferTo = "clavicle_r"
+        transfer.LOD_transferWeights(meshes, jointsToRemove, jointToTransferTo)
+        jointsToDelete.extend(jointsToRemove)
+
+        jointsToRemove = ["upperarm_twist_01_l", "upperarm_twist_02_l", "upperarm_twist_03_l"]
+        jointToTransferTo = "upperarm_l"
+        transfer.LOD_transferWeights(meshes, jointsToRemove, jointToTransferTo)
+        jointsToDelete.extend(jointsToRemove)
+
+        jointsToRemove = ["upperarm_twist_01_r", "upperarm_twist_02_r", "upperarm_twist_03_r"]
+        jointToTransferTo = "upperarm_r"
+        transfer.LOD_transferWeights(meshes, jointsToRemove, jointToTransferTo)
+        jointsToDelete.extend(jointsToRemove)
+
+        jointsToRemove = ["lowerarm_twist_01_l"]
+        jointToTransferTo = "lowerarm_l"
+        transfer.LOD_transferWeights(meshes, jointsToRemove, jointToTransferTo)
+        jointsToDelete.extend(jointsToRemove)
+
+        jointsToRemove = ["lowerarm_twist_01_r"]
+        jointToTransferTo = "lowerarm_r"
+        transfer.LOD_transferWeights(meshes, jointsToRemove, jointToTransferTo)
+        jointsToDelete.extend(jointsToRemove)
+
+        jointsToRemove = ["index_metacarpal_l", "index_01_l", "index_02_l", "index_03_l", "middle_metacarpal_l", "middle_01_l", "middle_02_l", "middle_03_l", "thumb_01_l", "thumb_02_l", "thumb_03_l"]
+        jointToTransferTo = "hand_l"
+        transfer.LOD_transferWeights(meshes, jointsToRemove, jointToTransferTo)
+        jointsToDelete.extend(jointsToRemove)
+
+        jointsToRemove = ["index_metacarpal_r", "index_01_r", "index_02_r", "index_03_r", "middle_metacarpal_r", "middle_01_r", "middle_02_r", "middle_03_r", "thumb_01_r", "thumb_02_r", "thumb_03_r"]
+        jointToTransferTo = "hand_r"
+        transfer.LOD_transferWeights(meshes, jointsToRemove, jointToTransferTo)
+        jointsToDelete.extend(jointsToRemove)
+
+        jointsToRemove = ["ear_l", "horn_ring", "ear_r", "eyeball_l", "eyeball_r", "eyelid_bott_l", "eyelid_bott_r", "eyelid_top_l", "eyelid_top_r", "jaw", "tongue_01", "tongue_02", "tongue_03"]
+        jointToTransferTo = "head"
+        transfer.LOD_transferWeights(meshes, jointsToRemove, jointToTransferTo)
+        jointsToDelete.extend(jointsToRemove)
+
+        jointsToRemove = ["thigh_twist_01_l", "thigh_twist_02_l"]
+        jointToTransferTo = "thigh_l"
+        transfer.LOD_transferWeights(meshes, jointsToRemove, jointToTransferTo)
+        jointsToDelete.extend(jointsToRemove)
+
+        jointsToRemove = ["thigh_twist_01_r", "thigh_twist_02_r"]
+        jointToTransferTo = "thigh_r"
+        transfer.LOD_transferWeights(meshes, jointsToRemove, jointToTransferTo)
+        jointsToDelete.extend(jointsToRemove)
+
+        jointsToRemove = ["calf_twist_01_l"]
+        jointToTransferTo = "calf_l"
+        transfer.LOD_transferWeights(meshes, jointsToRemove, jointToTransferTo)
+        jointsToDelete.extend(jointsToRemove)
+
+        jointsToRemove = ["calf_twist_01_r"]
+        jointToTransferTo = "calf_r"
+        transfer.LOD_transferWeights(meshes, jointsToRemove, jointToTransferTo)
+        jointsToDelete.extend(jointsToRemove)
+
+        # Delete the joints we no longer want.
+        for one in jointsToDelete:
+            if cmds.objExists(one):
+                cmds.delete(one)
+                
+        # Export the fbx file.
+        export.exportLOD("Grux_LOD3.fbx", "LOD3")
+        
+        # Re-open the file without saving.
+        fullPath = cmds.file(q = True, sceneName = True)
+        cmds.file(fullPath, open=True, f=True)
+
+        print "LOD3 Exported Successfully!!!!!!!!!!!!!!!"
+
+GruxLOD3()
+
+# Grux LOD4 -------------------------------------------------------------------------------------------------------------
+def GruxLOD4():
+    if saveResult == "Yes":
+
+        cmds.file(save=True)
+    
+        meshes = ["belt_geo_LOD4", "big_packs_geo_LOD4", "body_geo_LOD4", "eye_tear_duct_geo_LOD4", "eye_wetness_geo_LOD4", "eyeballs_geo_LOD4", "eyeshadow_geo_LOD4", "horn_ring_geo_LOD4", "hornwrap_01_geo_LOD4", "hornwrap_02_geo_LOD4", "loin_cloth_middle_geo_LOD4", "loin_cloth_sides_geo_LOD4", "small_packs_geo_LOD4", "teeth_geo_LOD4", "underwear_geo_LOD4", "weapon_left_geo_LOD4", "weapon_right_geo_LOD4", "wristwrap_geo_LOD4"]
+
+        # Export a skin weights file for each mesh we plan to operate on.
+        cmds.select(meshes)
+        menu.exportSelectedSkinWeights(True)
+
+        # Set the pose if needed
+        cmds.setAttr("thumb_03_l.rotateZ", -1.887)
+        cmds.setAttr("thumb_03_l.rotateY", -1.789)
+        cmds.setAttr("thumb_03_l.rotateX", -2.92)
+        cmds.setAttr("thumb_02_l.rotateZ", 8.435338581)
+        cmds.setAttr("thumb_01_l.rotateZ", 3.76)
+        cmds.setAttr("thumb_01_l.rotateY", -4.45)
+        cmds.setAttr("thumb_01_l.rotateX", 0.74)
+        cmds.setAttr("middle_03_l.rotateZ", 56.68)
+        cmds.setAttr("middle_02_l.rotateZ", 70.719)
+        cmds.setAttr("middle_01_l.rotateZ", 56.687)
+        cmds.setAttr("middle_metacarpal_l.rotateY", 2.497)
+        cmds.setAttr("index_03_l.rotateZ", 56.687)
+        cmds.setAttr("index_02_l.rotateZ", 56.687)
+        cmds.setAttr("index_01_l.rotateZ", 56.687)
+        cmds.setAttr("index_metacarpal_l.rotateZ", 1.0784)
+        cmds.setAttr("index_metacarpal_l.rotateY", -1.0436)
+        cmds.setAttr("index_metacarpal_l.rotateX", 1.5033)
+
+        cmds.setAttr("thumb_03_r.rotateZ", -1.887)
+        cmds.setAttr("thumb_03_r.rotateY", -1.789)
+        cmds.setAttr("thumb_03_r.rotateX", -2.92)
+        cmds.setAttr("thumb_02_r.rotateZ", 8.435338581)
+        cmds.setAttr("thumb_01_r.rotateZ", 3.76)
+        cmds.setAttr("thumb_01_r.rotateY", -4.45)
+        cmds.setAttr("thumb_01_r.rotateX", 0.74)
+        cmds.setAttr("middle_03_r.rotateZ", 56.68)
+        cmds.setAttr("middle_02_r.rotateZ", 70.719)
+        cmds.setAttr("middle_01_r.rotateZ", 56.687)
+        cmds.setAttr("middle_metacarpal_r.rotateY", 2.497)
+        cmds.setAttr("index_03_r.rotateZ", 56.687)
+        cmds.setAttr("index_02_r.rotateZ", 56.687)
+        cmds.setAttr("index_01_r.rotateZ", 56.687)
+        cmds.setAttr("index_metacarpal_r.rotateZ", 1.0784)
+        cmds.setAttr("index_metacarpal_r.rotateY", -1.0436)
+        cmds.setAttr("index_metacarpal_r.rotateX", 1.5033)
+
+        # Delete the history on the meshes so that they get baked down into their posed positions
+        cmds.delete(meshes, ch=True)
+
+        # Move the skeleton back to its bindpose (This script REQUIRES the bind pose to be the pose that we want to export the skeleton to.  This is the case for all known characters created with A.R.T.)
+        cmds.dagPose("root", bp=True, r=True) 
+
+        # Re-import the skinning we exported before onto the mesh so thats its bound in its posed position.
+        cmds.select(meshes)
+        impw.importMultipleSkinWeights("", True)
+        
+        jointsToDelete = []
+
+        # remove unneeded joints
+        jointsToRemove = ["belt", "big_pack_l", "big_pack_r", "loin_cloth_bk_l_01", "loin_cloth_bk_l_02", "loin_cloth_bk_l_03", "loin_cloth_bk_r_01", "loin_cloth_bk_r_02", "loin_cloth_bk_r_03", "loin_cloth_fr_l_01", "loin_cloth_fr_l_02", "loin_cloth_fr_l_03", "loin_cloth_fr_r_01", "loin_cloth_fr_r_02", "loin_cloth_fr_r_03", "small_pack_l", "small_pack_r"]
+        jointToTransferTo = "pelvis"
+        transfer.LOD_transferWeights(meshes, jointsToRemove, jointToTransferTo)
+        jointsToDelete.extend(jointsToRemove)
+
+        jointsToRemove = ["muscle_gut"]
+        jointToTransferTo = "spine_02"
+        transfer.LOD_transferWeights(meshes, jointsToRemove, jointToTransferTo)
+        jointsToDelete.extend(jointsToRemove)
+
+        jointsToRemove = ["muscle_pec_l", "muscle_pec_r"]
+        jointToTransferTo = "spine_03"
+        transfer.LOD_transferWeights(meshes, jointsToRemove, jointToTransferTo)
+        jointsToDelete.extend(jointsToRemove)
+
+        jointsToRemove = ["scapula_l"]
+        jointToTransferTo = "clavicle_l"
+        transfer.LOD_transferWeights(meshes, jointsToRemove, jointToTransferTo)
+        jointsToDelete.extend(jointsToRemove)
+
+        jointsToRemove = ["scapula_r"]
+        jointToTransferTo = "clavicle_r"
+        transfer.LOD_transferWeights(meshes, jointsToRemove, jointToTransferTo)
+        jointsToDelete.extend(jointsToRemove)
+
+        jointsToRemove = ["upperarm_twist_01_l", "upperarm_twist_02_l", "upperarm_twist_03_l"]
+        jointToTransferTo = "upperarm_l"
+        transfer.LOD_transferWeights(meshes, jointsToRemove, jointToTransferTo)
+        jointsToDelete.extend(jointsToRemove)
+
+        jointsToRemove = ["upperarm_twist_01_r", "upperarm_twist_02_r", "upperarm_twist_03_r"]
+        jointToTransferTo = "upperarm_r"
+        transfer.LOD_transferWeights(meshes, jointsToRemove, jointToTransferTo)
+        jointsToDelete.extend(jointsToRemove)
+
+        jointsToRemove = ["lowerarm_twist_01_l"]
+        jointToTransferTo = "lowerarm_l"
+        transfer.LOD_transferWeights(meshes, jointsToRemove, jointToTransferTo)
+        jointsToDelete.extend(jointsToRemove)
+
+        jointsToRemove = ["lowerarm_twist_01_r"]
+        jointToTransferTo = "lowerarm_r"
+        transfer.LOD_transferWeights(meshes, jointsToRemove, jointToTransferTo)
+        jointsToDelete.extend(jointsToRemove)
+
+        jointsToRemove = ["index_metacarpal_l", "index_01_l", "index_02_l", "index_03_l", "middle_metacarpal_l", "middle_01_l", "middle_02_l", "middle_03_l", "thumb_01_l", "thumb_02_l", "thumb_03_l"]
+        jointToTransferTo = "hand_l"
+        transfer.LOD_transferWeights(meshes, jointsToRemove, jointToTransferTo)
+        jointsToDelete.extend(jointsToRemove)
+
+        jointsToRemove = ["index_metacarpal_r", "index_01_r", "index_02_r", "index_03_r", "middle_metacarpal_r", "middle_01_r", "middle_02_r", "middle_03_r", "thumb_01_r", "thumb_02_r", "thumb_03_r"]
+        jointToTransferTo = "hand_r"
+        transfer.LOD_transferWeights(meshes, jointsToRemove, jointToTransferTo)
+        jointsToDelete.extend(jointsToRemove)
+
+        jointsToRemove = ["ear_l", "horn_ring", "ear_r", "eyeball_l", "eyeball_r", "eyelid_bott_l", "eyelid_bott_r", "eyelid_top_l", "eyelid_top_r", "jaw", "tongue_01", "tongue_02", "tongue_03"]
+        jointToTransferTo = "head"
+        transfer.LOD_transferWeights(meshes, jointsToRemove, jointToTransferTo)
+        jointsToDelete.extend(jointsToRemove)
+
+        jointsToRemove = ["thigh_twist_01_l", "thigh_twist_02_l"]
+        jointToTransferTo = "thigh_l"
+        transfer.LOD_transferWeights(meshes, jointsToRemove, jointToTransferTo)
+        jointsToDelete.extend(jointsToRemove)
+
+        jointsToRemove = ["thigh_twist_01_r", "thigh_twist_02_r"]
+        jointToTransferTo = "thigh_r"
+        transfer.LOD_transferWeights(meshes, jointsToRemove, jointToTransferTo)
+        jointsToDelete.extend(jointsToRemove)
+
+        jointsToRemove = ["calf_twist_01_l"]
+        jointToTransferTo = "calf_l"
+        transfer.LOD_transferWeights(meshes, jointsToRemove, jointToTransferTo)
+        jointsToDelete.extend(jointsToRemove)
+
+        jointsToRemove = ["calf_twist_01_r"]
+        jointToTransferTo = "calf_r"
+        transfer.LOD_transferWeights(meshes, jointsToRemove, jointToTransferTo)
+        jointsToDelete.extend(jointsToRemove)
+
+        # Delete the joints we no longer want.
+        for one in jointsToDelete:
+            if cmds.objExists(one):
+                cmds.delete(one)
+                
+        # Export the fbx file.
+        export.exportLOD("Grux_LOD4.fbx", "LOD4")
+        
+        # Re-open the file without saving.
+        fullPath = cmds.file(q = True, sceneName = True)
+        cmds.file(fullPath, open=True, f=True)
+
+        print "LOD4 Exported Successfully!!!!!!!!!!!!!!!"
+
+GruxLOD4()
+
+
diff --git a/Engine/Extras/Maya_AnimationRiggingTools/ARTv1/General/ART/Projects/Orion/Scripts/Grux_Post_Script.py b/Engine/Extras/Maya_AnimationRiggingTools/ARTv1/General/ART/Projects/Orion/Scripts/Grux_Post_Script.py
new file mode 100644
--- /dev/null
+++ b/Engine/Extras/Maya_AnimationRiggingTools/ARTv1/General/ART/Projects/Orion/Scripts/Grux_Post_Script.py
@@ -0,0 +1,338 @@
+import maya.cmds as cmds
+
+# ------------------------------------------------------------------
+# scale all of the controls to be more visible with the mesh.
+def scaleControls():
+    for side in ["_l", "_r"]:
+        # feet
+        cmds.select(["ik_foot_anim"+side+".cv[0:17]"])
+        cmds.scale(1.5, 0.8, 1.0, r = True, ocp = False)
+        cmds.move(0, 4.678845, -11.947602, r = True)
+
+        cmds.select(["ik_wrist"+side+"_anim.cv[0:7]"])
+        cmds.scale(1.9, 1.9, 1.9, r = True, ocp = False)
+
+        cmds.select(["weapon"+side+"_anim.cv[0:32]"])
+        cmds.scale(0.330274, 0.330274, 0.330274, r = True, ocp = False)
+        if side == "_l":
+            cmds.move(25, 0, 0, r = True)
+        else:
+            cmds.move(-25, 0, 0, r = True)            
+
+        cmds.select(["eyelid_bott"+side+"_anim.cv[0:32]"])
+        cmds.scale(0.1, 0.1, 0.01, r = True, ocp = False)
+        if side is "_l":
+            cmds.move(5, 0, 0, r = True, os=True)
+        else:
+            cmds.move(-5, 0, 0, r = True, os=True)
+
+        cmds.select(["eyelid_top"+side+"_anim.cv[0:32]"])
+        cmds.scale(0.1, 0.1, 0.01, r = True, ocp = False)
+        if side is "_l":
+            cmds.move(5, 0, 0, r = True, os=True)
+        else:
+            cmds.move(-5, 0, 0, r = True, os=True)
+
+        cmds.select(["eyeball"+side+"_anim.cv[0:32]"])
+        cmds.scale(0.24, 0.24, 0.24, r = True, ocp = False)
+        cmds.move(0, 0, 0, r = True, os=True)
+
+        cmds.select(["ear"+side+"_anim.cv[0:32]"])
+        cmds.scale(1.49, 0.67, 0.43, r = True, ocp = False)
+
+    # Other
+    cmds.select(["jaw_anim.cv[0:32]"])
+    cmds.scale(0.2, 0.2, 0.2, r = True, ocp = True)
+    cmds.move(0, -13, -19, r = True)
+
+    cmds.select(["fk_tongue_01_anim.cv[0:32]"])
+    cmds.scale(0.246972, 0.246972, 0.246972, r = True, ocp = False)
+    cmds.move(0, 0, 0, r = True)
+    cmds.select(["fk_tongue_02_anim.cv[0:32]"])
+    cmds.scale(0.246972, 0.246972, 0.246972, r = True, ocp = False)
+    cmds.move(0, 0, 0, r = True)
+    cmds.select(["fk_tongue_03_anim.cv[0:32]"])
+    cmds.scale(0.246972, 0.246972, 0.246972, r = True, ocp = False)
+    cmds.move(0, 0, 0, r = True)
+
+    cmds.select(["small_pack_l_anim.cv[0:4]"])
+    cmds.scale(1, 1, 1, r = True, ocp = True)
+    cmds.move(10, 0, 0, r = True, os=True, wd=True)
+
+    cmds.select(["small_pack_r_anim.cv[0:4]"])
+    cmds.scale(1, 1, 1, r = True, ocp = True)
+    cmds.move(-10, 0, 0, r = True, os=True, wd=True)
+
+    cmds.select(["big_pack_l_anim.cv[0:4]"])
+    cmds.scale(1, 1, 1, r = True, ocp = True)
+    cmds.move(12, 0, 0, r = True, os=True, wd=True)
+
+    cmds.select(["big_pack_r_anim.cv[0:4]"])
+    cmds.scale(1, 1, 1, r = True, ocp = True)
+    cmds.move(-12, 0, 0, r = True, os=True, wd=True)
+
+    cmds.select(None)
+
+scaleControls()
+
+# ------------------------------------------------------------------
+def shoulderRig():
+    
+    cmds.setAttr("Rig_Settings.lArmMode", 0)
+    cmds.setAttr("Rig_Settings.lClavMode", 0)
+    cmds.setAttr("Rig_Settings.rArmMode", 0)
+    cmds.setAttr("Rig_Settings.rClavMode", 0)
+
+    # Left Side
+    cmds.setDrivenKeyframe(["scapula_l_anim_grp.translate", "scapula_l_anim_grp.rotate"], cd = "driver_clavicle_l.rotateZ", itt = "spline", ott = "spline")
+
+    # Clav back and forward
+    cmds.setAttr("fk_clavicle_l_anim.rotateZ", 30)
+    cmds.setAttr("scapula_l_anim_grp.translate", 8.81813, 17.39541, 0, type="double3")
+    cmds.setAttr("scapula_l_anim_grp.rotate", 0, 0, 10.86792, type="double3")
+    cmds.setDrivenKeyframe(["scapula_l_anim_grp.translate", "scapula_l_anim_grp.rotate"], cd = "driver_clavicle_l.rotateZ", itt = "spline", ott = "spline")
+
+    cmds.setAttr("fk_clavicle_l_anim.rotateZ", -30)
+    cmds.setAttr("scapula_l_anim_grp.translate", -14.71763, -1.43961, 1.38045, type="double3")
+    cmds.setAttr("scapula_l_anim_grp.rotate", 0, 0, 36.12318, type="double3")
+    cmds.setDrivenKeyframe(["scapula_l_anim_grp.translate", "scapula_l_anim_grp.rotate"], cd = "driver_clavicle_l.rotateZ", itt = "spline", ott = "spline")
+
+    cmds.setAttr("fk_clavicle_l_anim.rotateZ", 0)
+
+    # Clav up and down
+    cmds.setDrivenKeyframe(["scapula_l_anim_grp.translate", "scapula_l_anim_grp.rotate"], cd = "driver_clavicle_l.rotateY", itt = "spline", ott = "spline")
+
+    cmds.setAttr("fk_clavicle_l_anim.rotateY", -35)
+    cmds.setAttr("scapula_l_anim_grp.translate", 1.02309, 4.7537, 5.18273, type="double3")
+    cmds.setAttr("scapula_l_anim_grp.rotate", 5.18273, -9.00623, 26.123, type="double3")
+    cmds.setDrivenKeyframe(["scapula_l_anim_grp.translate", "scapula_l_anim_grp.rotate"], cd = "driver_clavicle_l.rotateY", itt = "spline", ott = "spline")
+
+    cmds.setAttr("fk_clavicle_l_anim.rotateY", 20)
+    cmds.setAttr("scapula_l_anim_grp.translate", 0, 0, 4.3, type="double3")
+    cmds.setAttr("scapula_l_anim_grp.rotate", 0, -4.7, 0, type="double3")
+    cmds.setDrivenKeyframe(["scapula_l_anim_grp.translate", "scapula_l_anim_grp.rotate"], cd = "driver_clavicle_l.rotateY", itt = "spline", ott = "spline")
+
+    cmds.setAttr("fk_clavicle_l_anim.rotateY", 0)
+
+    # Right Side
+    cmds.setDrivenKeyframe(["scapula_r_anim_grp.translate", "scapula_r_anim_grp.rotate"], cd = "driver_clavicle_r.rotateZ", itt = "spline", ott = "spline")
+
+    # Clav back and forward
+    cmds.setAttr("fk_clavicle_r_anim.rotateZ", 30)
+    cmds.setAttr("scapula_r_anim_grp.translate", -8.81813, -17.39541, 0, type="double3")
+    cmds.setAttr("scapula_r_anim_grp.rotate", 0, 0, -10.86792, type="double3")
+    cmds.setDrivenKeyframe(["scapula_r_anim_grp.translate", "scapula_r_anim_grp.rotate"], cd = "driver_clavicle_r.rotateZ", itt = "spline", ott = "spline")
+
+    cmds.setAttr("fk_clavicle_r_anim.rotateZ", -30)
+    cmds.setAttr("scapula_r_anim_grp.translate", 14.71763, 1.43961, -1.38045, type="double3")
+    cmds.setAttr("scapula_r_anim_grp.rotate", 0, 0, -36.12318, type="double3")
+    cmds.setDrivenKeyframe(["scapula_r_anim_grp.translate", "scapula_r_anim_grp.rotate"], cd = "driver_clavicle_r.rotateZ", itt = "spline", ott = "spline")
+
+    cmds.setAttr("fk_clavicle_r_anim.rotateZ", 0)
+
+    # Clav up and down
+    cmds.setDrivenKeyframe(["scapula_r_anim_grp.translate", "scapula_r_anim_grp.rotate"], cd = "driver_clavicle_r.rotateY", itt = "spline", ott = "spline")
+
+    cmds.setAttr("fk_clavicle_r_anim.rotateY", -35)
+    cmds.setAttr("scapula_r_anim_grp.translate", -1.02309, -4.7537, -5.18273, type="double3")
+    cmds.setAttr("scapula_r_anim_grp.rotate", -5.18273, 9.00623, -26.123, type="double3")
+    cmds.setDrivenKeyframe(["scapula_r_anim_grp.translate", "scapula_r_anim_grp.rotate"], cd = "driver_clavicle_r.rotateY", itt = "spline", ott = "spline")
+
+    cmds.setAttr("fk_clavicle_r_anim.rotateY", 20)
+    cmds.setAttr("scapula_r_anim_grp.translate", 0, 0, -4.3, type="double3")
+    cmds.setAttr("scapula_r_anim_grp.rotate", 0, 4.7, 0, type="double3")
+    cmds.setDrivenKeyframe(["scapula_r_anim_grp.translate", "scapula_r_anim_grp.rotate"], cd = "driver_clavicle_r.rotateY", itt = "spline", ott = "spline")
+
+    cmds.setAttr("fk_clavicle_r_anim.rotateY", 0)
+
+    cmds.setAttr("Rig_Settings.lArmMode", 1)
+    cmds.setAttr("Rig_Settings.lClavMode", 1)
+    cmds.setAttr("Rig_Settings.rArmMode", 1)
+    cmds.setAttr("Rig_Settings.rClavMode", 1)
+
+shoulderRig()
+
+# ------------------------------------------------------------------
+# Create all of the blendshapes that are needed an hook them up.
+def createBlendshapes():
+    cmds.setAttr("Rig_Settings.lArmMode", 0)
+    cmds.setAttr("Rig_Settings.lClavMode", 0)
+    cmds.setAttr("Rig_Settings.rArmMode", 0)
+    cmds.setAttr("Rig_Settings.rClavMode", 0)
+
+    for side in ["_l", "_r"]:
+        # Create SDK's to drive the shoulder blendshape by the clavicle joint
+        cmds.setAttr("fk_arm"+side+"_anim.ry", 0)
+        cmds.setAttr("body_geo_blendShape.body_shoulders"+side+"_up_morph", 1)
+        cmds.setDrivenKeyframe("body_geo_blendShape.body_shoulders"+side+"_up_morph", cd="driver_upperarm"+side+".rotateY", itt="flat", ott="flat")
+        cmds.setAttr("fk_arm"+side+"_anim.ry", 67)
+        cmds.setAttr("body_geo_blendShape.body_shoulders"+side+"_up_morph", 0)
+        cmds.setDrivenKeyframe("body_geo_blendShape.body_shoulders"+side+"_up_morph", cd="driver_upperarm"+side+".rotateY", itt="flat", ott="flat")
+        cmds.setAttr("fk_arm"+side+"_anim.ry", 0)
+
+    cmds.setAttr("Rig_Settings.lArmMode", 1)
+    cmds.setAttr("Rig_Settings.lClavMode", 1)
+    cmds.setAttr("Rig_Settings.rArmMode", 1)
+    cmds.setAttr("Rig_Settings.rClavMode", 1)
+createBlendshapes()
+
+# ------------------------------------------------------------------
+def dynamicsOff():
+    for one in ["big_pack_l_anim", "small_pack_l_anim", "big_pack_r_anim", "small_pack_r_anim", "horn_ring_anim", "tongue_dyn_anim", "loin_cloth_fr_l_dyn_anim", "loin_cloth_fr_r_dyn_anim", "loin_cloth_bk_l_dyn_anim", "loin_cloth_bk_r_dyn_anim"]:
+        cmds.setAttr(one+".chainStartEnvelope", 0)
+        cmds.setAttr(one+".chainStartEnvelope", l=True)
+dynamicsOff()
+
+# ------------------------------------------------------------------
+def invertEyelids():
+    sides = ["_l", "_r"]
+    for side in sides:
+        cmds.delete("driver_eyelid_bott"+side+"_parentConstraint1")
+        cmds.rotate(0, -180, 0, "eyelid_bott"+side+"_anim_grp", r=True, os=True)
+        cmds.parentConstraint("eyelid_bott"+side+"_anim", "driver_eyelid_bott"+side, mo=True)
+invertEyelids()
+
+def invertRightEyeball():
+    cmds.delete("driver_eyeball_r_parentConstraint1")
+    cmds.rotate(0, -180, -180, "eyeball_r_anim_grp", r=True, os=True)
+    cmds.parentConstraint("eyeball_r_anim", "driver_eyeball_r", mo=True)
+invertRightEyeball()
+
+# ------------------------------------------------------------------
+def addCapsule(radius, halfHeight):
+    trueHalfHeight = (halfHeight*2)-(radius*2)
+    capsule = cmds.polyCylinder(name="collision_capsule", r=radius, h=trueHalfHeight, sx=18, sy=1, sz=5, ax=[0,0,1], rcp=True, cuv=3, ch=True)
+    cmds.setAttr(capsule[0]+".tz", halfHeight+2.2)
+    cmds.setAttr(capsule[0]+".rz", 90)
+
+    cmds.addAttr(capsule[0], ln="Radius", dv=radius, keyable=True)
+    cmds.addAttr(capsule[0], ln="HalfHeight", dv=halfHeight, keyable=True)
+
+    expr = capsule[1]+".radius = "+capsule[0]+".Radius;\n"+capsule[1]+".height = ("+capsule[0]+".HalfHeight*2) - ("+capsule[0]+".Radius*2);"
+    cmds.expression(s=expr, o=capsule[1], ae=1, uc="all")
+
+    cmds.setAttr(capsule[0]+".overrideEnabled", 1)
+    cmds.setAttr(capsule[0]+".overrideColor", 18)
+    cmds.setAttr(capsule[0]+".overrideShading", 0)
+
+    cmds.addAttr(capsule[0], ln="Shaded", dv=0, min=0, max=1, keyable=True)
+    cmds.connectAttr(capsule[0]+".Shaded", capsule[0]+".overrideShading")
+
+    cmds.parentConstraint("root", capsule[0], mo=True)
+    cmds.select(capsule[0], r=True)
+    capsuleDisplayLayer = cmds.createDisplayLayer(nr=True, name="Collision Capsule")
+    cmds.setAttr(capsuleDisplayLayer+".displayType", 2)
+    cmds.setAttr(capsuleDisplayLayer+".v", 0)
+    
+    for attr in ["tx", "ty", "tz", "rx", "ry", "rz", "sx", "sy", "sz"]:
+        cmds.setAttr(capsule[0]+"."+attr, lock=True, keyable=False, channelBox=False)
+
+    cmds.select(clear=True)
+addCapsule(60, 90)
+
+# ------------------------------------------------------------------
+def createGameCam():
+    gameCam = cmds.camera(centerOfInterest=5, focalLength=18, lensSqueezeRatio=1, cameraScale=1, horizontalFilmAperture=1.41732, horizontalFilmOffset=0, verticalFilmAperture=0.94488, verticalFilmOffset=0, filmFit="Overscan", overscan=1.05, motionBlur=0, shutterAngle=144, nearClipPlane=1.0, farClipPlane=10000, orthographic=0, orthographicWidth=30, panZoomEnabled=0, horizontalPan=0, verticalPan=0, zoom=1, displayFilmGate=True)[0]
+    print gameCam
+    gameCam = cmds.rename(gameCam, "GameCam")
+    cmds.camera(gameCam, e=True, horizontalFilmAperture=1.6724376, focalLength=21.2399857011)
+    cmds.setAttr(gameCam+".displayGateMaskOpacity", 1)
+    cmds.setAttr(gameCam+".displayGateMaskColor", 0, 0, 0, type="double3")
+    cmds.setAttr(gameCam+".r", 90, 0, 180, type="double3")
+    
+    # Create the control
+    control = cmds.circle(name=gameCam+"_Ctrl", ch=True, o=True, r=15)[0]
+    cmds.makeIdentity(control, a=True, t=True, r=True, s=True)
+    # Create groups for the sdk and the constraints
+    const_grp = cmds.group(control, name=control+"_grp")
+    if cmds.objExists("rig_grp"):
+        cmds.parent(const_grp, "rig_grp")
+    # Parent Constrain the cam to the control
+    cmds.parentConstraint(control, gameCam, mo=True)
+    
+    cmds.connectAttr(control+".visibility", gameCam+".visibility", f=True)
+    cmds.setAttr(const_grp+".v", 0)
+    
+    cmds.setAttr(const_grp+".t", 0, 500, 200, type="double3")
+createGameCam()
+# ------------------------------------------------------------------
+
+# ------------------------------------------------------------------
+# add ik bones Function that is called later after we open the export file.
+def addIKBones():
+    try:
+		# create the joints
+		cmds.select(clear = True)
+		ikFootRoot = cmds.joint(name = "ik_foot_root")
+		cmds.select(clear = True)
+		
+		cmds.select(clear = True)
+		ikFootLeft = cmds.joint(name = "ik_foot_l")
+		cmds.select(clear = True)
+		
+		cmds.select(clear = True)
+		ikFootRight = cmds.joint(name = "ik_foot_r")
+		cmds.select(clear = True)
+		
+		cmds.select(clear = True)
+		ikHandRoot = cmds.joint(name = "ik_hand_root")
+		cmds.select(clear = True)
+		
+		cmds.select(clear = True)
+		ikHandGun = cmds.joint(name = "ik_hand_gun")
+		cmds.select(clear = True)
+		
+		cmds.select(clear = True)
+		ikHandLeft = cmds.joint(name = "ik_hand_l")
+		cmds.select(clear = True)
+		
+		cmds.select(clear = True)
+		ikHandRight = cmds.joint(name = "ik_hand_r")
+		cmds.select(clear = True)
+		
+		# create hierarhcy
+		cmds.parent(ikFootRoot, "root")
+		cmds.parent(ikHandRoot, "root")
+		cmds.parent(ikFootLeft, ikFootRoot)
+		cmds.parent(ikFootRight, ikFootRoot)
+		cmds.parent(ikHandGun, ikHandRoot)
+		cmds.parent(ikHandLeft, ikHandGun)
+		cmds.parent(ikHandRight, ikHandGun)
+		
+		# constrain the joints
+		leftHandConstraint = cmds.parentConstraint("hand_l", ikHandLeft)[0]
+		rightHandGunConstraint = cmds.parentConstraint("hand_r", ikHandGun)[0]
+		rightHandConstraint = cmds.parentConstraint("hand_r", ikHandRight)[0]
+		leftFootConstraint = cmds.parentConstraint("foot_l", ikFootLeft)[0]
+		rightFootConstraint = cmds.parentConstraint("foot_r", ikFootRight)[0]
+	
+    except:
+		cmds.warning("Something went wrong")
+
+# ------------------------------------------------------------------
+# ADD the IK bones to the Export file.
+#save the AnimRig file before opening the export file.
+cmds.file(save = True, type = "mayaBinary", force = True)
+
+#ADD IK BONES
+#open export file
+sceneName = cmds.file(q = True, sceneName = True)
+exportFile = sceneName.partition("AnimRigs")[0] + "ExportFiles" +  sceneName.partition("AnimRigs")[2].partition(".mb")[0] + "_Export.mb"
+cmds.file(exportFile, open = True, force = True)
+
+#add ik bones
+addIKBones()
+
+# ------------------------------------------------------------------
+def addCustomAttrs():
+    attrs = cmds.listAttr("root")
+    exists = False
+    for one in attrs:
+        if one.find("DistanceCurve") != -1:
+            exists = True
+    if exists == False:
+        cmds.addAttr("root", ln="DistanceCurve", keyable=True)
+addCustomAttrs()
+
+#save the export file.
+cmds.file(save = True, type = "mayaBinary", force = True)
\ No newline at end of file
diff --git a/Engine/Extras/Maya_AnimationRiggingTools/ARTv1/General/ART/Projects/Orion/Scripts/Hammer_Post_Script.py b/Engine/Extras/Maya_AnimationRiggingTools/ARTv1/General/ART/Projects/Orion/Scripts/Hammer_Post_Script.py
new file mode 100644
--- /dev/null
+++ b/Engine/Extras/Maya_AnimationRiggingTools/ARTv1/General/ART/Projects/Orion/Scripts/Hammer_Post_Script.py
@@ -0,0 +1,317 @@
+import maya.cmds as cmds
+
+# ------------------------------------------------------------------
+# scale all of the controls to be more visible with the mesh.
+def scaleControls():
+    cmds.setAttr("leg_sys_grp.v", 0)
+
+    '''for side in ["_l", "_r"]:
+        cmds.select(["ik_wrist"+side+"_anim.cv[*]"])
+        cmds.scale(1, 1, 1, r = True, ocp = False)'''
+
+    objects = ["shoulderpad_l_anim", "shoulderpad_r_anim"]
+    for one in objects:
+        cmds.select([one+".cv[*]"])
+        cmds.scale(0.5, 0.5, 0.5, r=True, ocp=False)
+        cmds.move(0, 0, 0, r = True, os=True, wd=True)
+
+    cmds.select(["shoulderA_01_l_anim.cv[*]"])
+    cmds.scale(0.48, 0.48, 0.48, r = True, ocp = False)
+    cmds.move(16, -31, 10, r = True)            
+
+    cmds.select(["shoulderA_01_r_anim.cv[*]"])
+    cmds.scale(0.48, 0.48, 0.48, r = True, ocp = False)
+    cmds.move(-16, -31, 10, r = True)            
+
+    cmds.select(["shoulderA_02_l_anim.cv[*]"])
+    cmds.scale(0.48, 0.48, 0.48, r = True, ocp = False)
+    cmds.move(13, 27, 23, r = True)            
+
+    cmds.select(["shoulderA_02_r_anim.cv[*]"])
+    cmds.scale(0.48, 0.48, 0.48, r = True, ocp = False)
+    cmds.move(-13, 27, 23, r = True)            
+
+    cmds.select(["shoulderB_01_l_anim.cv[*]"])
+    cmds.scale(0.48, 0.48, 0.48, r = True, ocp = False)
+    cmds.move(16, -27, 4, r = True)            
+
+    cmds.select(["shoulderB_01_r_anim.cv[*]"])
+    cmds.scale(0.48, 0.48, 0.48, r = True, ocp = False)
+    cmds.move(-16, -27, 4, r = True)            
+
+    cmds.select(["shoulderC_01_l_anim.cv[*]"])
+    cmds.scale(0.48, 0.48, 0.48, r = True, ocp = False)
+    cmds.move(7, 0, 15, r = True)            
+
+    cmds.select(["shoulderC_01_r_anim.cv[*]"])
+    cmds.scale(0.48, 0.48, 0.48, r = True, ocp = False)
+    cmds.move(-7, 0, 15, r = True)            
+
+    cmds.select(["spintal_cord_01_anim.cv[*]"])
+    cmds.scale(0.48, 0.48, 0.48, r = True, ocp = False)
+    cmds.move(0, 35, 0, r = True)            
+
+    cmds.select(["forearm_armor_01_l_anim.cv[*]"])
+    cmds.scale(0.48, 0.48, 0.48, r = True, ocp = False)
+    cmds.move(0, 0, 14, r = True)            
+
+    cmds.select(["forearm_armor_01_r_anim.cv[*]"])
+    cmds.scale(0.48, 0.48, 0.48, r = True, ocp = False)
+    cmds.move(0, 0, 14, r = True)            
+
+    cmds.select(None)
+scaleControls()
+
+# ------------------------------------------------------------------
+def shoulderRig():
+    for side in ["_l", "_r"]:
+        cmds.select("shoulderpad"+side+"_anim.cv[*]")
+        cmds.scale(0.75, 0.75, 0.75, r = True)
+        if side == "_l":
+            cmds.move(0, 0, 35, r = True, os=True, wd=True)
+        else:
+            cmds.move(0, 0, -35, r = True, os=True, wd=True)
+            
+        cmds.addAttr("shoulderpad"+side+"_anim", ln = "FollowClav", dv = 0, min = 0, max = 1, keyable = True)
+        
+        cmds.select("clavicle"+side+"_anim.cv[*]")
+        if side == "_l":
+            cmds.move(22, 0, 56, r = True, os=True, wd=True)
+        else:
+            cmds.move(-22, 0, 56, r = True, os=True, wd=True)
+
+        
+        shoulder_loc = cmds.duplicate("shoulderpad"+side+"_anim_space_switcher", n="shoulder_loc"+side+"", po=True)[0]
+        cmds.parent(shoulder_loc, "ctrl_rig")
+        cmds.makeIdentity(shoulder_loc, t=True, r=True, a=True)
+        cmds.parentConstraint("driver_upperarm"+side+"", shoulder_loc, mo=True)
+        
+        const = cmds.pointConstraint(shoulder_loc, "shoulderpad"+side+"_anim_space_switcher", "shoulderpad"+side+"_anim_grp", mo=True, sk=["x"])[0]
+        
+        reverseNode = cmds.shadingNode("reverse", asUtility=True, name="shoulderpad"+side+"_reverse")
+        cmds.connectAttr("shoulderpad"+side+"_anim.FollowClav", const+"."+shoulder_loc+"W0", f=True)
+        cmds.connectAttr("shoulderpad"+side+"_anim.FollowClav", reverseNode+".inputX", f=True)
+        cmds.connectAttr(reverseNode+".outputX", const+".shoulderpad"+side+"_anim_space_switcherW1", f=True)
+    
+    # Create the attrs for the stages
+    cmds.addAttr("body_anim", ln = "Stages", dv = 0, min = 0, max = 4, keyable = True)
+        
+    # Set the controls to be in the positions for stage 0
+    cmds.setAttr("spintal_cord_01_anim.scale", 0.57, 0.57, 0.57, type="double3")
+    cmds.setAttr("forearm_armor_01_l_anim.scale", 0.72, 0.22, 0.14, type="double3")
+    cmds.setAttr("forearm_armor_01_r_anim.scale", 0.72, 0.22, 0.14, type="double3")
+
+    cmds.setAttr("shoulderA_01_l_anim.scale", 0.42, 0.61, 0.61, type="double3")
+    cmds.setAttr("shoulderA_02_l_anim.scale", 0.42, 0.61, 0.61, type="double3")
+    cmds.setAttr("shoulderB_01_l_anim_grp.translateX", -4.97)
+    cmds.setAttr("shoulderB_01_l_anim.scale", 0.31, 0.31, 0.31, type="double3")
+    cmds.setAttr("shoulderC_01_l_anim.scale", 0.165, 0.165, 0.165, type="double3")
+
+    cmds.setAttr("shoulderA_01_r_anim.scale", 0.42, 0.61, 0.61, type="double3")
+    cmds.setAttr("shoulderA_02_r_anim.scale", 0.42, 0.61, 0.61, type="double3")
+    cmds.setAttr("shoulderB_01_r_anim_grp.translateX", 4.97)
+    cmds.setAttr("shoulderB_01_r_anim.scale", 0.31, 0.31, 0.31, type="double3")
+    cmds.setAttr("shoulderC_01_r_anim.scale", 0.165, 0.165, 0.165, type="double3")
+
+    cmds.setAttr("shoulderB_blendshape.stage_3_shoulderB_MORPH", 0)
+    cmds.setAttr("shoulderC_blendshape.stage_3_shoulderC_MORPH", 0)
+    cmds.setAttr("shoulderC_blendshape.stage_4_shoulderC_MORPH", 0)
+    cmds.setAttr("shoulderB_blendshape.stage_4_shoulderB_MORPH", 0)
+    cmds.setAttr("spinal_cord_blendshape.stage_4_spinal_cord_MORPH", 0)
+
+    cmds.setAttr("body_anim.Stages", 0)
+
+    # Create SDK's.
+    cmds.setDrivenKeyframe(["spintal_cord_01_anim.scale", "forearm_armor_01_l_anim.scale", "forearm_armor_01_r_anim.scale", "shoulderA_01_l_anim.scale", "shoulderA_02_l_anim.scale", "shoulderB_01_l_anim_grp.translateX", "shoulderB_01_l_anim.scale", "shoulderC_01_l_anim.scale", "shoulderA_01_r_anim.scale", "shoulderA_02_r_anim.scale", "shoulderB_01_r_anim_grp.translateX", "shoulderB_01_r_anim.scale", "shoulderC_01_r_anim.scale", "shoulderB_blendshape.stage_3_shoulderB_MORPH", "shoulderC_blendshape.stage_3_shoulderC_MORPH", "shoulderC_blendshape.stage_4_shoulderC_MORPH", "shoulderB_blendshape.stage_4_shoulderB_MORPH", "spinal_cord_blendshape.stage_4_spinal_cord_MORPH"], cd = "body_anim.Stages", itt = "flat", ott = "flat")
+
+    cmds.setAttr("body_anim.Stages", 1)
+    cmds.setAttr("shoulderA_01_l_anim.scale", 1, 1, 1, type="double3")
+    cmds.setAttr("shoulderA_02_l_anim.scale", 1, 1, 1, type="double3")
+    cmds.setAttr("shoulderA_01_r_anim.scale", 1, 1, 1, type="double3")
+    cmds.setAttr("shoulderA_02_r_anim.scale", 1, 1, 1, type="double3")
+    cmds.setDrivenKeyframe(["spintal_cord_01_anim.scale", "forearm_armor_01_l_anim.scale", "forearm_armor_01_r_anim.scale", "shoulderA_01_l_anim.scale", "shoulderA_02_l_anim.scale", "shoulderB_01_l_anim_grp.translateX", "shoulderB_01_l_anim.scale", "shoulderC_01_l_anim.scale", "shoulderA_01_r_anim.scale", "shoulderA_02_r_anim.scale", "shoulderB_01_r_anim_grp.translateX", "shoulderB_01_r_anim.scale", "shoulderC_01_r_anim.scale", "shoulderB_blendshape.stage_3_shoulderB_MORPH", "shoulderC_blendshape.stage_3_shoulderC_MORPH", "shoulderC_blendshape.stage_4_shoulderC_MORPH", "shoulderB_blendshape.stage_4_shoulderB_MORPH", "spinal_cord_blendshape.stage_4_spinal_cord_MORPH"], cd = "body_anim.Stages", itt = "flat", ott = "flat")
+
+
+    cmds.setAttr("body_anim.Stages", 2)
+    cmds.setAttr("spintal_cord_01_anim.scale", 1, 1, 1, type="double3")
+    cmds.setAttr("shoulderB_01_l_anim_grp.translateX", 0)
+    cmds.setAttr("shoulderB_01_l_anim.scale", 1, 1, 1, type="double3")
+    cmds.setAttr("shoulderB_01_r_anim_grp.translateX", 0)
+    cmds.setAttr("shoulderB_01_r_anim.scale", 1, 1, 1, type="double3")
+    cmds.setDrivenKeyframe(["spintal_cord_01_anim.scale", "forearm_armor_01_l_anim.scale", "forearm_armor_01_r_anim.scale", "shoulderA_01_l_anim.scale", "shoulderA_02_l_anim.scale", "shoulderB_01_l_anim_grp.translateX", "shoulderB_01_l_anim.scale", "shoulderC_01_l_anim.scale", "shoulderA_01_r_anim.scale", "shoulderA_02_r_anim.scale", "shoulderB_01_r_anim_grp.translateX", "shoulderB_01_r_anim.scale", "shoulderC_01_r_anim.scale", "shoulderB_blendshape.stage_3_shoulderB_MORPH", "shoulderC_blendshape.stage_3_shoulderC_MORPH", "shoulderC_blendshape.stage_4_shoulderC_MORPH", "shoulderB_blendshape.stage_4_shoulderB_MORPH", "spinal_cord_blendshape.stage_4_spinal_cord_MORPH"], cd = "body_anim.Stages", itt = "flat", ott = "flat")
+
+    cmds.setAttr("body_anim.Stages", 3)
+    cmds.setAttr("shoulderC_01_l_anim.scale", 1, 1, 1, type="double3")
+    cmds.setAttr("shoulderC_01_r_anim.scale", 1, 1, 1, type="double3")
+    cmds.setAttr("forearm_armor_01_l_anim.scale", 1, 1, 1, type="double3")
+    cmds.setAttr("forearm_armor_01_r_anim.scale", 1, 1, 1, type="double3")
+    cmds.setAttr("shoulderB_blendshape.stage_3_shoulderB_MORPH", 1)
+    cmds.setAttr("shoulderC_blendshape.stage_3_shoulderC_MORPH", 1)
+    cmds.setDrivenKeyframe(["spintal_cord_01_anim.scale", "forearm_armor_01_l_anim.scale", "forearm_armor_01_r_anim.scale", "shoulderA_01_l_anim.scale", "shoulderA_02_l_anim.scale", "shoulderB_01_l_anim_grp.translateX", "shoulderB_01_l_anim.scale", "shoulderC_01_l_anim.scale", "shoulderA_01_r_anim.scale", "shoulderA_02_r_anim.scale", "shoulderB_01_r_anim_grp.translateX", "shoulderB_01_r_anim.scale", "shoulderC_01_r_anim.scale", "shoulderB_blendshape.stage_3_shoulderB_MORPH", "shoulderC_blendshape.stage_3_shoulderC_MORPH", "shoulderC_blendshape.stage_4_shoulderC_MORPH", "shoulderB_blendshape.stage_4_shoulderB_MORPH", "spinal_cord_blendshape.stage_4_spinal_cord_MORPH"], cd = "body_anim.Stages", itt = "flat", ott = "flat")
+
+    cmds.setAttr("body_anim.Stages", 4)
+    cmds.setAttr("shoulderB_blendshape.stage_3_shoulderB_MORPH", 0)
+    cmds.setAttr("shoulderC_blendshape.stage_3_shoulderC_MORPH", 0)
+    cmds.setAttr("shoulderC_blendshape.stage_4_shoulderC_MORPH", 1)
+    cmds.setAttr("shoulderB_blendshape.stage_4_shoulderB_MORPH", 1)
+    cmds.setAttr("spinal_cord_blendshape.stage_4_spinal_cord_MORPH", 1)
+    cmds.setDrivenKeyframe(["spintal_cord_01_anim.scale", "forearm_armor_01_l_anim.scale", "forearm_armor_01_r_anim.scale", "shoulderA_01_l_anim.scale", "shoulderA_02_l_anim.scale", "shoulderB_01_l_anim_grp.translateX", "shoulderB_01_l_anim.scale", "shoulderC_01_l_anim.scale", "shoulderA_01_r_anim.scale", "shoulderA_02_r_anim.scale", "shoulderB_01_r_anim_grp.translateX", "shoulderB_01_r_anim.scale", "shoulderC_01_r_anim.scale", "shoulderB_blendshape.stage_3_shoulderB_MORPH", "shoulderC_blendshape.stage_3_shoulderC_MORPH", "shoulderC_blendshape.stage_4_shoulderC_MORPH", "shoulderB_blendshape.stage_4_shoulderB_MORPH", "spinal_cord_blendshape.stage_4_spinal_cord_MORPH"], cd = "body_anim.Stages", itt = "flat", ott = "flat")
+
+    cmds.setAttr("body_anim.Stages", 0)
+
+shoulderRig()
+
+# ------------------------------------------------------------------
+def armsLegsRig():
+    cmds.setAttr("Rig_Settings.lUpperarmTwist2Amount", 0.3)
+    cmds.setAttr("Rig_Settings.rUpperarmTwist2Amount", 0.3)
+
+    cmds.setAttr("Rig_Settings.lForearmTwist2Amount", 0.3)
+    cmds.setAttr("Rig_Settings.rForearmTwist2Amount", 0.3)
+armsLegsRig()
+
+# ------------------------------------------------------------------
+def dynamicsOff():
+    for one in ["hip_cloth_bk_l_dyn_anim", "hip_cloth_bk_r_dyn_anim", "hip_cloth_fr_l_dyn_anim", "hip_cloth_fr_r_dyn_anim", "lowerarm_cloth_bk_l_dyn_anim", "lowerarm_cloth_bk_r_dyn_anim", "lowerarm_cloth_fr_l_dyn_anim", "lowerarm_cloth_fr_r_dyn_anim", "pelvis_cloth_bk_dyn_anim", "pelvis_cloth_fr_dyn_anim", "spine_cloth_bk_l_dyn_anim", "spine_cloth_bk_r_dyn_anim", "spine_cloth_fr_l_dyn_anim", "spine_cloth_fr_r_dyn_anim", "stage1_clav_cloth_l_dyn_anim", "stage1_clav_cloth_r_dyn_anim", "upperarm_cloth_l_dyn_anim", "upperarm_cloth_r_dyn_anim"]:
+        cmds.setAttr(one+".chainStartEnvelope", 0)
+        cmds.setAttr(one+".chainStartEnvelope", l=True)
+dynamicsOff()
+
+# ------------------------------------------------------------------
+def addCapsule(radius, halfHeight):
+    trueHalfHeight = (halfHeight*2)-(radius*2)
+    capsule = cmds.polyCylinder(name="collision_capsule", r=radius, h=trueHalfHeight, sx=18, sy=1, sz=5, ax=[0,0,1], rcp=True, cuv=3, ch=True)
+    cmds.setAttr(capsule[0]+".tz", halfHeight+2.2)
+    cmds.setAttr(capsule[0]+".rz", 90)
+
+    cmds.addAttr(capsule[0], ln="Radius", dv=radius, keyable=True)
+    cmds.addAttr(capsule[0], ln="HalfHeight", dv=halfHeight, keyable=True)
+
+    expr = capsule[1]+".radius = "+capsule[0]+".Radius;\n"+capsule[1]+".height = ("+capsule[0]+".HalfHeight*2) - ("+capsule[0]+".Radius*2);"
+    cmds.expression(s=expr, o=capsule[1], ae=1, uc="all")
+
+    cmds.setAttr(capsule[0]+".overrideEnabled", 1)
+    cmds.setAttr(capsule[0]+".overrideColor", 18)
+    cmds.setAttr(capsule[0]+".overrideShading", 0)
+
+    cmds.addAttr(capsule[0], ln="Shaded", dv=0, min=0, max=1, keyable=True)
+    cmds.connectAttr(capsule[0]+".Shaded", capsule[0]+".overrideShading")
+
+    cmds.parentConstraint("root", capsule[0], mo=True)
+    cmds.select(capsule[0], r=True)
+    capsuleDisplayLayer = cmds.createDisplayLayer(nr=True, name="Collision Capsule")
+    cmds.setAttr(capsuleDisplayLayer+".displayType", 2)
+    cmds.setAttr(capsuleDisplayLayer+".v", 0)
+    
+    for attr in ["tx", "ty", "tz", "rx", "ry", "rz", "sx", "sy", "sz"]:
+        cmds.setAttr(capsule[0]+"."+attr, lock=True, keyable=False, channelBox=False)
+
+    cmds.select(clear=True)
+addCapsule(42, 86)
+
+# ------------------------------------------------------------------
+def createGameCam():
+    gameCam = cmds.camera(centerOfInterest=5, focalLength=18, lensSqueezeRatio=1, cameraScale=1, horizontalFilmAperture=1.41732, horizontalFilmOffset=0, verticalFilmAperture=0.94488, verticalFilmOffset=0, filmFit="Overscan", overscan=1.05, motionBlur=0, shutterAngle=144, nearClipPlane=1.0, farClipPlane=10000, orthographic=0, orthographicWidth=30, panZoomEnabled=0, horizontalPan=0, verticalPan=0, zoom=1, displayFilmGate=True)[0]
+    print gameCam
+    gameCam = cmds.rename(gameCam, "GameCam")
+    cmds.camera(gameCam, e=True, horizontalFilmAperture=1.6724376, focalLength=21.2399857011)
+    cmds.setAttr(gameCam+".displayGateMaskOpacity", 1)
+    cmds.setAttr(gameCam+".displayGateMaskColor", 0, 0, 0, type="double3")
+    cmds.setAttr(gameCam+".r", 90, 0, 180, type="double3")
+    
+    # Create the control
+    control = cmds.circle(name=gameCam+"_Ctrl", ch=True, o=True, r=15)[0]
+    cmds.makeIdentity(control, a=True, t=True, r=True, s=True)
+    # Create groups for the sdk and the constraints
+    const_grp = cmds.group(control, name=control+"_grp")
+    if cmds.objExists("rig_grp"):
+        cmds.parent(const_grp, "rig_grp")
+    # Parent Constrain the cam to the control
+    cmds.parentConstraint(control, gameCam, mo=True)
+    
+    cmds.connectAttr(control+".visibility", gameCam+".visibility", f=True)
+    cmds.setAttr(const_grp+".v", 0)
+    
+    cmds.setAttr(const_grp+".t", 0, 500, 200, type="double3")
+createGameCam()
+
+# ------------------------------------------------------------------
+# add ik bones Function that is called later after we open the export file.
+def addIKBones():
+    try:
+		# create the joints
+		cmds.select(clear = True)
+		ikFootRoot = cmds.joint(name = "ik_foot_root")
+		cmds.select(clear = True)
+		
+		cmds.select(clear = True)
+		ikFootLeft = cmds.joint(name = "ik_foot_l")
+		cmds.select(clear = True)
+		
+		cmds.select(clear = True)
+		ikFootRight = cmds.joint(name = "ik_foot_r")
+		cmds.select(clear = True)
+		
+		cmds.select(clear = True)
+		ikHandRoot = cmds.joint(name = "ik_hand_root")
+		cmds.select(clear = True)
+		
+		cmds.select(clear = True)
+		ikHandGun = cmds.joint(name = "ik_hand_gun")
+		cmds.select(clear = True)
+		
+		cmds.select(clear = True)
+		ikHandLeft = cmds.joint(name = "ik_hand_l")
+		cmds.select(clear = True)
+		
+		cmds.select(clear = True)
+		ikHandRight = cmds.joint(name = "ik_hand_r")
+		cmds.select(clear = True)
+		
+		# create hierarhcy
+		cmds.parent(ikFootRoot, "root")
+		cmds.parent(ikHandRoot, "root")
+		cmds.parent(ikFootLeft, ikFootRoot)
+		cmds.parent(ikFootRight, ikFootRoot)
+		cmds.parent(ikHandGun, ikHandRoot)
+		cmds.parent(ikHandLeft, ikHandGun)
+		cmds.parent(ikHandRight, ikHandGun)
+		
+		# constrain the joints
+		leftHandConstraint = cmds.parentConstraint("hand_l", ikHandLeft)[0]
+		rightHandGunConstraint = cmds.parentConstraint("hand_r", ikHandGun)[0]
+		rightHandConstraint = cmds.parentConstraint("hand_r", ikHandRight)[0]
+		leftFootConstraint = cmds.parentConstraint("foot_l", ikFootLeft)[0]
+		rightFootConstraint = cmds.parentConstraint("foot_r", ikFootRight)[0]
+	
+    except:
+		cmds.warning("Something went wrong")
+
+# ------------------------------------------------------------------
+# ADD the IK bones to the Export file.
+#save the AnimRig file before opening the export file.
+cmds.file(save = True, type = "mayaBinary", force = True)
+
+#ADD IK BONES
+#open export file
+sceneName = cmds.file(q = True, sceneName = True)
+exportFile = sceneName.partition("AnimRigs")[0] + "ExportFiles" +  sceneName.partition("AnimRigs")[2].partition(".mb")[0] + "_Export.mb"
+cmds.file(exportFile, open = True, force = True)
+
+#add ik bones
+addIKBones()
+
+# ------------------------------------------------------------------
+def addCustomAttrs():
+    attrs = cmds.listAttr("root")
+    exists = False
+    for one in attrs:
+        if one.find("DistanceCurve") != -1:
+            exists = True
+    if exists == False:
+        cmds.addAttr("root", ln="DistanceCurve", keyable=True)
+addCustomAttrs()
+
+# ------------------------------------------------------------------
+#save the export file.
+cmds.file(save = True, type = "mayaBinary", force = True)
