Branching

diff --git a/Engine/Extras/Maya_AnimationRiggingTools/ARTv1/General/ART/Projects/Orion/Scripts/HyperBreach_CreateLODs.py b/Engine/Extras/Maya_AnimationRiggingTools/ARTv1/General/ART/Projects/Orion/Scripts/HyperBreach_CreateLODs.py
new file mode 100644
--- /dev/null
+++ b/Engine/Extras/Maya_AnimationRiggingTools/ARTv1/General/ART/Projects/Orion/Scripts/HyperBreach_CreateLODs.py
@@ -0,0 +1,822 @@
+import maya.cmds as cmds
+import customMayaMenu as menu
+import Tools.ART_ImportMultipleWeights as impw
+import Tools.ART_TransferWeights as transfer
+import Tools.ART_ExportLODtoFBX as export
+reload(export)
+
+################################################## HYPERBREACH #########################################################################
+# Warn the user that this operation will save their file.
+saveResult = cmds.confirmDialog(title = "Save Warning", message="In order to continue your file must be saved.  Would you like to save it?  If yes it will be saved and after the operation is complete your file will be re-opened.", button = ["Yes", "Cancel"], defaultButton='Yes', cancelButton='Cancel', dismissString='Cancel')
+
+# If they are ok with saving the file, warn them that it will also delete all of the current files in their skin weights directory.
+if saveResult == "Yes":
+    deleteSkinFilesResult = cmds.confirmDialog(title = "SkinWeights Warning", message="In order to continue we must delete all of the files in your ..\General\ART\SkinWeights\ directory.  Would you like to continue?.", button = ["Yes", "Cancel"], defaultButton='Yes', cancelButton='Cancel', dismissString='Cancel')
+    if deleteSkinFilesResult == "Yes":
+        toolsPath = cmds.internalVar(usd = True) + "mayaTools.txt"
+        if os.path.exists(toolsPath):
+            f = open(toolsPath, 'r')
+            mayaToolsDir = f.readline()
+            f.close()
+
+        path = mayaToolsDir + "\General\ART\SkinWeights\\"
+        files = os.listdir(path)
+        #delete all weight cache files in the system folder
+        for file in files:
+            os.remove(path + file)
+            #print path + file
+
+
+
+# HyperBreach LOD1 -------------------------------------------------------------------------------------------------------------
+def HyperBreachLOD1():
+    if saveResult == "Yes":
+
+        cmds.file(save=True)
+
+        meshes = ["backlauncher_l_geo_LOD1", "backlauncher_r_geo_LOD1", "backspine_geo_LOD1", "bicep_l_geo_LOD1", "bicep_l_hatch_geo_LOD1", "bicep_r_geo_LOD1", "bicep_r_hatch_geo_LOD1", "bodyback_geo_LOD1", "bodypanel_l_geo_LOD1", "bodyshell_geo_LOD1", "calfhatch_l_geo_LOD1", "calfhatch_r_geo_LOD1", "calfjet_l_geo_LOD1", "calfjet_r_geo_LOD1", "calfs_geo_LOD1", "clav_l_geo_LOD1", "clav_r_geo_LOD1", "controlarm_l1_geo_LOD1", "controlarm_l2_geo_LOD1", "controlarm_r1_geo_LOD1", "controlarm_r2_geo_LOD1", "controlgrip_l_geo_LOD1", "controlgrip_r_geo_LOD1", "foot_l_geo_LOD1", "foot_r_geo_LOD1", "footflap_l_geo_LOD1", "footflap_r_geo_LOD1", "grenade1_geo_LOD1", "grenade2_geo_LOD1", "grenade3_geo_LOD1", "grenadefinger_1a_geo_LOD1", "grenadefinger_1b_geo_LOD1", "grenadefinger_1c_geo_LOD1", "grenadefinger_2a_geo_LOD1", "grenadefinger_2b_geo_LOD1", "grenadefinger_2c_geo_LOD1", "grenadefinger_3a_geo_LOD1", "grenadefinger_3b_geo_LOD1", "grenadefinger_3c_geo_LOD1", "grenadelauncherr_geo_LOD1", "grenadepiston_geo_LOD1", "hipcenter_geo_LOD1", "hipplate_geo_LOD1", "hose_geo_LOD1", "jetarm_l_geo_LOD1", "jetarm_r_geo_LOD1", "jetfoot_l_geo_LOD1", "jetfoot_r_geo_LOD1", "leg_l_geo_LOD1", "leg_r_geo_LOD1", "lowerspine_geo_LOD1", "mine_geo_LOD1", "minebase_geo_LOD1", "minedeploy_geo_LOD1", "minepivot_geo_LOD1", "missle_geo_LOD1", "missledoor1_geo_LOD1", "missledoor2_geo_LOD1", "misslelatch1_geo_LOD1", "misslelatch2_geo_LOD1", "misslelatch3_geo_LOD1", "misslelauncher_geo_LOD1", "misslepiston_geo_LOD1", "pedal_l_geo_LOD1", "pedal_r_geo_LOD1", "pedalarm_l_geo_LOD1", "pedalarm_r_geo_LOD1", "seat_geo_LOD1", "seatbase_geo_LOD1", "seatpivot_geo_LOD1", "shoulder_l_geo_LOD1", "shoulder_r_geo_LOD1", "shoulderbone_l_geo_LOD1", "shoulderbone_r_geo_LOD1", "shoulderpad_l_geo_LOD1", "shoulderpad_r_geo_LOD1", "tailplate_geo_LOD1", "toe_l_geo_LOD1", "toe_r_geo_LOD1", "turbine_geo_LOD1", "Rider:cigar_geo_LOD1", "Rider:cuffs_geo_LOD1", "Rider:eye_balls_geo_LOD1", "Rider:eye_cover_geo_LOD1", "Rider:feet_geo_LOD1", "Rider:gun_barrel_geo_LOD1", "Rider:gun_clip_geo_LOD1", "Rider:gun_flaps_geo_LOD1", "Rider:gun_geo_LOD1", "Rider:hands_geo_LOD1", "Rider:head_geo_LOD1", "Rider:holster_geo_LOD1", "Rider:HyperBreach_Rabbit_Eyelashes Rider:HyperBreach_Rabbit_Hair Rider:insignia_geo_LOD1", "Rider:leggings_geo_LOD1", "Rider:legs_geo_LOD1", "Rider:lower_belt_geo_LOD1", "Rider:lower_teeth_geo_LOD1", "Rider:shoulders_geo_LOD1", "Rider:tongue_geo_LOD1", "Rider:torso_geo_LOD1", "Rider:upper_teeth_geo_LOD1"]
+        
+        # Export a skin weights file for each mesh we plan to operate on.
+        cmds.select(meshes)
+        menu.exportSelectedSkinWeights(True)
+
+        # Set the pose if needed
+        cmds.setAttr("rider_thumb_02_l.rotateZ", 22.03618142)
+        cmds.setAttr("rider_thumb_02_l.rotateY", 5.675048056)
+        cmds.setAttr("rider_thumb_02_l.rotateX", 22.94177697)
+        cmds.setAttr("rider_ring_02_l.rotateZ", 54.79433274)
+        cmds.setAttr("rider_pinky_02_l.rotateZ", 54.79433274)
+        cmds.setAttr("rider_middle_02_l.rotateZ", 59.79453692)
+        cmds.setAttr("rider_index_02_l.rotateZ", 64.20477863)
+
+        cmds.setAttr("rider_thumb_02_r.rotateZ", 22.03618142)
+        cmds.setAttr("rider_thumb_02_r.rotateY", 5.675048056)
+        cmds.setAttr("rider_thumb_02_r.rotateX", 22.94177697)
+        cmds.setAttr("rider_ring_02_r.rotateZ", 54.79433274)
+        cmds.setAttr("rider_pinky_02_r.rotateZ", 54.79433274)
+        cmds.setAttr("rider_middle_02_r.rotateZ", 59.79453692)
+        cmds.setAttr("rider_index_02_r.rotateZ", 64.20477863)
+
+       # Delete the history on the meshes so that they get baked down into their posed positions
+        cmds.delete(meshes, ch=True)
+
+        # Move the skeleton back to its bindpose (This script REQUIRES the bind pose to be the pose that we want to export the skeleton to.  This is the case for all known characters created with A.R.T.)
+        cmds.dagPose("root", bp=True, r=True) 
+
+        # Re-import the skinning we exported before onto the mesh so thats its bound in its posed position.
+        cmds.select(meshes)
+        impw.importMultipleSkinWeights("", True)
+
+        jointsToDelete = []
+
+        # remove unneeded joints
+        jointsToRemove = ["lower_fin_end"]
+        jointToTransferTo = "lower_fin"
+        transfer.LOD_transferWeights(meshes, jointsToRemove, jointToTransferTo)
+        jointsToDelete.extend(jointsToRemove)
+
+        jointsToRemove = ["cockpit_foot_pedal_r"]
+        jointToTransferTo = "cockpit_pedal_r"
+        transfer.LOD_transferWeights(meshes, jointsToRemove, jointToTransferTo)
+        jointsToDelete.extend(jointsToRemove)
+
+        jointsToRemove = ["cockpit_foot_pedal_l"]
+        jointToTransferTo = "cockpit_pedal_l"
+        transfer.LOD_transferWeights(meshes, jointsToRemove, jointToTransferTo)
+        jointsToDelete.extend(jointsToRemove)
+
+        jointsToRemove = ["controlArm_r_03"]
+        jointToTransferTo = "controlArm_r_02"
+        transfer.LOD_transferWeights(meshes, jointsToRemove, jointToTransferTo)
+        jointsToDelete.extend(jointsToRemove)
+
+        jointsToRemove = ["controlArm_l_03"]
+        jointToTransferTo = "controlArm_l_02"
+        transfer.LOD_transferWeights(meshes, jointsToRemove, jointToTransferTo)
+        jointsToDelete.extend(jointsToRemove)
+
+        jointsToRemove = ["neck_01", "head"]
+        jointToTransferTo = "spine_03"
+        transfer.LOD_transferWeights(meshes, jointsToRemove, jointToTransferTo)
+        jointsToDelete.extend(jointsToRemove)
+
+        jointsToRemove = ["hand_l"]
+        jointToTransferTo = "lowerarm_l"
+        transfer.LOD_transferWeights(meshes, jointsToRemove, jointToTransferTo)
+        jointsToDelete.extend(jointsToRemove)
+
+        jointsToRemove = ["index_03_r"]
+        jointToTransferTo = "index_02_r"
+        transfer.LOD_transferWeights(meshes, jointsToRemove, jointToTransferTo)
+        jointsToDelete.extend(jointsToRemove)
+
+        jointsToRemove = ["middle_03_r"]
+        jointToTransferTo = "middle_02_r"
+        transfer.LOD_transferWeights(meshes, jointsToRemove, jointToTransferTo)
+        jointsToDelete.extend(jointsToRemove)
+
+        jointsToRemove = ["ring_03_r"]
+        jointToTransferTo = "ring_02_r"
+        transfer.LOD_transferWeights(meshes, jointsToRemove, jointToTransferTo)
+        jointsToDelete.extend(jointsToRemove)
+
+        jointsToRemove = ["foot_spike_l"]
+        jointToTransferTo = "foot_l"
+        transfer.LOD_transferWeights(meshes, jointsToRemove, jointToTransferTo)
+        jointsToDelete.extend(jointsToRemove)
+
+        jointsToRemove = ["foot_spike_r"]
+        jointToTransferTo = "foot_r"
+        transfer.LOD_transferWeights(meshes, jointsToRemove, jointToTransferTo)
+        jointsToDelete.extend(jointsToRemove)
+
+        jointsToRemove = ["rider_index_02_l"]
+        jointToTransferTo = "rider_index_01_l"
+        transfer.LOD_transferWeights(meshes, jointsToRemove, jointToTransferTo)
+        jointsToDelete.extend(jointsToRemove)
+
+        jointsToRemove = ["rider_middle_02_l"]
+        jointToTransferTo = "rider_middle_01_l"
+        transfer.LOD_transferWeights(meshes, jointsToRemove, jointToTransferTo)
+        jointsToDelete.extend(jointsToRemove)
+
+        jointsToRemove = ["rider_pinky_02_l"]
+        jointToTransferTo = "rider_pinky_01_l"
+        transfer.LOD_transferWeights(meshes, jointsToRemove, jointToTransferTo)
+        jointsToDelete.extend(jointsToRemove)
+
+        jointsToRemove = ["rider_ring_02_l"]
+        jointToTransferTo = "rider_ring_01_l"
+        transfer.LOD_transferWeights(meshes, jointsToRemove, jointToTransferTo)
+        jointsToDelete.extend(jointsToRemove)
+
+        jointsToRemove = ["rider_thumb_02_l"]
+        jointToTransferTo = "rider_thumb_01_l"
+        transfer.LOD_transferWeights(meshes, jointsToRemove, jointToTransferTo)
+        jointsToDelete.extend(jointsToRemove)
+
+        jointsToRemove = ["rider_index_02_r"]
+        jointToTransferTo = "rider_index_01_r"
+        transfer.LOD_transferWeights(meshes, jointsToRemove, jointToTransferTo)
+        jointsToDelete.extend(jointsToRemove)
+
+        jointsToRemove = ["rider_middle_02_r"]
+        jointToTransferTo = "rider_middle_01_r"
+        transfer.LOD_transferWeights(meshes, jointsToRemove, jointToTransferTo)
+        jointsToDelete.extend(jointsToRemove)
+
+        jointsToRemove = ["rider_pinky_02_r"]
+        jointToTransferTo = "rider_pinky_01_r"
+        transfer.LOD_transferWeights(meshes, jointsToRemove, jointToTransferTo)
+        jointsToDelete.extend(jointsToRemove)
+
+        jointsToRemove = ["rider_ring_02_r"]
+        jointToTransferTo = "rider_ring_01_r"
+        transfer.LOD_transferWeights(meshes, jointsToRemove, jointToTransferTo)
+        jointsToDelete.extend(jointsToRemove)
+
+        jointsToRemove = ["rider_thumb_02_r"]
+        jointToTransferTo = "rider_thumb_01_r"
+        transfer.LOD_transferWeights(meshes, jointsToRemove, jointToTransferTo)
+        jointsToDelete.extend(jointsToRemove)
+
+        jointsToRemove = ["rider_heel_l"]
+        jointToTransferTo = "rider_calf_l"
+        transfer.LOD_transferWeights(meshes, jointsToRemove, jointToTransferTo)
+        jointsToDelete.extend(jointsToRemove)
+
+        jointsToRemove = ["rider_heel_r"]
+        jointToTransferTo = "rider_calf_r"
+        transfer.LOD_transferWeights(meshes, jointsToRemove, jointToTransferTo)
+        jointsToDelete.extend(jointsToRemove)
+
+        # Delete the joints we no longer want.
+        for one in jointsToDelete:
+            if cmds.objExists(one):
+                cmds.delete(one)
+                
+        # Export the fbx file.
+        export.exportLOD("HyperBreach_LOD1.fbx", "LOD1")
+        
+        # Re-open the file without saving.
+        fullPath = cmds.file(q = True, sceneName = True)
+        cmds.file(fullPath, open=True, f=True)
+
+        print "LOD1 Exported Successfully!!!!!!!!!!!!!!!"
+
+HyperBreachLOD1()
+
+# HyperBreach LOD2 -------------------------------------------------------------------------------------------------------------
+def HyperBreachLOD2():
+    if saveResult == "Yes":
+
+        cmds.file(save=True)
+        
+        meshes = ["backlauncher_l_geo_LOD2", "backlauncher_r_geo_LOD2", "backspine_geo_LOD2", "bicep_l_geo_LOD2", "bicep_l_hatch_geo_LOD2", "bicep_r_geo_LOD2", "bicep_r_hatch_geo_LOD2", "bodyback_geo_LOD2", "bodypanel_l_geo_LOD2", "bodyshell_geo_LOD2", "calfhatch_l_geo_LOD2", "calfhatch_r_geo_LOD2", "calfjet_l_geo_LOD2", "calfjet_r_geo_LOD2", "calfs_geo_LOD2", "clav_l_geo_LOD2", "clav_r_geo_LOD2", "controlarm_l1_geo_LOD2", "controlarm_l2_geo_LOD2", "controlarm_r1_geo_LOD2", "controlarm_r2_geo_LOD2", "controlgrip_l_geo_LOD2", "controlgrip_r_geo_LOD2", "foot_l_geo_LOD2", "foot_r_geo_LOD2", "footflap_l_geo_LOD2", "footflap_r_geo_LOD2", "grenade1_geo_LOD2", "grenade2_geo_LOD2", "grenade3_geo_LOD2", "grenadefinger_1a_geo_LOD2", "grenadefinger_1b_geo_LOD2", "grenadefinger_1c_geo_LOD2", "grenadefinger_2a_geo_LOD2", "grenadefinger_2b_geo_LOD2", "grenadefinger_2c_geo_LOD2", "grenadefinger_3a_geo_LOD2", "grenadefinger_3b_geo_LOD2", "grenadefinger_3c_geo_LOD2", "grenadelauncherr_geo_LOD2", "grenadepiston_geo_LOD2", "hipcenter_geo_LOD2", "hipplate_geo_LOD2", "hose_geo_LOD2", "jetarm_l_geo_LOD2", "jetarm_r_geo_LOD2", "jetfoot_l_geo_LOD2", "jetfoot_r_geo_LOD2", "leg_l_geo_LOD2", "leg_r_geo_LOD2", "lowerspine_geo_LOD2", "mine_geo_LOD2", "minebase_geo_LOD2", "minedeploy_geo_LOD2", "minepivot_geo_LOD2", "missle_geo_LOD2", "missledoor1_geo_LOD2", "missledoor2_geo_LOD2", "misslelatch1_geo_LOD2", "misslelatch2_geo_LOD2", "misslelatch3_geo_LOD2", "misslelauncher_geo_LOD2", "misslepiston_geo_LOD2", "pedal_l_geo_LOD2", "pedal_r_geo_LOD2", "pedalarm_l_geo_LOD2", "pedalarm_r_geo_LOD2", "seat_geo_LOD2", "seatbase_geo_LOD2", "seatpivot_geo_LOD2", "shoulder_l_geo_LOD2", "shoulder_r_geo_LOD2", "shoulderbone_l_geo_LOD2", "shoulderbone_r_geo_LOD2", "shoulderpad_l_geo_LOD2", "shoulderpad_r_geo_LOD2", "tailplate_geo_LOD2", "toe_l_geo_LOD2", "toe_r_geo_LOD2", "turbine_geo_LOD2", "Rider:cigar_geo_LOD2", "Rider:cuffs_geo_LOD2", "Rider:eye_balls_geo_LOD2", "Rider:eye_cover_geo_LOD2", "Rider:feet_geo_LOD2", "Rider:gun_barrel_geo_LOD2", "Rider:gun_clip_geo_LOD2", "Rider:gun_flaps_geo_LOD2", "Rider:gun_geo_LOD2", "Rider:hands_geo_LOD2", "Rider:head_geo_LOD2", "Rider:holster_geo_LOD2", "Rider:HyperBreach_Rabbit_Eyelashes Rider:HyperBreach_Rabbit_Hair Rider:insignia_geo_LOD2", "Rider:leggings_geo_LOD2", "Rider:legs_geo_LOD2", "Rider:lower_belt_geo_LOD2", "Rider:lower_teeth_geo_LOD2", "Rider:shoulders_geo_LOD2", "Rider:tongue_geo_LOD2", "Rider:torso_geo_LOD2", "Rider:upper_teeth_geo_LOD2"]
+
+        # Export a skin weights file for each mesh we plan to operate on.
+        cmds.select(meshes)
+        menu.exportSelectedSkinWeights(True)
+
+        # Set the pose if needed
+        cmds.setAttr("rider_thumb_02_l.rotateZ", 22.03618142)
+        cmds.setAttr("rider_thumb_02_l.rotateY", 5.675048056)
+        cmds.setAttr("rider_thumb_02_l.rotateX", 22.94177697)
+        cmds.setAttr("rider_thumb_01_l.rotateZ", 4.168007795)
+        cmds.setAttr("rider_thumb_01_l.rotateY", -1.675941069)
+        cmds.setAttr("rider_thumb_01_l.rotateX", 9.067155995)
+        cmds.setAttr("rider_ring_02_l.rotateZ", 54.79433274)
+        cmds.setAttr("rider_ring_01_l.rotateZ", 54.76549897)
+        cmds.setAttr("rider_pinky_02_l.rotateZ", 54.79433274)
+        cmds.setAttr("rider_pinky_01_l.rotateZ", 53.17737673)
+        cmds.setAttr("rider_middle_02_l.rotateZ", 59.79453692)
+        cmds.setAttr("rider_middle_01_l.rotateZ", 51.05024479)
+        cmds.setAttr("rider_index_02_l.rotateZ", 64.20477863)
+        cmds.setAttr("rider_index_01_l.rotateZ", 46.94060635)
+        cmds.setAttr("rider_index_01_l.rotateY", 3.692762633)
+        cmds.setAttr("rider_index_01_l.rotateX", -8.969261741)
+
+        cmds.setAttr("rider_thumb_02_r.rotateZ", 22.03618142)
+        cmds.setAttr("rider_thumb_02_r.rotateY", 5.675048056)
+        cmds.setAttr("rider_thumb_02_r.rotateX", 22.94177697)
+        cmds.setAttr("rider_thumb_01_r.rotateZ", 4.168007795)
+        cmds.setAttr("rider_thumb_01_r.rotateY", -1.675941069)
+        cmds.setAttr("rider_thumb_01_r.rotateX", 9.067155995)
+        cmds.setAttr("rider_ring_02_r.rotateZ", 54.79433274)
+        cmds.setAttr("rider_ring_01_r.rotateZ", 54.76549897)
+        cmds.setAttr("rider_pinky_02_r.rotateZ", 54.79433274)
+        cmds.setAttr("rider_pinky_01_r.rotateZ", 53.17737673)
+        cmds.setAttr("rider_middle_02_r.rotateZ", 59.79453692)
+        cmds.setAttr("rider_middle_01_r.rotateZ", 51.05024479)
+        cmds.setAttr("rider_index_02_r.rotateZ", 64.20477863)
+        cmds.setAttr("rider_index_01_r.rotateZ", 46.94060635)
+        cmds.setAttr("rider_index_01_r.rotateY", 3.692762633)
+        cmds.setAttr("rider_index_01_r.rotateX", -8.969261741)
+
+        # Delete the history on the meshes so that they get baked down into their posed positions
+        cmds.delete(meshes, ch=True)
+
+        # Move the skeleton back to its bindpose (This script REQUIRES the bind pose to be the pose that we want to export the skeleton to.  This is the case for all known characters created with A.R.T.)
+        cmds.dagPose("root", bp=True, r=True) 
+
+        # Re-import the skinning we exported before onto the mesh so thats its bound in its posed position.
+        cmds.select(meshes)
+        impw.importMultipleSkinWeights("", True)
+
+        jointsToDelete = []
+
+         # remove unneeded joints
+        jointsToRemove = ["lower_fin", "lower_fin_end"]
+        jointToTransferTo = "pelvis"
+        transfer.LOD_transferWeights(meshes, jointsToRemove, jointToTransferTo)
+        jointsToDelete.extend(jointsToRemove)
+
+        jointsToRemove = ["turbine_hinge", "turbine", "turbine_root", "cockpit_foot_pedal_r", "cockpit_pedal_r", "cockpit_foot_pedal_l", "cockpit_pedal_l"]
+        jointToTransferTo = "spine_01"
+        transfer.LOD_transferWeights(meshes, jointsToRemove, jointToTransferTo)
+        jointsToDelete.extend(jointsToRemove)
+
+        jointsToRemove = ["controlArm_r_03", "controlArm_r_02"]
+        jointToTransferTo = "controlArm_r_01"
+        transfer.LOD_transferWeights(meshes, jointsToRemove, jointToTransferTo)
+        jointsToDelete.extend(jointsToRemove)
+
+        jointsToRemove = ["controlArm_l_03", "controlArm_l_02"]
+        jointToTransferTo = "controlArm_l_01"
+        transfer.LOD_transferWeights(meshes, jointsToRemove, jointToTransferTo)
+        jointsToDelete.extend(jointsToRemove)
+
+        jointsToRemove = ["rocket_l"]
+        jointToTransferTo = "rocket_root_l"
+        transfer.LOD_transferWeights(meshes, jointsToRemove, jointToTransferTo)
+        jointsToDelete.extend(jointsToRemove)
+
+        jointsToRemove = ["rocket_r"]
+        jointToTransferTo = "rocket_root_r"
+        transfer.LOD_transferWeights(meshes, jointsToRemove, jointToTransferTo)
+        jointsToDelete.extend(jointsToRemove)
+
+        jointsToRemove = ["upper_fin", "neck_01", "head"]
+        jointToTransferTo = "spine_03"
+        transfer.LOD_transferWeights(meshes, jointsToRemove, jointToTransferTo)
+        jointsToDelete.extend(jointsToRemove)
+
+        jointsToRemove = ["tube_lats_r"]
+        jointToTransferTo = "clavicle_r"
+        transfer.LOD_transferWeights(meshes, jointsToRemove, jointToTransferTo)
+        jointsToDelete.extend(jointsToRemove)
+
+        jointsToRemove = ["arm_jet_l"]
+        jointToTransferTo = "upperarm_l"
+        transfer.LOD_transferWeights(meshes, jointsToRemove, jointToTransferTo)
+        jointsToDelete.extend(jointsToRemove)
+
+        jointsToRemove = ["tube_arm_r", "arm_jet_r"]
+        jointToTransferTo = "upperarm_r"
+        transfer.LOD_transferWeights(meshes, jointsToRemove, jointToTransferTo)
+        jointsToDelete.extend(jointsToRemove)
+
+        jointsToRemove = ["bay_door_a", "bay_door_b", "hand_l", "recoil_l", "tube_arm_l"]
+        jointToTransferTo = "lowerarm_l"
+        transfer.LOD_transferWeights(meshes, jointsToRemove, jointToTransferTo)
+        jointsToDelete.extend(jointsToRemove)
+
+        jointsToRemove = ["rocket_arm_b"]
+        jointToTransferTo = "rocket_arm"
+        transfer.LOD_transferWeights(meshes, jointsToRemove, jointToTransferTo)
+        jointsToDelete.extend(jointsToRemove)
+
+        jointsToRemove = ["hand_r", "recoil_r", "shell_1", "shell_2", "shell_3"]
+        jointToTransferTo = "lowerarm_r"
+        transfer.LOD_transferWeights(meshes, jointsToRemove, jointToTransferTo)
+        jointsToDelete.extend(jointsToRemove)
+
+        jointsToRemove = ["index_02_r", "index_03_r"]
+        jointToTransferTo = "index_01_r"
+        transfer.LOD_transferWeights(meshes, jointsToRemove, jointToTransferTo)
+        jointsToDelete.extend(jointsToRemove)
+
+        jointsToRemove = ["middle_02_r", "middle_03_r"]
+        jointToTransferTo = "middle_01_r"
+        transfer.LOD_transferWeights(meshes, jointsToRemove, jointToTransferTo)
+        jointsToDelete.extend(jointsToRemove)
+
+        jointsToRemove = ["ring_02_r", "ring_03_r"]
+        jointToTransferTo = "ring_01_r"
+        transfer.LOD_transferWeights(meshes, jointsToRemove, jointToTransferTo)
+        jointsToDelete.extend(jointsToRemove)
+
+        jointsToRemove = ["leg_socket_l"]
+        jointToTransferTo = "leg_socket_root_l"
+        transfer.LOD_transferWeights(meshes, jointsToRemove, jointToTransferTo)
+        jointsToDelete.extend(jointsToRemove)
+
+        jointsToRemove = ["leg_socket_r"]
+        jointToTransferTo = "leg_socket_root_r"
+        transfer.LOD_transferWeights(meshes, jointsToRemove, jointToTransferTo)
+        jointsToDelete.extend(jointsToRemove)
+
+        jointsToRemove = ["leg_jet_l"]
+        jointToTransferTo = "calf_l"
+        transfer.LOD_transferWeights(meshes, jointsToRemove, jointToTransferTo)
+        jointsToDelete.extend(jointsToRemove)
+
+        jointsToRemove = ["leg_jet_r"]
+        jointToTransferTo = "calf_r"
+        transfer.LOD_transferWeights(meshes, jointsToRemove, jointToTransferTo)
+        jointsToDelete.extend(jointsToRemove)
+
+        jointsToRemove = ["foot_spike_l"]
+        jointToTransferTo = "foot_l"
+        transfer.LOD_transferWeights(meshes, jointsToRemove, jointToTransferTo)
+        jointsToDelete.extend(jointsToRemove)
+
+        jointsToRemove = ["foot_spike_r"]
+        jointToTransferTo = "foot_r"
+        transfer.LOD_transferWeights(meshes, jointsToRemove, jointToTransferTo)
+        jointsToDelete.extend(jointsToRemove)
+
+        jointsToRemove = ["rider_index_01_l", "rider_index_02_l", "rider_middle_01_l", "rider_middle_02_l", "rider_pinky_01_l", "rider_pinky_02_l", "rider_ring_01_l", "rider_ring_02_l", "rider_thumb_01_l", "rider_thumb_02_l"]
+        jointToTransferTo = "rider_hand_l"
+        transfer.LOD_transferWeights(meshes, jointsToRemove, jointToTransferTo)
+        jointsToDelete.extend(jointsToRemove)
+
+        jointsToRemove = ["rider_index_01_r", "rider_index_02_r", "rider_middle_01_r", "rider_middle_02_r", "rider_pinky_01_r", "rider_pinky_02_r", "rider_ring_01_r", "rider_ring_02_r", "rider_thumb_01_r", "rider_thumb_02_r"]
+        jointToTransferTo = "rider_hand_r"
+        transfer.LOD_transferWeights(meshes, jointsToRemove, jointToTransferTo)
+        jointsToDelete.extend(jointsToRemove)
+
+        jointsToRemove = ["rider_earMain_l_01", "rider_earMain_l_02", "rider_earMain_l_03", "rider_earMain_r_01", "rider_earMain_r_02", "rider_earMain_r_03", "rider_eyeball_l", "rider_eyeball_r", "rider_eyelid_bott_l", "rider_eyelid_bott_r", "rider_eyelid_top_l", "rider_eyelid_top_r", "rider_jaw"]
+        jointToTransferTo = "rider_head"
+        transfer.LOD_transferWeights(meshes, jointsToRemove, jointToTransferTo)
+        jointsToDelete.extend(jointsToRemove)
+
+        jointsToRemove = ["rider_heel_l"]
+        jointToTransferTo = "rider_calf_l"
+        transfer.LOD_transferWeights(meshes, jointsToRemove, jointToTransferTo)
+        jointsToDelete.extend(jointsToRemove)
+
+        jointsToRemove = ["rider_heel_r"]
+        jointToTransferTo = "rider_calf_r"
+        transfer.LOD_transferWeights(meshes, jointsToRemove, jointToTransferTo)
+        jointsToDelete.extend(jointsToRemove)
+
+        jointsToRemove = ["rider_ball_l"]
+        jointToTransferTo = "rider_foot_l"
+        transfer.LOD_transferWeights(meshes, jointsToRemove, jointToTransferTo)
+        jointsToDelete.extend(jointsToRemove)
+
+        jointsToRemove = ["rider_ball_r"]
+        jointToTransferTo = "rider_foot_r"
+        transfer.LOD_transferWeights(meshes, jointsToRemove, jointToTransferTo)
+        jointsToDelete.extend(jointsToRemove)
+        
+        # Delete the joints we no longer want.
+        for one in jointsToDelete:
+            if cmds.objExists(one):
+                cmds.delete(one)
+                
+        # Export the fbx file.
+        export.exportLOD("HyperBreach_LOD2.fbx", "LOD2")
+
+        # Re-open the file without saving.
+        fullPath = cmds.file(q = True, sceneName = True)
+        cmds.file(fullPath, open=True, f=True)
+
+        print "LOD2 Exported Successfully!!!!!!!!!!!!!!!"
+
+HyperBreachLOD2()
+
+# HyperBreach LOD3 -------------------------------------------------------------------------------------------------------------
+def HyperBreachLOD3():
+    if saveResult == "Yes":
+
+        cmds.file(save=True)
+    
+        meshes = ["backlauncher_l_geo_LOD3", "backlauncher_r_geo_LOD3", "backspine_geo_LOD3", "bicep_l_geo_LOD3", "bicep_l_hatch_geo_LOD3", "bicep_r_geo_LOD3", "bicep_r_hatch_geo_LOD3", "bodyback_geo_LOD3", "bodypanel_l_geo_LOD3", "bodyshell_geo_LOD3", "calfhatch_l_geo_LOD3", "calfhatch_r_geo_LOD3", "calfjet_l_geo_LOD3", "calfjet_r_geo_LOD3", "calfs_geo_LOD3", "clav_l_geo_LOD3", "clav_r_geo_LOD3", "controlarm_l1_geo_LOD3", "controlarm_l2_geo_LOD3", "controlarm_r1_geo_LOD3", "controlarm_r2_geo_LOD3", "controlgrip_l_geo_LOD3", "controlgrip_r_geo_LOD3", "foot_l_geo_LOD3", "foot_r_geo_LOD3", "footflap_l_geo_LOD3", "footflap_r_geo_LOD3", "grenade1_geo_LOD3", "grenade2_geo_LOD3", "grenade3_geo_LOD3", "grenadefinger_1a_geo_LOD3", "grenadefinger_1b_geo_LOD3", "grenadefinger_1c_geo_LOD3", "grenadefinger_2a_geo_LOD3", "grenadefinger_2b_geo_LOD3", "grenadefinger_2c_geo_LOD3", "grenadefinger_3a_geo_LOD3", "grenadefinger_3b_geo_LOD3", "grenadefinger_3c_geo_LOD3", "grenadelauncherr_geo_LOD3", "grenadepiston_geo_LOD3", "hipcenter_geo_LOD3", "hipplate_geo_LOD3", "hose_geo_LOD3", "jetarm_l_geo_LOD3", "jetarm_r_geo_LOD3", "jetfoot_l_geo_LOD3", "jetfoot_r_geo_LOD3", "leg_l_geo_LOD3", "leg_r_geo_LOD3", "lowerspine_geo_LOD3", "mine_geo_LOD3", "minebase_geo_LOD3", "minedeploy_geo_LOD3", "minepivot_geo_LOD3", "missle_geo_LOD3", "missledoor1_geo_LOD3", "missledoor2_geo_LOD3", "misslelatch1_geo_LOD3", "misslelatch2_geo_LOD3", "misslelatch3_geo_LOD3", "misslelauncher_geo_LOD3", "misslepiston_geo_LOD3", "pedal_l_geo_LOD3", "pedal_r_geo_LOD3", "pedalarm_l_geo_LOD3", "pedalarm_r_geo_LOD3", "seat_geo_LOD3", "seatbase_geo_LOD3", "seatpivot_geo_LOD3", "shoulder_l_geo_LOD3", "shoulder_r_geo_LOD3", "shoulderbone_l_geo_LOD3", "shoulderbone_r_geo_LOD3", "shoulderpad_l_geo_LOD3", "shoulderpad_r_geo_LOD3", "tailplate_geo_LOD3", "toe_l_geo_LOD3", "toe_r_geo_LOD3", "turbine_geo_LOD3", "Rider:cigar_geo_LOD3", "Rider:cuffs_geo_LOD3", "Rider:eye_balls_geo_LOD3", "Rider:eye_cover_geo_LOD3", "Rider:feet_geo_LOD3", "Rider:gun_barrel_geo_LOD3", "Rider:gun_clip_geo_LOD3", "Rider:gun_flaps_geo_LOD3", "Rider:gun_geo_LOD3", "Rider:hands_geo_LOD3", "Rider:head_geo_LOD3", "Rider:holster_geo_LOD3", "Rider:HyperBreach_Rabbit_Eyelashes Rider:HyperBreach_Rabbit_Hair Rider:insignia_geo_LOD3", "Rider:leggings_geo_LOD3", "Rider:legs_geo_LOD3", "Rider:lower_belt_geo_LOD3", "Rider:lower_teeth_geo_LOD3", "Rider:shoulders_geo_LOD3", "Rider:tongue_geo_LOD3", "Rider:torso_geo_LOD3", "Rider:upper_teeth_geo_LOD3"]
+    
+        # Export a skin weights file for each mesh we plan to operate on.
+        cmds.select(meshes)
+        menu.exportSelectedSkinWeights(True)
+
+        # Set the pose if needed
+        cmds.setAttr("rider_thumb_02_l.rotateZ", 22.03618142)
+        cmds.setAttr("rider_thumb_02_l.rotateY", 5.675048056)
+        cmds.setAttr("rider_thumb_02_l.rotateX", 22.94177697)
+        cmds.setAttr("rider_thumb_01_l.rotateZ", 4.168007795)
+        cmds.setAttr("rider_thumb_01_l.rotateY", -1.675941069)
+        cmds.setAttr("rider_thumb_01_l.rotateX", 9.067155995)
+        cmds.setAttr("rider_ring_02_l.rotateZ", 54.79433274)
+        cmds.setAttr("rider_ring_01_l.rotateZ", 54.76549897)
+        cmds.setAttr("rider_pinky_02_l.rotateZ", 54.79433274)
+        cmds.setAttr("rider_pinky_01_l.rotateZ", 53.17737673)
+        cmds.setAttr("rider_middle_02_l.rotateZ", 59.79453692)
+        cmds.setAttr("rider_middle_01_l.rotateZ", 51.05024479)
+        cmds.setAttr("rider_index_02_l.rotateZ", 64.20477863)
+        cmds.setAttr("rider_index_01_l.rotateZ", 46.94060635)
+        cmds.setAttr("rider_index_01_l.rotateY", 3.692762633)
+        cmds.setAttr("rider_index_01_l.rotateX", -8.969261741)
+
+        cmds.setAttr("rider_thumb_02_r.rotateZ", 22.03618142)
+        cmds.setAttr("rider_thumb_02_r.rotateY", 5.675048056)
+        cmds.setAttr("rider_thumb_02_r.rotateX", 22.94177697)
+        cmds.setAttr("rider_thumb_01_r.rotateZ", 4.168007795)
+        cmds.setAttr("rider_thumb_01_r.rotateY", -1.675941069)
+        cmds.setAttr("rider_thumb_01_r.rotateX", 9.067155995)
+        cmds.setAttr("rider_ring_02_r.rotateZ", 54.79433274)
+        cmds.setAttr("rider_ring_01_r.rotateZ", 54.76549897)
+        cmds.setAttr("rider_pinky_02_r.rotateZ", 54.79433274)
+        cmds.setAttr("rider_pinky_01_r.rotateZ", 53.17737673)
+        cmds.setAttr("rider_middle_02_r.rotateZ", 59.79453692)
+        cmds.setAttr("rider_middle_01_r.rotateZ", 51.05024479)
+        cmds.setAttr("rider_index_02_r.rotateZ", 64.20477863)
+        cmds.setAttr("rider_index_01_r.rotateZ", 46.94060635)
+        cmds.setAttr("rider_index_01_r.rotateY", 3.692762633)
+        cmds.setAttr("rider_index_01_r.rotateX", -8.969261741)
+
+        # Delete the history on the meshes so that they get baked down into their posed positions
+        cmds.delete(meshes, ch=True)
+
+        # Move the skeleton back to its bindpose (This script REQUIRES the bind pose to be the pose that we want to export the skeleton to.  This is the case for all known characters created with A.R.T.)
+        cmds.dagPose("root", bp=True, r=True) 
+
+        # Re-import the skinning we exported before onto the mesh so thats its bound in its posed position.
+        cmds.select(meshes)
+        impw.importMultipleSkinWeights("", True)
+        
+        jointsToDelete = []
+
+        # remove unneeded joints
+        jointsToRemove = ["hip_l", "hip_r", "lower_fin", "lower_fin_end"]
+        jointToTransferTo = "pelvis"
+        transfer.LOD_transferWeights(meshes, jointsToRemove, jointToTransferTo)
+        jointsToDelete.extend(jointsToRemove)
+
+        jointsToRemove = ["turbine_hinge", "turbine", "turbine_root", "controlArm_r_03", "controlArm_r_02", "controlArm_r_01", "controlArm_l_03", "controlArm_l_02", "controlArm_l_01", "cockpit_seat", "cockpit_seat_base", "cockpit_foot_pedal_r", "cockpit_pedal_r", "cockpit_foot_pedal_l", "cockpit_pedal_l"]
+        jointToTransferTo = "spine_01"
+        transfer.LOD_transferWeights(meshes, jointsToRemove, jointToTransferTo)
+        jointsToDelete.extend(jointsToRemove)
+
+        jointsToRemove = ["rocket_transform", "rocket_root_l", "rocket_l", "rocket_root_r", "rocket_r"]
+        jointToTransferTo = "spine_02"
+        transfer.LOD_transferWeights(meshes, jointsToRemove, jointToTransferTo)
+        jointsToDelete.extend(jointsToRemove)
+
+        jointsToRemove = ["upper_fin", "shoulder_pad_r", "shoulder_pad_l", "neck_01", "head"]
+        jointToTransferTo = "spine_03"
+        transfer.LOD_transferWeights(meshes, jointsToRemove, jointToTransferTo)
+        jointsToDelete.extend(jointsToRemove)
+
+        jointsToRemove = ["cannon", "mine"]
+        jointToTransferTo = "cannon_arm_2"
+        transfer.LOD_transferWeights(meshes, jointsToRemove, jointToTransferTo)
+        jointsToDelete.extend(jointsToRemove)
+
+        jointsToRemove = ["shoulder_mid_l"]
+        jointToTransferTo = "clavicle_l"
+        transfer.LOD_transferWeights(meshes, jointsToRemove, jointToTransferTo)
+        jointsToDelete.extend(jointsToRemove)
+
+        jointsToRemove = ["tube_lats_r", "shoulder_mid_r"]
+        jointToTransferTo = "clavicle_r"
+        transfer.LOD_transferWeights(meshes, jointsToRemove, jointToTransferTo)
+        jointsToDelete.extend(jointsToRemove)
+
+        jointsToRemove = ["arm_jet_l"]
+        jointToTransferTo = "upperarm_l"
+        transfer.LOD_transferWeights(meshes, jointsToRemove, jointToTransferTo)
+        jointsToDelete.extend(jointsToRemove)
+
+        jointsToRemove = ["tube_arm_r", "arm_jet_r"]
+        jointToTransferTo = "upperarm_r"
+        transfer.LOD_transferWeights(meshes, jointsToRemove, jointToTransferTo)
+        jointsToDelete.extend(jointsToRemove)
+
+        jointsToRemove = ["bay_door_a", "bay_door_b", "hand_l", "recoil_l", "tube_arm_l"]
+        jointToTransferTo = "lowerarm_l"
+        transfer.LOD_transferWeights(meshes, jointsToRemove, jointToTransferTo)
+        jointsToDelete.extend(jointsToRemove)
+
+        jointsToRemove = ["rocket_arm_b"]
+        jointToTransferTo = "rocket_arm"
+        transfer.LOD_transferWeights(meshes, jointsToRemove, jointToTransferTo)
+        jointsToDelete.extend(jointsToRemove)
+
+        jointsToRemove = ["hand_r", "index_01_r", "index_02_r", "index_03_r", "middle_01_r", "middle_02_r", "middle_03_r", "ring_01_r", "ring_02_r", "ring_03_r", "recoil_r", "shell_1", "shell_2", "shell_3"]
+        jointToTransferTo = "lowerarm_r"
+        transfer.LOD_transferWeights(meshes, jointsToRemove, jointToTransferTo)
+        jointsToDelete.extend(jointsToRemove)
+
+        jointsToRemove = ["leg_socket_root_l", "leg_socket_l"]
+        jointToTransferTo = "thigh_l"
+        transfer.LOD_transferWeights(meshes, jointsToRemove, jointToTransferTo)
+        jointsToDelete.extend(jointsToRemove)
+
+        jointsToRemove = ["leg_socket_root_r", "leg_socket_r"]
+        jointToTransferTo = "thigh_r"
+        transfer.LOD_transferWeights(meshes, jointsToRemove, jointToTransferTo)
+        jointsToDelete.extend(jointsToRemove)
+
+        jointsToRemove = ["calf_twist_01_l", "leg_jet_l"]
+        jointToTransferTo = "calf_l"
+        transfer.LOD_transferWeights(meshes, jointsToRemove, jointToTransferTo)
+        jointsToDelete.extend(jointsToRemove)
+
+        jointsToRemove = ["leg_jet_r", "calf_twist_01_r"]
+        jointToTransferTo = "calf_r"
+        transfer.LOD_transferWeights(meshes, jointsToRemove, jointToTransferTo)
+        jointsToDelete.extend(jointsToRemove)
+
+        jointsToRemove = ["ball_l", "foot_spike_l"]
+        jointToTransferTo = "foot_l"
+        transfer.LOD_transferWeights(meshes, jointsToRemove, jointToTransferTo)
+        jointsToDelete.extend(jointsToRemove)
+
+        jointsToRemove = ["foot_spike_r", "ball_r"]
+        jointToTransferTo = "foot_r"
+        transfer.LOD_transferWeights(meshes, jointsToRemove, jointToTransferTo)
+        jointsToDelete.extend(jointsToRemove)
+
+        jointsToRemove = ["rider_index_01_l", "rider_index_02_l", "rider_middle_01_l", "rider_middle_02_l", "rider_pinky_01_l", "rider_pinky_02_l", "rider_ring_01_l", "rider_ring_02_l", "rider_thumb_01_l", "rider_thumb_02_l"]
+        jointToTransferTo = "rider_hand_l"
+        transfer.LOD_transferWeights(meshes, jointsToRemove, jointToTransferTo)
+        jointsToDelete.extend(jointsToRemove)
+
+        jointsToRemove = ["rider_index_01_r", "rider_index_02_r", "rider_middle_01_r", "rider_middle_02_r", "rider_pinky_01_r", "rider_pinky_02_r", "rider_ring_01_r", "rider_ring_02_r", "rider_thumb_01_r", "rider_thumb_02_r"]
+        jointToTransferTo = "rider_hand_r"
+        transfer.LOD_transferWeights(meshes, jointsToRemove, jointToTransferTo)
+        jointsToDelete.extend(jointsToRemove)
+
+        jointsToRemove = ["rider_earMain_l_01", "rider_earMain_l_02", "rider_earMain_l_03", "rider_earMain_r_01", "rider_earMain_r_02", "rider_earMain_r_03", "rider_eyeball_l", "rider_eyeball_r", "rider_eyelid_bott_l", "rider_eyelid_bott_r", "rider_eyelid_top_l", "rider_eyelid_top_r", "rider_jaw"]
+        jointToTransferTo = "rider_head"
+        transfer.LOD_transferWeights(meshes, jointsToRemove, jointToTransferTo)
+        jointsToDelete.extend(jointsToRemove)
+
+        jointsToRemove = ["rider_heel_l"]
+        jointToTransferTo = "rider_calf_l"
+        transfer.LOD_transferWeights(meshes, jointsToRemove, jointToTransferTo)
+        jointsToDelete.extend(jointsToRemove)
+
+        jointsToRemove = ["rider_heel_r"]
+        jointToTransferTo = "rider_calf_r"
+        transfer.LOD_transferWeights(meshes, jointsToRemove, jointToTransferTo)
+        jointsToDelete.extend(jointsToRemove)
+
+        jointsToRemove = ["rider_ball_l"]
+        jointToTransferTo = "rider_foot_l"
+        transfer.LOD_transferWeights(meshes, jointsToRemove, jointToTransferTo)
+        jointsToDelete.extend(jointsToRemove)
+
+        jointsToRemove = ["rider_ball_r"]
+        jointToTransferTo = "rider_foot_r"
+        transfer.LOD_transferWeights(meshes, jointsToRemove, jointToTransferTo)
+        jointsToDelete.extend(jointsToRemove)
+
+        # Delete the joints we no longer want.
+        for one in jointsToDelete:
+            if cmds.objExists(one):
+                cmds.delete(one)
+                
+        # Export the fbx file.
+        export.exportLOD("HyperBreach_LOD3.fbx", "LOD3")
+        
+        # Re-open the file without saving.
+        fullPath = cmds.file(q = True, sceneName = True)
+        cmds.file(fullPath, open=True, f=True)
+
+        print "LOD3 Exported Successfully!!!!!!!!!!!!!!!"
+
+HyperBreachLOD3()
+
+# HyperBreach LOD4 -------------------------------------------------------------------------------------------------------------
+def HyperBreachLOD4():
+    if saveResult == "Yes":
+
+        cmds.file(save=True)
+    
+        meshes = ["backlauncher_l_geo_LOD4", "backlauncher_r_geo_LOD4", "backspine_geo_LOD4", "bicep_l_geo_LOD4", "bicep_l_hatch_geo_LOD4", "bicep_r_geo_LOD4", "bicep_r_hatch_geo_LOD4", "bodyback_geo_LOD4", "bodypanel_l_geo_LOD4", "bodyshell_geo_LOD4", "calfhatch_l_geo_LOD4", "calfhatch_r_geo_LOD4", "calfjet_l_geo_LOD4", "calfjet_r_geo_LOD4", "calfs_geo_LOD4", "clav_l_geo_LOD4", "clav_r_geo_LOD4", "controlarm_l1_geo_LOD4", "controlarm_l2_geo_LOD4", "controlarm_r1_geo_LOD4", "controlarm_r2_geo_LOD4", "controlgrip_l_geo_LOD4", "controlgrip_r_geo_LOD4", "foot_l_geo_LOD4", "foot_r_geo_LOD4", "footflap_l_geo_LOD4", "footflap_r_geo_LOD4", "grenade1_geo_LOD4", "grenade2_geo_LOD4", "grenade3_geo_LOD4", "grenadefinger_1a_geo_LOD4", "grenadefinger_1b_geo_LOD4", "grenadefinger_1c_geo_LOD4", "grenadefinger_2a_geo_LOD4", "grenadefinger_2b_geo_LOD4", "grenadefinger_2c_geo_LOD4", "grenadefinger_3a_geo_LOD4", "grenadefinger_3b_geo_LOD4", "grenadefinger_3c_geo_LOD4", "grenadelauncherr_geo_LOD4", "grenadepiston_geo_LOD4", "hipcenter_geo_LOD4", "hipplate_geo_LOD4", "hose_geo_LOD4", "jetarm_l_geo_LOD4", "jetarm_r_geo_LOD4", "jetfoot_l_geo_LOD4", "jetfoot_r_geo_LOD4", "leg_l_geo_LOD4", "leg_r_geo_LOD4", "lowerspine_geo_LOD4", "mine_geo_LOD4", "minebase_geo_LOD4", "minedeploy_geo_LOD4", "minepivot_geo_LOD4", "missle_geo_LOD4", "missledoor1_geo_LOD4", "missledoor2_geo_LOD4", "misslelatch1_geo_LOD4", "misslelatch2_geo_LOD4", "misslelatch3_geo_LOD4", "misslelauncher_geo_LOD4", "misslepiston_geo_LOD4", "pedal_l_geo_LOD4", "pedal_r_geo_LOD4", "pedalarm_l_geo_LOD4", "pedalarm_r_geo_LOD4", "seat_geo_LOD4", "seatbase_geo_LOD4", "seatpivot_geo_LOD4", "shoulder_l_geo_LOD4", "shoulder_r_geo_LOD4", "shoulderbone_l_geo_LOD4", "shoulderbone_r_geo_LOD4", "shoulderpad_l_geo_LOD4", "shoulderpad_r_geo_LOD4", "tailplate_geo_LOD4", "toe_l_geo_LOD4", "toe_r_geo_LOD4", "turbine_geo_LOD4", "Rider:cigar_geo_LOD4", "Rider:cuffs_geo_LOD4", "Rider:eye_balls_geo_LOD4", "Rider:eye_cover_geo_LOD4", "Rider:feet_geo_LOD4", "Rider:gun_barrel_geo_LOD4", "Rider:gun_clip_geo_LOD4", "Rider:gun_flaps_geo_LOD4", "Rider:gun_geo_LOD4", "Rider:hands_geo_LOD4", "Rider:head_geo_LOD4", "Rider:holster_geo_LOD4", "Rider:HyperBreach_Rabbit_Eyelashes Rider:HyperBreach_Rabbit_Hair Rider:insignia_geo_LOD4", "Rider:leggings_geo_LOD4", "Rider:legs_geo_LOD4", "Rider:lower_belt_geo_LOD4", "Rider:lower_teeth_geo_LOD4", "Rider:shoulders_geo_LOD4", "Rider:tongue_geo_LOD4", "Rider:torso_geo_LOD4", "Rider:upper_teeth_geo_LOD4"]
+
+        # Export a skin weights file for each mesh we plan to operate on.
+        cmds.select(meshes)
+        menu.exportSelectedSkinWeights(True)
+
+        # Set the pose if needed
+        cmds.setAttr("rider_thumb_02_l.rotateZ", 22.03618142)
+        cmds.setAttr("rider_thumb_02_l.rotateY", 5.675048056)
+        cmds.setAttr("rider_thumb_02_l.rotateX", 22.94177697)
+        cmds.setAttr("rider_thumb_01_l.rotateZ", 4.168007795)
+        cmds.setAttr("rider_thumb_01_l.rotateY", -1.675941069)
+        cmds.setAttr("rider_thumb_01_l.rotateX", 9.067155995)
+        cmds.setAttr("rider_ring_02_l.rotateZ", 54.79433274)
+        cmds.setAttr("rider_ring_01_l.rotateZ", 54.76549897)
+        cmds.setAttr("rider_pinky_02_l.rotateZ", 54.79433274)
+        cmds.setAttr("rider_pinky_01_l.rotateZ", 53.17737673)
+        cmds.setAttr("rider_middle_02_l.rotateZ", 59.79453692)
+        cmds.setAttr("rider_middle_01_l.rotateZ", 51.05024479)
+        cmds.setAttr("rider_index_02_l.rotateZ", 64.20477863)
+        cmds.setAttr("rider_index_01_l.rotateZ", 46.94060635)
+        cmds.setAttr("rider_index_01_l.rotateY", 3.692762633)
+        cmds.setAttr("rider_index_01_l.rotateX", -8.969261741)
+
+        cmds.setAttr("rider_thumb_02_r.rotateZ", 22.03618142)
+        cmds.setAttr("rider_thumb_02_r.rotateY", 5.675048056)
+        cmds.setAttr("rider_thumb_02_r.rotateX", 22.94177697)
+        cmds.setAttr("rider_thumb_01_r.rotateZ", 4.168007795)
+        cmds.setAttr("rider_thumb_01_r.rotateY", -1.675941069)
+        cmds.setAttr("rider_thumb_01_r.rotateX", 9.067155995)
+        cmds.setAttr("rider_ring_02_r.rotateZ", 54.79433274)
+        cmds.setAttr("rider_ring_01_r.rotateZ", 54.76549897)
+        cmds.setAttr("rider_pinky_02_r.rotateZ", 54.79433274)
+        cmds.setAttr("rider_pinky_01_r.rotateZ", 53.17737673)
+        cmds.setAttr("rider_middle_02_r.rotateZ", 59.79453692)
+        cmds.setAttr("rider_middle_01_r.rotateZ", 51.05024479)
+        cmds.setAttr("rider_index_02_r.rotateZ", 64.20477863)
+        cmds.setAttr("rider_index_01_r.rotateZ", 46.94060635)
+        cmds.setAttr("rider_index_01_r.rotateY", 3.692762633)
+        cmds.setAttr("rider_index_01_r.rotateX", -8.969261741)
+
+        # Delete the history on the meshes so that they get baked down into their posed positions
+        cmds.delete(meshes, ch=True)
+
+        # Move the skeleton back to its bindpose (This script REQUIRES the bind pose to be the pose that we want to export the skeleton to.  This is the case for all known characters created with A.R.T.)
+        cmds.dagPose("root", bp=True, r=True) 
+
+        # Re-import the skinning we exported before onto the mesh so thats its bound in its posed position.
+        cmds.select(meshes)
+        impw.importMultipleSkinWeights("", True)
+        
+        jointsToDelete = []
+
+        # remove unneeded joints
+        jointsToRemove = ["hip_l", "hip_r", "lower_fin", "lower_fin_end"]
+        jointToTransferTo = "pelvis"
+        transfer.LOD_transferWeights(meshes, jointsToRemove, jointToTransferTo)
+        jointsToDelete.extend(jointsToRemove)
+
+        jointsToRemove = ["turbine_hinge", "turbine", "turbine_root", "controlArm_r_03", "controlArm_r_02", "controlArm_r_01", "controlArm_l_03", "controlArm_l_02", "controlArm_l_01", "cockpit_seat", "cockpit_seat_base", "cockpit_foot_pedal_r", "cockpit_pedal_r", "cockpit_foot_pedal_l", "cockpit_pedal_l"]
+        jointToTransferTo = "spine_01"
+        transfer.LOD_transferWeights(meshes, jointsToRemove, jointToTransferTo)
+        jointsToDelete.extend(jointsToRemove)
+
+        jointsToRemove = ["rocket_transform", "rocket_root_l", "rocket_l", "rocket_root_r", "rocket_r"]
+        jointToTransferTo = "spine_02"
+        transfer.LOD_transferWeights(meshes, jointsToRemove, jointToTransferTo)
+        jointsToDelete.extend(jointsToRemove)
+
+        jointsToRemove = ["cannon_root", "cannon_arm_1", "cannon_arm_2", "cannon", "mine", "upper_fin", "shoulder_pad_r", "shoulder_pad_l", "neck_01", "head"]
+        jointToTransferTo = "spine_03"
+        transfer.LOD_transferWeights(meshes, jointsToRemove, jointToTransferTo)
+        jointsToDelete.extend(jointsToRemove)
+
+        jointsToRemove = ["shoulder_mid_l"]
+        jointToTransferTo = "clavicle_l"
+        transfer.LOD_transferWeights(meshes, jointsToRemove, jointToTransferTo)
+        jointsToDelete.extend(jointsToRemove)
+
+        jointsToRemove = ["tube_lats_r", "shoulder_mid_r"]
+        jointToTransferTo = "clavicle_r"
+        transfer.LOD_transferWeights(meshes, jointsToRemove, jointToTransferTo)
+        jointsToDelete.extend(jointsToRemove)
+
+        jointsToRemove = ["arm_jet_l"]
+        jointToTransferTo = "upperarm_l"
+        transfer.LOD_transferWeights(meshes, jointsToRemove, jointToTransferTo)
+        jointsToDelete.extend(jointsToRemove)
+
+        jointsToRemove = ["tube_arm_r", "arm_jet_r"]
+        jointToTransferTo = "upperarm_r"
+        transfer.LOD_transferWeights(meshes, jointsToRemove, jointToTransferTo)
+        jointsToDelete.extend(jointsToRemove)
+
+        jointsToRemove = ["bay_door_a", "bay_door_b", "hand_l", "recoil_l", "rocket_arm_root", "rocket_arm", "rocket_arm_b", "rocket", "tube_arm_l"]
+        jointToTransferTo = "lowerarm_l"
+        transfer.LOD_transferWeights(meshes, jointsToRemove, jointToTransferTo)
+        jointsToDelete.extend(jointsToRemove)
+
+        jointsToRemove = ["hand_r", "index_01_r", "index_02_r", "index_03_r", "middle_01_r", "middle_02_r", "middle_03_r", "ring_01_r", "ring_02_r", "ring_03_r", "recoil_r", "shell_1", "shell_2", "shell_3"]
+        jointToTransferTo = "lowerarm_r"
+        transfer.LOD_transferWeights(meshes, jointsToRemove, jointToTransferTo)
+        jointsToDelete.extend(jointsToRemove)
+
+        jointsToRemove = ["leg_socket_root_l", "leg_socket_l"]
+        jointToTransferTo = "thigh_l"
+        transfer.LOD_transferWeights(meshes, jointsToRemove, jointToTransferTo)
+        jointsToDelete.extend(jointsToRemove)
+
+        jointsToRemove = ["leg_socket_root_r", "leg_socket_r"]
+        jointToTransferTo = "thigh_r"
+        transfer.LOD_transferWeights(meshes, jointsToRemove, jointToTransferTo)
+        jointsToDelete.extend(jointsToRemove)
+
+        jointsToRemove = ["calf_twist_01_l", "leg_jet_l"]
+        jointToTransferTo = "calf_l"
+        transfer.LOD_transferWeights(meshes, jointsToRemove, jointToTransferTo)
+        jointsToDelete.extend(jointsToRemove)
+
+        jointsToRemove = ["leg_jet_r", "calf_twist_01_r"]
+        jointToTransferTo = "calf_r"
+        transfer.LOD_transferWeights(meshes, jointsToRemove, jointToTransferTo)
+        jointsToDelete.extend(jointsToRemove)
+
+        jointsToRemove = ["ball_l", "foot_spike_l"]
+        jointToTransferTo = "foot_l"
+        transfer.LOD_transferWeights(meshes, jointsToRemove, jointToTransferTo)
+        jointsToDelete.extend(jointsToRemove)
+
+        jointsToRemove = ["foot_spike_r", "ball_r"]
+        jointToTransferTo = "foot_r"
+        transfer.LOD_transferWeights(meshes, jointsToRemove, jointToTransferTo)
+        jointsToDelete.extend(jointsToRemove)
+
+        jointsToRemove = ["rider_index_01_l", "rider_index_02_l", "rider_middle_01_l", "rider_middle_02_l", "rider_pinky_01_l", "rider_pinky_02_l", "rider_ring_01_l", "rider_ring_02_l", "rider_thumb_01_l", "rider_thumb_02_l"]
+        jointToTransferTo = "rider_hand_l"
+        transfer.LOD_transferWeights(meshes, jointsToRemove, jointToTransferTo)
+        jointsToDelete.extend(jointsToRemove)
+
+        jointsToRemove = ["rider_index_01_r", "rider_index_02_r", "rider_middle_01_r", "rider_middle_02_r", "rider_pinky_01_r", "rider_pinky_02_r", "rider_ring_01_r", "rider_ring_02_r", "rider_thumb_01_r", "rider_thumb_02_r"]
+        jointToTransferTo = "rider_hand_r"
+        transfer.LOD_transferWeights(meshes, jointsToRemove, jointToTransferTo)
+        jointsToDelete.extend(jointsToRemove)
+
+        jointsToRemove = ["rider_neck_02", "rider_head", "rider_earMain_l_01", "rider_earMain_l_02", "rider_earMain_l_03", "rider_earMain_r_01", "rider_earMain_r_02", "rider_earMain_r_03", "rider_eyeball_l", "rider_eyeball_r", "rider_eyelid_bott_l", "rider_eyelid_bott_r", "rider_eyelid_top_l", "rider_eyelid_top_r", "rider_jaw"]
+        jointToTransferTo = "rider_neck_01"
+        transfer.LOD_transferWeights(meshes, jointsToRemove, jointToTransferTo)
+        jointsToDelete.extend(jointsToRemove)
+
+        jointsToRemove = ["rider_heel_l"]
+        jointToTransferTo = "rider_calf_l"
+        transfer.LOD_transferWeights(meshes, jointsToRemove, jointToTransferTo)
+        jointsToDelete.extend(jointsToRemove)
+
+        jointsToRemove = ["rider_heel_r"]
+        jointToTransferTo = "rider_calf_r"
+        transfer.LOD_transferWeights(meshes, jointsToRemove, jointToTransferTo)
+        jointsToDelete.extend(jointsToRemove)
+
+        jointsToRemove = ["rider_ball_l"]
+        jointToTransferTo = "rider_foot_l"
+        transfer.LOD_transferWeights(meshes, jointsToRemove, jointToTransferTo)
+        jointsToDelete.extend(jointsToRemove)
+
+        jointsToRemove = ["rider_ball_r"]
+        jointToTransferTo = "rider_foot_r"
+        transfer.LOD_transferWeights(meshes, jointsToRemove, jointToTransferTo)
+        jointsToDelete.extend(jointsToRemove)
+        
+        # Delete the joints we no longer want.
+        for one in jointsToDelete:
+            if cmds.objExists(one):
+                cmds.delete(one)
+                
+        # Export the fbx file.
+        export.exportLOD("HyperBreach_LOD4.fbx", "LOD4")
+        
+        # Re-open the file without saving.
+        fullPath = cmds.file(q = True, sceneName = True)
+        cmds.file(fullPath, open=True, f=True)
+
+        print "LOD4 Exported Successfully!!!!!!!!!!!!!!!"
+
+HyperBreachLOD4()
+
+
diff --git a/Engine/Extras/Maya_AnimationRiggingTools/ARTv1/General/ART/Projects/Orion/Scripts/HyperBreach_Pilot_Post_Script.py b/Engine/Extras/Maya_AnimationRiggingTools/ARTv1/General/ART/Projects/Orion/Scripts/HyperBreach_Pilot_Post_Script.py
new file mode 100644
--- /dev/null
+++ b/Engine/Extras/Maya_AnimationRiggingTools/ARTv1/General/ART/Projects/Orion/Scripts/HyperBreach_Pilot_Post_Script.py
@@ -0,0 +1,174 @@
+import maya.cmds as cmds
+import maya.mel as mel
+import math
+
+# ------------------------------------------------------------------
+def scaleControls():
+    
+    # eyes
+    for side in ["_l", "_r"]:
+        cmds.select(["eyeball"+side+"_anim.cv[0:32]"])
+        cmds.scale(0.1, 0.1, 0.05, r=True, ocp = False)
+        if side == "_l":
+            cmds.move(0, -3.757386, 0, r=True, os=True)
+        else:
+            cmds.move(0, -3.757386, 0, r=True, os=True)
+
+        cmds.select(["eyelid_bott"+side+"_anim.cv[0:32]"])
+        cmds.scale(0.1, 0.1, 0.01, r=True, ocp = False)
+        if side == "_l":
+            cmds.move(0, -3.757386, 0.5, r=True, os=True)
+        else:
+            cmds.move(0, 3.757386, 0.5, r=True, os=True)
+
+        cmds.select(["eyelid_top"+side+"_anim.cv[0:32]"])
+        cmds.scale(0.1, 0.1, 0.01, r=True, ocp = False)
+        if side == "_l":
+            cmds.move(0, -3.757386, 0.5, r=True, os=True)
+        else:
+            cmds.move(0, 3.757386, 0.5, r=True, os=True)
+    
+        cmds.select(["heel"+side+"_anim.cv[0:32]"])
+        cmds.scale(0.34, 0.34, 0.34, r = True, ocp = True)
+
+    # jaw
+    cmds.select(["jaw_anim.cv[0:32]"])
+    cmds.scale(0.294922, 0.294922, 0.294922, r = True, ocp = True)
+    cmds.move(0, -12, -10, r = True)
+    
+    # ears
+    cmds.select(["fk_earMain_l_01_anim.cv[0:7]", "fk_earMain_r_01_anim.cv[0:7]"])
+    cmds.scale(1.522128, 1.522128, 1.522128, r = True, ocp = False)
+
+    cmds.select(["fk_earMain_l_03_anim.cv[0:7]", "fk_earMain_r_03_anim.cv[0:7]"])
+    cmds.scale(0.696643, 0.696643, 0.696643, r = True, ocp = False)
+    
+    # neck and head
+    cmds.select(["neck_01_fk_anim.cv[0:7]"])
+    cmds.scale(2.18565, 2.18565, 2.18565, r = True, ocp = False)
+
+    cmds.select(["neck_02_fk_anim.cv[0:7]"])
+    cmds.scale(2.18565, 2.18565, 2.18565, r = True, ocp = False)
+
+    cmds.select(["head_fk_anim.cv[0:7]"])
+    cmds.scale(2.483665, 2.483665, 2.483665, r = True, ocp = True)
+    
+    # foot
+    cmds.select(["ik_foot_anim_l.cv[0:17]"])
+    cmds.scale(1.4, 1.4, 1.4, r = True, ocp = False)
+    cmds.move(0.8, 0, -2, r = True)
+    
+    cmds.select(["ik_foot_anim_r.cv[0:17]"])
+    cmds.scale(1.4, 1.4, 1.4, r = True, ocp = False)
+    cmds.move(-0.8, 0, -2, r = True)
+
+    cmds.select(None)
+
+scaleControls()
+
+# ------------------------------------------------------------------
+def invertEyelids():
+    sides = ["_l", "_r"]
+    for side in sides:
+        cmds.delete("driver_eyelid_bott"+side+"_parentConstraint1")
+        cmds.rotate(0, -180, 0, "eyelid_bott"+side+"_anim_grp", r=True, os=True)
+        cmds.parentConstraint("eyelid_bott"+side+"_anim", "driver_eyelid_bott"+side, mo=True)
+invertEyelids()
+
+def invertRightEyeball():
+    cmds.delete("driver_eyeball_r_parentConstraint1")
+    cmds.rotate(0, -180, -180, "eyeball_r_anim_grp", r=True, os=True)
+    cmds.parentConstraint("eyeball_r_anim", "driver_eyeball_r", mo=True)
+invertRightEyeball()
+
+# ------------------------------------------------------------------
+def dogLegRig():
+    mel.eval("ikSpringSolver")
+    for side in ["_l", "_r"]:
+        ###########################
+        # DELETE old IK handle
+        cmds.delete("foot_ikHandle"+side)
+        
+        # Disconnect the attrs driving the foot joint so that we can re-drive it with a parent constraint.
+        attrToDisconnect = "foot"+side+".translate"
+        connection = cmds.connectionInfo(attrToDisconnect, sfd=True)
+        cmds.disconnectAttr(connection, attrToDisconnect)
+        cmds.delete("foot"+side+"_orientConstraint1")
+        
+        # PARENT driver_foot_l to driver_heel_l
+        cmds.parent("driver_foot"+side, "driver_heel"+side)
+        
+        # Parent constrain the deform foot to follow the driver foot
+        cmds.parentConstraint("driver_foot"+side, "foot"+side, mo=True)
+    
+        # DUPLICATE driver_heel_l and rename to ik_leg_heel_l
+        ik_leg_heel = cmds.duplicate("driver_heel"+side, n="ik_leg_heel"+side, po=True)[0]
+        # PARENT ik_leg_heel_l to ik_leg_calf_l
+        cmds.parent(ik_leg_heel, "ik_leg_calf"+side)
+        # PARENT ik_leg_foot_l to ik_leg_heel_l(dupe)
+        cmds.parent("ik_leg_foot"+side, ik_leg_heel)
+        cmds.parent("ik_leg_thigh"+side, w=True)
+        cmds.makeIdentity("ik_leg_thigh"+side, a=True, r=True)
+        cmds.parentConstraint("leg_joints_grp"+side, "ik_leg_thigh"+side, mo=True)
+        cmds.parent("ik_leg_thigh"+side, "rig_grp")
+        
+        # DUPLICATE driver_heel_l and rename to fk_leg_heel_l
+        fk_leg_heel = cmds.duplicate("driver_heel"+side, n="fk_leg_heel"+side, po=True)[0]
+        # PARENT fk_leg_heel_l to fk_leg_calf_l
+        cmds.parent(fk_leg_heel, "fk_leg_calf"+side)
+        # PARENT fk_leg_foot_l to fk_leg_heel_l(dupe)
+        cmds.parent("fk_leg_foot"+side, fk_leg_heel)
+        # ORIENT CONSTRAIN fk_leg_foot_l to heel_l_anim
+        cmds.orientConstraint("heel"+side+"_anim", "fk_leg_heel"+side, mo=True)
+        cmds.parentConstraint("heel"+side+"_anim", "fk_foot"+side+"_anim_grp", mo=True)
+
+        #DUPLICATE driver_heel_l and rename to result_leg_heel_l
+        result_leg_heel = cmds.duplicate("driver_heel"+side, n="result_leg_heel"+side, po=True)[0]
+        # PARENT result_leg_heel_l to result_leg_calf_l
+        cmds.parent(result_leg_heel, "result_leg_calf"+side)
+        # PARENT result_leg_foot_l to result_leg_heel_l(dupe)
+        cmds.parent("result_leg_foot"+side, result_leg_heel)
+        
+        #LOAD ikSpringSolver;
+        # CREATE spring IK from ik_leg_thigh_l to ik_leg_foot_l
+        foot_ikhandle = cmds.ikHandle(n="foot_ikHandle"+side, sol="ikSpringSolver", sj="ik_leg_thigh"+side, ee="ik_leg_foot"+side)[0]
+        cmds.setAttr(foot_ikhandle+".v", 0)
+        
+        #PARENT ik handle to ik_foot_ball_pivot_l
+        cmds.parent(foot_ikhandle, "ik_foot_ball_pivot"+side)
+
+        # CREATE PV consraint on ik handle to noflip_pv_loc_l
+        PVConst = cmds.poleVectorConstraint("noflip_pv_loc"+side, foot_ikhandle)[0]
+        if side == "_l":
+            cmds.setAttr(PVConst+".offsetX", -200)
+        else:
+            cmds.setAttr(PVConst+".offsetX", 200)
+
+        # CONNECT the knee twist attr on ik_foot_anim_l to the twist attr on new ik handle
+        cmds.connectAttr("ik_foot_anim"+side+".knee_twist", foot_ikhandle+".twist", f=True)
+        # DELETE driver_heel_l_parentConstraint1
+        cmds.delete("driver_heel"+side+"_parentConstraint1")
+        cmds.parentConstraint(result_leg_heel, "driver_heel"+side, mo=True)
+        
+        # CREATE a parent constraint using ik_leg_heel_l, fk_leg_heel_l and drive result_leg_heel_l
+        resultConstraint = cmds.parentConstraint(fk_leg_heel, ik_leg_heel, result_leg_heel, mo=True)[0]
+        # CREATE a reverse node and hook up a IK/FK switch between the Rig_Settings.lLegMode and the 2 weights on this new constraint
+        shortside = side.replace("_", "")
+        reverseNode = cmds.shadingNode("reverse", asUtility=True, name=result_leg_heel+"_reverse")
+        cmds.connectAttr("Rig_Settings."+shortside+"LegMode", resultConstraint+"."+ik_leg_heel+"W1", f=True)
+        cmds.connectAttr("Rig_Settings."+shortside+"LegMode", reverseNode+".inputX", f=True)
+        cmds.connectAttr(reverseNode+".outputX", resultConstraint+"."+fk_leg_heel+"W0", f=True)
+
+        # Create a KneeFlex Attribute
+        cmds.addAttr("ik_foot_anim"+side, ln="KneeFlex", at="double", min=0, max=1, dv=0.5, k=True)
+        # Connect KneeFlex to :
+        reverseNodeBias = cmds.shadingNode("reverse", asUtility=True, name=result_leg_heel+"_reverse")
+        cmds.connectAttr("ik_foot_anim"+side+".KneeFlex", foot_ikhandle+".springAngleBias[0].springAngleBias_FloatValue", f=True)
+        cmds.connectAttr("ik_foot_anim"+side+".KneeFlex", reverseNodeBias+".inputX", f=True)
+        cmds.connectAttr(reverseNodeBias+".outputX", foot_ikhandle+".springAngleBias[1].springAngleBias_FloatValue", f=True)
+        
+        visReverseNode = cmds.shadingNode("reverse", asUtility=True, name="heel"+side+"_anim_reverse")
+        cmds.connectAttr("Rig_Settings."+shortside+"LegMode", visReverseNode+".inputX", f=True)
+        cmds.connectAttr(visReverseNode+".outputX", "heel"+side+"_parent_grp.v", f=True)
+
+dogLegRig()
\ No newline at end of file
diff --git a/Engine/Extras/Maya_AnimationRiggingTools/ARTv1/General/ART/Projects/Orion/Scripts/HyperBreach_Post_Script.py b/Engine/Extras/Maya_AnimationRiggingTools/ARTv1/General/ART/Projects/Orion/Scripts/HyperBreach_Post_Script.py
new file mode 100644
--- /dev/null
+++ b/Engine/Extras/Maya_AnimationRiggingTools/ARTv1/General/ART/Projects/Orion/Scripts/HyperBreach_Post_Script.py
@@ -0,0 +1,819 @@
+import maya.mel as mel
+import maya.cmds as cmds
+import Tools.ART_CreateHydraulicRigs as hyd
+reload(hyd)
+import Tools.ART_addSpaceSwitching as space
+reload(space)
+
+# ------------------------------------------------------------------
+# add ik bones
+def addIKBones():
+    try:
+		# create the joints
+		cmds.select(clear = True)
+		ikFootRoot = cmds.joint(name = "ik_foot_root")
+		cmds.select(clear = True)
+		
+		cmds.select(clear = True)
+		ikFootLeft = cmds.joint(name = "ik_foot_l")
+		cmds.select(clear = True)
+		
+		cmds.select(clear = True)
+		ikFootRight = cmds.joint(name = "ik_foot_r")
+		cmds.select(clear = True)
+		
+		cmds.select(clear = True)
+		ikHandRoot = cmds.joint(name = "ik_hand_root")
+		cmds.select(clear = True)
+		
+		cmds.select(clear = True)
+		ikHandGun = cmds.joint(name = "ik_hand_gun")
+		cmds.select(clear = True)
+		
+		cmds.select(clear = True)
+		ikHandLeft = cmds.joint(name = "ik_hand_l")
+		cmds.select(clear = True)
+		
+		cmds.select(clear = True)
+		ikHandRight = cmds.joint(name = "ik_hand_r")
+		cmds.select(clear = True)
+		
+		# create hierarhcy
+		cmds.parent(ikFootRoot, "root")
+		cmds.parent(ikHandRoot, "root")
+		cmds.parent(ikFootLeft, ikFootRoot)
+		cmds.parent(ikFootRight, ikFootRoot)
+		cmds.parent(ikHandGun, ikHandRoot)
+		cmds.parent(ikHandLeft, ikHandGun)
+		cmds.parent(ikHandRight, ikHandGun)
+		
+		# constrain the joints
+		leftHandConstraint = cmds.parentConstraint("hand_l", ikHandLeft)[0]
+		rightHandGunConstraint = cmds.parentConstraint("hand_r", ikHandGun)[0]
+		rightHandConstraint = cmds.parentConstraint("hand_r", ikHandRight)[0]
+		leftFootConstraint = cmds.parentConstraint("foot_l", ikFootLeft)[0]
+		rightFootConstraint = cmds.parentConstraint("foot_r", ikFootRight)[0]
+	
+    except:
+		cmds.warning("Something went wrong")
+		
+# ------------------------------------------------------------------
+# Modify the cannon on his left shoulder for HyperBreach
+def modifyCannon():
+    
+    # Lock attributes
+    for control in ["cannon_arm_1_anim", "cannon_arm_2_anim"]:
+        for attr in [".tx", ".ty", ".tz", ".rx", ".ry"]:
+            cmds.setAttr(control + attr, lock = True, keyable = False, channelBox = False)
+        
+    # modify nurbs curves
+    cmds.select("cannon_anim.cv[0:32]")
+    cmds.scale(0, 1, 1, r = True)
+    cmds.move(0, 0, 5.4, r = True)
+
+    cmds.select("mine_anim.cv[0:32]")
+    cmds.scale(0, 1, 1, r = True)
+
+    cmds.select("cannon_arm_2_anim.cv[0:32]")
+    cmds.scale(0.8, 0.8, 0, r = True)
+
+    cmds.select("cannon_arm_1_anim.cv[0:32]")
+    cmds.scale(1.1, 1.1, 0, r = True)
+
+    cmds.select("cannon_root_anim.cv[0:32]")
+    cmds.scale(1.2, 1.2, 1.2, r = True)
+
+    cmds.addAttr("cannon_root_anim", ln="PanelMorph", at="double", min=0, max=1, dv=0, k=True)
+    cmds.connectAttr("cannon_root_anim.PanelMorph", "bodypanel_l_blendShape.bodypanel_l_geo_open", f=True)
+
+    # set color
+    for control in ["cannon_root_anim", "cannon_arm_1_anim", "cannon_arm_2_anim", "cannon_anim", "mine_anim"]:
+        cmds.setAttr(control + ".overrideColor", 26)
+
+    cmds.select(None)
+modifyCannon()
+
+# ------------------------------------------------------------------
+# Modify the Long Range Missile (LRM) for HyperBreach.  This is the missing on the left arm.
+def modifyLRM():
+    # Scale controls
+    cmds.select("rocket_anim.cv[0:32]")
+    cmds.scale(0.7, 0.7, 0.7, r = True)
+    cmds.move(0, 0, 8, r = True, os = True, wd = True)
+   
+    cmds.select("bay_door_a_anim.cv[0:32]")
+    cmds.scale(0.24, 0.45, 0.7, r = True)
+    cmds.move(0, 0, 0, r = True, os = True, wd = True)
+
+    cmds.select("bay_door_b_anim.cv[0:32]")
+    cmds.scale(0.24, 0.45, 0.7, r = True)
+    cmds.move(0, 0, 0, r = True, os = True, wd = True)
+
+    cmds.select("rocket_arm_root_anim.cv[0:32]")
+    cmds.scale(0.05, 1, 1, r = True)
+    cmds.move(0, 0, 0, r = True, os = True, wd = True)
+
+    cmds.select("rocket_arm_anim.cv[0:32]")
+    cmds.scale(0.68, 0.05, 0.68, r = True)
+    cmds.move(0, 0, 0, r = True, os = True, wd = True)
+
+    cmds.select("rocket_arm_b_anim.cv[0:32]")
+    cmds.scale(0.7, 0.05, 0.7, r = True)
+    cmds.move(0, 0, 0, r = True, os = True, wd = True)
+
+    # adding sdks
+    cmds.addAttr("rocket_anim", ln = "slideY", dv = 0, min = 0, max = 10, keyable = True)
+    cmds.addAttr("rocket_anim", ln = "slideX", dv = 0, min = 0, max = 10, keyable = True)
+    cmds.addAttr("rocket_anim", ln = "rocketGrab", dv = 0, min = 0, max = 10, keyable = True)
+    cmds.addAttr("rocket_anim", ln = "rocketClasp", dv = 0, min = 0, max = 10, keyable = True)
+    
+    cmds.setDrivenKeyframe(["rocket_arm_root_anim.ty"], cd = "rocket_anim.slideY", itt = "spline", ott = "spline")
+    cmds.setAttr("rocket_anim.slideY", 10)
+    cmds.setAttr("rocket_arm_root_anim.ty", 50)
+    cmds.setDrivenKeyframe(["rocket_arm_root_anim.ty"], cd = "rocket_anim.slideY", itt = "spline", ott = "spline")
+    
+    cmds.setDrivenKeyframe(["rocket_arm_root_anim.tx"], cd = "rocket_anim.slideX", itt = "spline", ott = "spline")
+    cmds.setAttr("rocket_anim.slideX", 10)
+    cmds.setAttr("rocket_arm_root_anim.tx", -10.5)
+    cmds.setDrivenKeyframe(["rocket_arm_root_anim.tx"], cd = "rocket_anim.slideX", itt = "spline", ott = "spline")
+    
+    cmds.setDrivenKeyframe(["rocket_arm_anim.ry"], cd = "rocket_anim.rocketGrab", itt = "spline", ott = "spline")
+    cmds.setAttr("rocket_anim.rocketGrab", 10)
+    cmds.setAttr("rocket_arm_anim.ry", -90)
+    cmds.setDrivenKeyframe(["rocket_arm_anim.ry"], cd = "rocket_anim.rocketGrab", itt = "spline", ott = "spline")
+    
+    cmds.setDrivenKeyframe(["rocket_arm_b_anim.ry"], cd = "rocket_anim.rocketClasp", itt = "spline", ott = "spline")
+    cmds.setAttr("rocket_anim.rocketClasp", 10)
+    cmds.setAttr("rocket_arm_b_anim.ry", -90)
+    cmds.setDrivenKeyframe(["rocket_arm_b_anim.ry"], cd = "rocket_anim.rocketClasp", itt = "spline", ott = "spline")
+    
+    
+    cmds.setAttr("rocket_anim.rocketGrab", 0)
+    cmds.setAttr("rocket_anim.rocketClasp", 0)
+    cmds.setAttr("rocket_anim.slideY", 0)
+    cmds.setAttr("rocket_anim.slideX", 0)
+    
+    # bay doors
+    cmds.addAttr("rocket_anim", ln = "bayDoors", dv = 0, min = 0, max = 10, keyable = True)
+    cmds.setDrivenKeyframe(["bay_door_a_anim.ry", "bay_door_b_anim.ry"], cd = "rocket_anim.bayDoors", itt = "spline", ott = "spline")
+    cmds.setAttr("rocket_anim.bayDoors", 10)
+    cmds.setAttr("bay_door_a_anim.ry", -90)
+    cmds.setAttr("bay_door_b_anim.ry", 90)
+    cmds.setDrivenKeyframe(["bay_door_a_anim.ry", "bay_door_b_anim.ry"], cd = "rocket_anim.bayDoors", itt = "spline", ott = "spline")
+    cmds.setAttr("rocket_anim.bayDoors", 0)
+    
+    # scale
+    cmds.connectAttr("rocket_anim.scale", "rocket_arm_anim.scale")
+    cmds.connectAttr("rocket_anim.scale", "rocket_arm_root_anim.scale")
+
+    # morphs
+    cmds.addAttr("rocket_anim", ln = "extend_front", dv = 0, min = 0, max = 1, keyable = True)
+    cmds.addAttr("rocket_anim", ln = "extend_back", dv = 0, min = 0, max = 1, keyable = True)
+    cmds.addAttr("rocket_anim", ln = "extend_fins", dv = 0, min = 0, max = 1, keyable = True)
+    
+    rocket_morphs_reverse = cmds.shadingNode("reverse", n="rocket_morphs_reverse", asUtility=True)
+    cmds.connectAttr("rocket_anim.extend_front", rocket_morphs_reverse+".inputX")
+    cmds.connectAttr("rocket_anim.extend_back", rocket_morphs_reverse+".inputY")
+    cmds.connectAttr("rocket_anim.extend_fins", rocket_morphs_reverse+".inputZ")
+
+    cmds.connectAttr(rocket_morphs_reverse+".outputX", "HB_Rocket_Shapes.missle_geo_extend_front")
+    cmds.connectAttr(rocket_morphs_reverse+".outputY", "HB_Rocket_Shapes.missle_geo_extend_back")
+    cmds.connectAttr(rocket_morphs_reverse+".outputZ", "HB_Rocket_Shapes.missle_geo_extend_fins")
+
+    cmds.select(None)
+modifyLRM()
+
+# ------------------------------------------------------------------
+# Modify the rocket controls for the Ultimate for HyperBreach.  This is the rockets on his back that transform out.
+def modifyUltimate():
+    # back rocket launchers
+    cmds.select("rocket_root_l_anim.cv[0:32]")
+    cmds.scale(1, 0, 1, r = True)
+    cmds.move(0, 14, 0, r = True, os = True, wd = True)
+
+    cmds.select("rocket_l_anim.cv[0:32]")
+    cmds.scale(0.5, 0.1, 0.5, r = True)
+    cmds.move(0, 10, 0, r = True, os = True, wd = True)
+
+    cmds.addAttr("rocket_root_l_anim", ln = "rocketSlide", dv = 0, min =0, max = 10, keyable = True)
+    cmds.setDrivenKeyframe(["rocket_l_anim_grp.tz", "rocket_l_anim_grp.ty"], cd = "rocket_root_l_anim.rocketSlide", itt = "spline", ott = "spline")
+    cmds.setAttr("rocket_root_l_anim.rocketSlide", 10)
+    cmds.setAttr("rocket_l_anim_grp.ty", -7.5)
+    cmds.setAttr("rocket_l_anim_grp.tz", 17)
+    cmds.setDrivenKeyframe(["rocket_l_anim_grp.tz", "rocket_l_anim_grp.ty"], cd = "rocket_root_l_anim.rocketSlide", itt = "spline", ott = "spline")
+    cmds.setAttr("rocket_root_l_anim.rocketSlide", 0)
+
+    cmds.select("rocket_root_r_anim.cv[0:32]")
+    cmds.scale(1, 0, 1, r = True)
+    cmds.move(0, -14, 0, r = True, os = True, wd = True)
+
+    cmds.select("rocket_r_anim.cv[0:32]")
+    cmds.scale(0.5, 0.1, 0.5, r = True)
+    cmds.move(0, -10, 0, r = True, os = True, wd = True)
+        
+    cmds.addAttr("rocket_root_r_anim", ln = "rocketSlide", dv = 0, min =0, max = 10, keyable = True)
+    cmds.setDrivenKeyframe(["rocket_r_anim_grp.tz", "rocket_r_anim_grp.ty"], cd = "rocket_root_r_anim.rocketSlide", itt = "spline", ott = "spline")
+    cmds.setAttr("rocket_root_r_anim.rocketSlide", 10)
+    cmds.setAttr("rocket_r_anim_grp.ty", 7.5)
+    cmds.setAttr("rocket_r_anim_grp.tz", -17)
+    cmds.setDrivenKeyframe(["rocket_r_anim_grp.tz", "rocket_r_anim_grp.ty"], cd = "rocket_root_r_anim.rocketSlide", itt = "spline", ott = "spline")
+    cmds.setAttr("rocket_root_r_anim.rocketSlide", 0)
+
+    # back rockets
+    cmds.select("rocket_transform_anim.cv[0:32]")
+    cmds.scale(2.2, 0.2, 1, r = True)
+    cmds.move(0, 48, 0, r = True, os = True, wd = True)
+    cmds.connectAttr("rocket_root_l_anim.scale", "rocket_l_anim.scale")
+    cmds.connectAttr("rocket_root_r_anim.scale", "rocket_r_anim.scale")
+
+    cmds.select(None)
+modifyUltimate()
+
+# ------------------------------------------------------------------
+# Modify the cockpit rig control arms for HyperBreach
+def rigControlArms(side):
+    # remove the existing rig and replace with an IK rig
+    cmds.select("driver_controlArm_"+side+"_01", hi = True)
+    cmds.delete(constraints = True)
+    
+    # create Spring IK for control arms
+    mel.eval("ikSpringSolver;")
+    ArmIK = cmds.ikHandle(sj = "driver_controlArm_"+side+"_01", ee = "driver_controlArm_"+side+"_03", sol = "ikRPsolver")[0]
+    
+    cmds.setAttr(ArmIK + ".v", 0)
+    
+    # setup new constraints
+    cmds.orientConstraint("driver_controlArm_"+side+"_01", "controlArm_"+side+"_01")
+    cmds.orientConstraint("driver_controlArm_"+side+"_02", "controlArm_"+side+"_02")
+    cmds.orientConstraint("driver_controlArm_"+side+"_03", "controlArm_"+side+"_03")
+    #cmds.orientConstraint("driver_controlArm_"+side+"_04", "controlArm_"+side+"_04")
+    
+    # Create control arm rig controls
+    Ctrl = cmds.circle(name = side+"_ctrl_arm_anim")[0]
+    cmds.setAttr(Ctrl + ".scale", 10, 10, 10, type = "double3")
+    cmds.setAttr(Ctrl + ".rx", 90)
+    constraint = cmds.pointConstraint(ArmIK, Ctrl)[0]
+    cmds.delete(constraint)
+    
+    CtrlGrp = cmds.group(empty = True, name = side+"_ctrl_arm_anim_grp")
+    constraint = cmds.parentConstraint("controlArm_"+side+"_03", CtrlGrp)
+    cmds.delete(constraint)
+    cmds.parent(Ctrl, CtrlGrp)
+    cmds.makeIdentity(Ctrl, apply = True, t = 1, r = 1, s = 1)
+    cmds.parent(ArmIK, Ctrl)
+    cmds.orientConstraint(Ctrl, "driver_controlArm_"+side+"_03", mo = True)
+
+    ikCtrlSpaceSwitchFollow = cmds.duplicate(CtrlGrp, po = True, n = side+"_ctrl_arm_anim_space_switcher_follow")[0]
+    ikCtrlSpaceSwitch = cmds.duplicate(CtrlGrp, po = True, n = side+"_ctrl_arm_anim_space_switcher")[0]
+
+    cmds.parent(ikCtrlSpaceSwitchFollow, "controlArm_"+side+"_master_ctrl_grp")
+    cmds.parent(ikCtrlSpaceSwitch, ikCtrlSpaceSwitchFollow)
+    cmds.parent(CtrlGrp, ikCtrlSpaceSwitch)
+
+    spaceSwitchers = [side+"_ctrl_arm_anim_space_switcher_follow"]
+    for node in spaceSwitchers:
+        #create a 'world' locator to constrain to
+        worldLoc = cmds.spaceLocator(name = node + "_world_pos")[0]
+        cmds.setAttr(worldLoc + ".v", 0)
+
+        #position world loc to be in same place as node
+        constraint = cmds.parentConstraint(node, worldLoc)[0]
+        cmds.delete(constraint)
+        constraint = cmds.parentConstraint(worldLoc, node)[0]
+
+        #add the attr to the space switcher node for that space
+        spaceSwitchNode = node.rpartition("_follow")[0]
+        cmds.select(spaceSwitchNode)
+        cmds.addAttr(ln = "space_world", minValue = 0, maxValue = 1, dv = 0, keyable = True)
+
+        #connect that attr to the constraint
+        cmds.connectAttr(spaceSwitchNode + ".space_world", constraint + "." + worldLoc + "W0")
+
+        #parent worldLoc under the offset_anim
+        cmds.parent(worldLoc, "offset_anim")
+
+    # hide original controls
+    cmds.setAttr("controlArm_"+side+"_dyn_ctrl_grp.v", 0, lock = True)
+    cmds.setAttr("fk_controlArm_"+side+"_01_lag_grp.v", 0, lock = True)
+
+    cmds.select(None)
+
+rigControlArms("l")
+rigControlArms("r")
+
+# ------------------------------------------------------------------
+# Modify the cockpit rig seat for HyperBreach
+def rigCockpit():
+    cmds.select("cockpit_pedal_l_anim.cv[0:32]")
+    cmds.scale(0.5, 0.5, 0.5, r = True)
+    cmds.move(6.7, 0, 0, r = True, os = True, wd = True)
+
+    cmds.select("cockpit_pedal_r_anim.cv[0:32]")
+    cmds.scale(0.5, 0.5, 0.5, r = True)
+    cmds.move(-6.7, 0, 0, r = True, os = True, wd = True)
+
+    cmds.select("cockpit_foot_pedal_l_anim.cv[0:32]")
+    cmds.scale(0.5, 0.5, 0.5, r = True)
+
+    cmds.select("cockpit_foot_pedal_r_anim.cv[0:32]")
+    cmds.scale(0.5, 0.5, 0.5, r = True)
+
+    cmds.select("cockpit_seat_anim.cv[0:32]")
+    cmds.scale(1.4, 0.4, 0.4, r = True)
+    cmds.move(0, 0, 0, r = True, os = True, wd = True)
+
+    # Lock attributes
+    for control in ["cockpit_seat_base_anim", "cockpit_pedal_l_anim", "cockpit_pedal_r_anim", "cockpit_foot_pedal_l_anim", "cockpit_foot_pedal_r_anim", "cockpit_seat_anim"]:
+        for attr in [".tx", ".ty", ".tz", ".rz", ".ry"]:
+            cmds.setAttr(control + attr, lock = True, keyable = False, channelBox = False)
+
+    cmds.select(None)
+rigCockpit()
+
+# ------------------------------------------------------------------
+# Modify the torso for HyperBreach
+def modifyTorso():
+    # Set Spine to FK mode
+    cmds.setAttr("Rig_Settings.spine_ik", 0)
+    #cmds.setAttr("Rig_Settings.spine_ik", lock = True)
+    
+    cmds.setAttr("Rig_Settings.spine_fk", 1)
+    #cmds.setAttr("Rig_Settings.spine_fk", lock = True)
+    
+    # Lock attributes on spine 3
+    cmds.setAttr("spine_03_anim.length", lock = True, keyable = False, channelBox = False)
+    cmds.setAttr("spine_03_anim.tx", lock = True, keyable = False, channelBox = False)
+    cmds.setAttr("spine_03_anim.ty", lock = True, keyable = False, channelBox = False)
+    cmds.setAttr("spine_03_anim.tz", lock = True, keyable = False, channelBox = False)
+    #cmds.setAttr("spine_03_anim.ry", lock = True, keyable = False, channelBox = False)
+    #cmds.setAttr("spine_03_anim.rx", lock = True, keyable = False, channelBox = False)
+    
+    # Hide Neck/Head Controls
+    cmds.setAttr("neck_01_fk_anim_grp.v", 0)
+    
+    # Turbine
+    cmds.select("turbine_root_anim.cv[0:32]")
+    cmds.scale(0.6, 0.6, 0.6, r = True)
+    cmds.move(0, 47, 42.5, r = True, os=True, wd=True)
+    cmds.setAttr("turbine_root_anim.overrideColor", 20)
+    
+    cmds.select("turbine_hinge_anim.cv[0:32]")
+    cmds.scale(0.5, 0.5, 0.5, r = True)
+    cmds.move(0, 43, 13, r = True, os=True, wd=True)
+    
+    cmds.select("turbine_anim.cv[0:32]")
+    cmds.scale(0.4, 0.4, 0.4, r = True)
+    cmds.move(0, 44.5, 0, r = True, os=True, wd=True)
+    
+    # lock turbine hinge anim attributes
+    for attr in [".tx", ".ty", ".tz", ".ry", ".rz", ".sx", ".sy", ".sz"]:
+	cmds.setAttr("turbine_hinge_anim" + attr, lock = True, keyable = False)
+	
+    # lower fin
+    cmds.select("lower_fin_anim.cv[0:32]")
+    cmds.scale(0.6, 0.6, 0.6, r = True)
+    cmds.move(0, 42, 5, r = True, os=True, wd=True)
+
+    cmds.select("lower_fin_end_anim.cv[0:32]")
+    cmds.scale(0.5, 0.5, 0.5, r = True)
+    cmds.move(0, 37, 8, r = True, os=True, wd=True)
+
+    cmds.orientConstraint("spine_01_anim", "lower_fin_anim_space_switcher", mo = True)
+
+    cmds.addAttr("lower_fin_end_anim", ln = "spinFromHips", dv = 0, keyable = True)
+    cmds.connectAttr("lower_fin_end_anim.spinFromHips", "lower_fin_anim.rx")
+    cmds.setAttr("lower_fin_anim.rx", lock = True, keyable = False, channelBox = False)
+
+    # upper fin
+    cmds.select("upper_fin_anim.cv[0:32]")
+    cmds.scale(0.5, 0.5, 0.5, r = True)
+    cmds.move(0, 50, 3, r = True, os=True, wd=True)
+    cmds.setAttr("upper_fin_anim.overrideColor", 25)
+    for attr in [".tx", ".ty", ".tz", ".ry", ".rz", ".sx", ".sy", ".sz"]:
+	cmds.setAttr("upper_fin_anim" + attr, lock = True, keyable = False, channelBox = False)
+	
+    cmds.select(None)
+modifyTorso()
+
+# ------------------------------------------------------------------
+# Modify the arms for HyperBreach
+def modifyArms():
+    # Set Clavicles to FK
+    cmds.setAttr("Rig_Settings.lClavMode", 0)
+    cmds.setAttr("Rig_Settings.rClavMode", 0)    
+
+    # Set arm mode to FK
+    cmds.setAttr("Rig_Settings.lArmMode", 0)
+    #cmds.setAttr("Rig_Settings.lArmMode", lock = True)
+    cmds.setAttr("Rig_Settings.rArmMode", 0)
+    #cmds.setAttr("Rig_Settings.rArmMode", lock = True)
+    
+    # arm jet controls
+    cmds.select("arm_jet_l_anim.cv[0:32]")
+    cmds.scale(0.37, 1, 0.37, r = True)
+    cmds.move(0, 20, 0, r = True, os = True, wd = True)    
+    
+    cmds.select("arm_jet_r_anim.cv[0:32]")
+    cmds.scale(0.37, 1, 0.37, r = True)
+    cmds.move(0, -20, 0, r = True, os = True, wd = True)      
+
+    # Set up the extra shoulder part to stay torso aligned
+    cmds.select("shoulder_mid_l_anim.cv[0:32]")
+    cmds.scale(0.5, 0.5, 0.5, r = True)
+    cmds.move(10, 0, 18, r = True, os=True, wd=True)
+    cmds.delete("shoulder_mid_l_parent_grp_parentConstraint1")
+    cmds.parentConstraint("clavicle_l", "shoulder_mid_l_parent_grp", mo=True, sr=["x", "y", "z"])
+    cmds.orientConstraint("driver_spine_03", "shoulder_mid_l_parent_grp", mo=True)
+    
+    cmds.select("shoulder_mid_r_anim.cv[0:32]")
+    cmds.scale(0.5, 0.5, 0.5, r = True)
+    cmds.move(-10, 0, -18, r = True, os=True, wd=True)
+    cmds.delete("shoulder_mid_r_parent_grp_parentConstraint1")
+    cmds.parentConstraint("clavicle_r", "shoulder_mid_r_parent_grp", mo=True, sr=["x", "y", "z"])
+    cmds.orientConstraint("driver_spine_03", "shoulder_mid_r_parent_grp", mo=True)
+
+    cmds.select("shoulder_pad_l_anim.cv[0:4]")
+    cmds.scale(0.75, 0.75, 0.75, r = True)
+    cmds.move(0, 0, 4.5, r = True, os=True, wd=True)
+    shoulder_pad_l_anim_drv = cmds.group("shoulder_pad_l_anim", name="shoulder_pad_l_anim_drv")
+    cmds.aimConstraint("driver_upperarm_l", shoulder_pad_l_anim_drv, aimVector = [1, 0, 0], upVector = [0, 0, 1], worldUpType="objectrotation", worldUpVector=[0, 1, 0], worldUpObject="driver_clavicle_l", mo=True, skip = ["x", "z"])
+    cmds.transformLimits("shoulder_pad_l_anim_drv", ry=[0, 1.92], ery=[0, 1])
+    
+    cmds.select("shoulder_pad_r_anim.cv[0:4]")
+    cmds.scale(0.75, 0.75, 0.75, r = True)
+    cmds.move(0, 0, -4.5, r = True, os=True, wd=True)
+    shoulder_pad_r_anim_drv = cmds.group("shoulder_pad_r_anim", name="shoulder_pad_r_anim_drv")
+    cmds.aimConstraint("driver_upperarm_r", shoulder_pad_r_anim_drv, aimVector = [-1, 0, 0], upVector = [0, 0, -1], worldUpType="objectrotation", worldUpVector=[0, 1, 0], worldUpObject="driver_clavicle_r", mo=True, skip = ["x", "z"])
+    cmds.transformLimits("shoulder_pad_r_anim_drv", ry=[0, 1.92], ery=[0, 1])
+
+    cmds.setAttr("shoulder_pad_l_anim.chainStartEnvelope", 0)
+    cmds.setAttr("shoulder_pad_r_anim.chainStartEnvelope", 0)
+
+    cmds.addAttr("arm_jet_l_anim", ln = "bay_doors_l", dv = 0, min = 0, max = 1, keyable = True)
+    cmds.addAttr("arm_jet_r_anim", ln = "bay_doors_r", dv = 0, min = 0, max = 1, keyable = True)
+
+    cmds.connectAttr("arm_jet_l_anim.bay_doors_l", "bicep_l_hatch_blendShape.bicep_l_hatch_geo_open")
+    cmds.connectAttr("arm_jet_r_anim.bay_doors_r", "bicep_r_hatch_blendShape.bicep_r_hatch_geo_open")
+
+    cmds.select(None)
+modifyArms()
+
+# ------------------------------------------------------------------
+# Modify the legs for HyperBreach
+def modifyLegs():
+    # Left Leg Hips
+    cmds.aimConstraint("leg_socket_l_anim", "hip_l_anim", aimVector = [1, 0, 0], upVector = [0, 0, 1], skip = "x")
+    cmds.aimConstraint("hip_l_anim_grp", "leg_socket_l_anim", aimVector = [1, 0, 0], upVector = [0, 0, 1], skip = "x")
+    cmds.aimConstraint("hip_l_anim_grp", "leg_socket_root_l_anim", aimVector = [-1, 0, 0], upVector = [0, 0, 0], skip = ["x", "y"], mo = True)
+    
+    # Left Leg IK Foot Control
+    cmds.select("ik_foot_anim_l.cv[0:17]")
+    cmds.scale(1.6, 1.6, 1.6, r = True)
+    cmds.select(clear = True)
+    cmds.select("ik_foot_anim_l.cv[0:17]")
+    cmds.move(0, 0, -8, r = True, os = True, wd = True)
+    
+    # Right Leg Hips
+    cmds.aimConstraint("leg_socket_r_anim", "hip_r_anim", aimVector = [-1, 0, 0], upVector = [0, 0, -1], skip = "x")
+    cmds.aimConstraint("hip_r_anim_grp", "leg_socket_r_anim", aimVector = [-1, 0, 0], upVector = [0, 0, -1], skip = "x")
+    cmds.aimConstraint("hip_r_anim_grp", "leg_socket_root_r_anim", aimVector = [1, 0, 0], upVector = [0, 0, 0], skip = ["x", "y"], mo = True)
+    
+    # Right Leg IK Foot Control
+    cmds.select("ik_foot_anim_r.cv[0:17]")
+    cmds.scale(1.6, 1.6, 1.6, r = True)
+    cmds.select(clear = True)
+    cmds.select("ik_foot_anim_r.cv[0:17]")
+    cmds.move(0, 0, -8, r = True, os = True, wd = True)
+    
+    # Leg Jet controls
+    cmds.select("leg_jet_l_anim.cv[0:32]")
+    cmds.scale(0.5, 0.5, 1, r = True)
+    cmds.move(0, 0, -20, r = True, os = True, wd = True)    
+
+    cmds.select("leg_jet_r_anim.cv[0:32]")
+    cmds.scale(0.5, 0.5, 1, r = True)
+    cmds.move(0, 0, 20, r = True, os = True, wd = True) 
+    
+    # morphs
+    cmds.addAttr("leg_jet_l_anim", ln = "bay_doors_l", dv = 0, min = 0, max = 1, keyable = True)
+    cmds.addAttr("leg_jet_r_anim", ln = "bay_doors_r", dv = 0, min = 0, max = 1, keyable = True)
+
+    cmds.addAttr("leg_jet_l_anim", ln = "foot_jet_l", dv = 0, min = 0, max = 1, keyable = True)
+    cmds.addAttr("leg_jet_r_anim", ln = "foot_jet_r", dv = 0, min = 0, max = 1, keyable = True)
+    
+    cmds.connectAttr("leg_jet_l_anim.bay_doors_l", "calfhatch_l_blendShape.calfhatch_l_geo_open")
+    cmds.connectAttr("leg_jet_r_anim.bay_doors_r", "calfhatch_r_blendShape.calfhatch_r_geo_open")
+
+    cmds.connectAttr("leg_jet_l_anim.foot_jet_l", "jetfoot_l_blendShape.jetfoot_l_geo_open")
+    cmds.connectAttr("leg_jet_r_anim.foot_jet_r", "jetfoot_r_blendShape.jetfoot_r_geo_open")
+
+    cmds.select(None)
+modifyLegs()
+
+# ------------------------------------------------------------------
+def dynamicsOff():
+    for one in ["shoulder_pad_l_anim", "shoulder_pad_r_anim", "controlArm_l_dyn_anim", "controlArm_r_dyn_anim"]:
+        cmds.setAttr(one+".chainStartEnvelope", 0)
+        cmds.setAttr(one+".chainStartEnvelope", l=True)
+dynamicsOff()
+
+# ------------------------------------------------------------------
+def createGameCam():
+    gameCam = cmds.camera(centerOfInterest=5, focalLength=18, lensSqueezeRatio=1, cameraScale=1, horizontalFilmAperture=1.41732, horizontalFilmOffset=0, verticalFilmAperture=0.94488, verticalFilmOffset=0, filmFit="Overscan", overscan=1.05, motionBlur=0, shutterAngle=144, nearClipPlane=1.0, farClipPlane=10000, orthographic=0, orthographicWidth=30, panZoomEnabled=0, horizontalPan=0, verticalPan=0, zoom=1, displayFilmGate=True)[0]
+    print gameCam
+    gameCam = cmds.rename(gameCam, "GameCam")
+    cmds.camera(gameCam, e=True, horizontalFilmAperture=1.6724376, focalLength=21.2399857011)
+    cmds.setAttr(gameCam+".displayGateMaskOpacity", 1)
+    cmds.setAttr(gameCam+".displayGateMaskColor", 0, 0, 0, type="double3")
+    cmds.setAttr(gameCam+".r", 90, 0, 180, type="double3")
+    
+    # Create the control
+    control = cmds.circle(name=gameCam+"_Ctrl", ch=True, o=True, r=15)[0]
+    cmds.makeIdentity(control, a=True, t=True, r=True, s=True)
+    # Create groups for the sdk and the constraints
+    const_grp = cmds.group(control, name=control+"_grp")
+    if cmds.objExists("rig_grp"):
+        cmds.parent(const_grp, "rig_grp")
+    # Parent Constrain the cam to the control
+    cmds.parentConstraint(control, gameCam, mo=True)
+    
+    cmds.connectAttr(control+".visibility", gameCam+".visibility", f=True)
+    cmds.setAttr(const_grp+".v", 0)
+    
+    cmds.setAttr(const_grp+".t", 0, 500, 200, type="double3")
+createGameCam()
+
+# ------------------------------------------------------------------
+def addCapsule(radius, halfHeight):
+    trueHalfHeight = (halfHeight*2)-(radius*2)
+    capsule = cmds.polyCylinder(name="collision_capsule", r=radius, h=trueHalfHeight, sx=18, sy=1, sz=5, ax=[0,0,1], rcp=True, cuv=3, ch=True)
+    cmds.setAttr(capsule[0]+".tz", halfHeight+2.2)
+    cmds.setAttr(capsule[0]+".rz", 90)
+
+    cmds.addAttr(capsule[0], ln="Radius", dv=radius, keyable=True)
+    cmds.addAttr(capsule[0], ln="HalfHeight", dv=halfHeight, keyable=True)
+
+    expr = capsule[1]+".radius = "+capsule[0]+".Radius;\n"+capsule[1]+".height = ("+capsule[0]+".HalfHeight*2) - ("+capsule[0]+".Radius*2);"
+    cmds.expression(s=expr, o=capsule[1], ae=1, uc="all")
+
+    cmds.setAttr(capsule[0]+".overrideEnabled", 1)
+    cmds.setAttr(capsule[0]+".overrideColor", 18)
+    cmds.setAttr(capsule[0]+".overrideShading", 0)
+
+    cmds.addAttr(capsule[0], ln="Shaded", dv=0, min=0, max=1, keyable=True)
+    cmds.connectAttr(capsule[0]+".Shaded", capsule[0]+".overrideShading")
+
+    cmds.parentConstraint("root", capsule[0], mo=True)
+    cmds.select(capsule[0], r=True)
+    capsuleDisplayLayer = cmds.createDisplayLayer(nr=True, name="Collision Capsule")
+    cmds.setAttr(capsuleDisplayLayer+".displayType", 2)
+    cmds.setAttr(capsuleDisplayLayer+".v", 0)
+    
+    for attr in ["tx", "ty", "tz", "rx", "ry", "rz", "sx", "sy", "sz"]:
+        cmds.setAttr(capsule[0]+"."+attr, lock=True, keyable=False, channelBox=False)
+
+    cmds.select(clear=True)
+addCapsule(72, 108)
+    
+# ------------------------------------------------------------------
+# ------------------------------------------------------------------
+#The destination attribute 'Rider:ik_foot_anim_l_space_switcher_follow_parentConstraint1.W2' cannot be found.
+# Import in and hook up the pilot for HyperBreach
+def hookUpPilot():
+    import Tools.ART_addSpaceSwitching as space
+    reload(space)
+    
+    # import in rabbit pilot
+    cmds.file("D:/Build/usr/jeremy_ernst/MayaTools/General/ART/Projects/Orion/AnimRigs/HyperBreach_Pilot.mb", i=True, options='v=0', type="mayaBinary", namespace="Rider", mergeNamespacesOnClash=False)
+
+    cmds.setAttr("Rider:Rig_Settings.spine_fk", 1)
+    cmds.setAttr("Rider:Rig_Settings.spine_ik", 0)
+    
+    # set elbows to body space
+    cmds.setAttr("Rider:ik_elbow_l_anim_space_switcher.space_body_anim", 1)
+    cmds.setAttr("Rider:ik_elbow_r_anim_space_switcher.space_body_anim", 1)
+    
+    cmds.setAttr("Rider:ik_wrist_l_anim_space_switcher.space_body_anim", 1)
+    cmds.setAttr("Rider:ik_wrist_r_anim_space_switcher.space_body_anim", 1)
+    
+    cmds.setAttr("Rider:ik_foot_anim_l_space_switcher.space_body_anim", 1)
+    cmds.setAttr("Rider:ik_foot_anim_r_space_switcher.space_body_anim", 1)
+
+    # turn on auto clavicles
+    cmds.setAttr("Rider:clavicle_l_anim.autoShoulders", 1)
+    cmds.setAttr("Rider:clavicle_r_anim.autoShoulders", 1)
+
+    # set head to world space
+    cmds.setAttr("Rider:head_fk_anim.fkOrientation", 3)
+    
+    # Add space switching to the feet on the foot pedals
+    space.addSpaceSwitching("Rider:ik_foot_anim_l", "cockpit_foot_pedal_l_anim", True)
+    space.addSpaceSwitching("Rider:ik_foot_anim_r", "cockpit_foot_pedal_r_anim", True)
+
+    #position rider
+    cmds.setAttr("Rider:body_anim.tx", 116.498)
+    cmds.setAttr("Rider:body_anim.ty", -1.819)
+    cmds.setAttr("Rider:body_anim.rz", -6.414)
+    
+    cmds.setAttr("Rider:spine_01_anim.rz", -9.83)
+    cmds.setAttr("Rider:spine_02_anim.rz", -8.684)
+    cmds.setAttr("Rider:neck_01_fk_anim.rz", 4.835)
+    cmds.setAttr("Rider:neck_02_fk_anim.rz", 4.835)
+    cmds.setAttr("Rider:head_fk_anim.rz", -8.96)
+    
+    cmds.setAttr("Rider:ik_foot_anim_l.tx", 5.161)
+    cmds.setAttr("Rider:ik_foot_anim_l.ty", -15.223)
+    cmds.setAttr("Rider:ik_foot_anim_l.tz", 134.924)
+    cmds.setAttr("Rider:ik_foot_anim_l.rx", 23.82491)
+
+    cmds.setAttr("Rider:ik_foot_anim_r.tx", -5.161)
+    cmds.setAttr("Rider:ik_foot_anim_r.ty", -15.223)
+    cmds.setAttr("Rider:ik_foot_anim_r.tz", 134.924)
+    cmds.setAttr("Rider:ik_foot_anim_r.rx", 23.82491)
+
+    cmds.setAttr("Rider:ik_wrist_l_anim.tx", -21.485)
+    cmds.setAttr("Rider:ik_wrist_l_anim.ty", -30.379)
+    cmds.setAttr("Rider:ik_wrist_l_anim.tz", -15.152)
+    cmds.setAttr("Rider:ik_wrist_l_anim.rx", -88.324)
+    cmds.setAttr("Rider:ik_wrist_l_anim.ry", 2.964)
+    cmds.setAttr("Rider:ik_wrist_l_anim.rz", -79.552)
+
+    cmds.setAttr("Rider:ik_wrist_r_anim.tx", 21.485)
+    cmds.setAttr("Rider:ik_wrist_r_anim.ty", -30.379)
+    cmds.setAttr("Rider:ik_wrist_r_anim.tz", -15.152)
+    cmds.setAttr("Rider:ik_wrist_r_anim.rx", -88.324)
+    cmds.setAttr("Rider:ik_wrist_r_anim.ry", -2.964)
+    cmds.setAttr("Rider:ik_wrist_r_anim.rz", 79.552)
+
+    cmds.setAttr("Rider:ik_elbow_l_anim.tx", -9.817)  
+    cmds.setAttr("Rider:ik_elbow_l_anim.ty", -0.806)  
+    cmds.setAttr("Rider:ik_elbow_l_anim.tz", -28.635)  
+
+    cmds.setAttr("Rider:ik_elbow_r_anim.tx", 9.817)  
+    cmds.setAttr("Rider:ik_elbow_r_anim.ty", -0.806)  
+    cmds.setAttr("Rider:ik_elbow_r_anim.tz", -28.635)  
+
+    for side in ["_l", "_r"]:
+        cmds.setAttr("cockpit_pedal"+side+"_anim.rx", 48.537)
+        cmds.setAttr("cockpit_foot_pedal"+side+"_anim.rx", -44.85)
+
+        cmds.setAttr("Rider:clavicle"+side+"_anim.tx", 2.602)
+        cmds.setAttr("Rider:clavicle"+side+"_anim.ty", 3.083)
+        cmds.setAttr("Rider:clavicle"+side+"_anim.tz", 3.789)
+
+        cmds.setAttr("Rider:toe_wiggle_ctrl"+side+".rz", 7.49313)
+
+        # Finger Pose
+        cmds.setAttr("Rider:thumb_finger_fk_ctrl_1"+side+".rx", 7.398)
+        cmds.setAttr("Rider:thumb_finger_fk_ctrl_1"+side+".ry", 27.169)
+        cmds.setAttr("Rider:thumb_finger_fk_ctrl_1"+side+".rz", 6.525)
+        cmds.setAttr("Rider:thumb_finger_fk_ctrl_2"+side+".rx", 20.505)
+        cmds.setAttr("Rider:thumb_finger_fk_ctrl_2"+side+".ry", -9.166)
+        cmds.setAttr("Rider:thumb_finger_fk_ctrl_2"+side+".rz", 42.297)
+
+        cmds.setAttr("Rider:index_finger_fk_ctrl_1"+side+".rz", 86.544)
+        cmds.setAttr("Rider:index_finger_fk_ctrl_2"+side+".rz", 60.686)
+
+        cmds.setAttr("Rider:middle_finger_fk_ctrl_1"+side+".rx", 4.235)    
+        cmds.setAttr("Rider:middle_finger_fk_ctrl_1"+side+".ry", -5.416)
+        cmds.setAttr("Rider:middle_finger_fk_ctrl_1"+side+".rz", 89.196)    
+        cmds.setAttr("Rider:middle_finger_fk_ctrl_2"+side+".rz", 60.686)
+
+        cmds.setAttr("Rider:ring_finger_fk_ctrl_1"+side+".rx", 12.858)
+        cmds.setAttr("Rider:ring_finger_fk_ctrl_1"+side+".ry", -16.85)
+        cmds.setAttr("Rider:ring_finger_fk_ctrl_1"+side+".rz", 77.769)
+        cmds.setAttr("Rider:ring_finger_fk_ctrl_2"+side+".rz", 60.686)
+
+        cmds.setAttr("Rider:pinky_finger_fk_ctrl_1"+side+".rx", 18.727)
+        cmds.setAttr("Rider:pinky_finger_fk_ctrl_1"+side+".ry", -29.74)
+        cmds.setAttr("Rider:pinky_finger_fk_ctrl_1"+side+".rz", 60.753)
+        cmds.setAttr("Rider:pinky_finger_fk_ctrl_2"+side+".rz", 32.296)
+
+    # parent control arm grps under rider ik hand ctrls
+    space.addSpaceSwitching("l_ctrl_arm_anim", "Rider:ik_wrist_l_anim", True)
+    space.addSpaceSwitching("l_ctrl_arm_anim", "cockpit_seat_anim", False)
+    space.addSpaceSwitching("r_ctrl_arm_anim", "Rider:ik_wrist_r_anim", True)
+    space.addSpaceSwitching("r_ctrl_arm_anim", "cockpit_seat_anim", False)
+
+    # connect rider to the seat
+    space.addSpaceSwitching("Rider:body_anim", "cockpit_seat_anim", True)
+    cmds.parentConstraint("spine_02_anim", "Rider:chest_ik_anim_space_switcher", mo = True)
+
+    # expose knee locs for rider and make knees follow body
+    cmds.select(["Rider:noflip_aim_grp_l", "Rider:noflip_aim_grp_r"])
+    cmds.delete(constraints = True)
+
+    cmds.parentConstraint("body_anim", "Rider:noflip_aim_grp_l", mo = True)
+    cmds.parentConstraint("body_anim", "Rider:noflip_aim_grp_r", mo = True)
+
+    # join skeletons
+    cmds.select("Rider:root", hi = True)
+    joints = cmds.ls(sl = True)
+    for joint in joints:
+        newName = joint.partition(":")[2]
+        cmds.rename(joint, "rider_" + newName)
+    cmds.parent("rider_root", "root")    
+
+    # Disconnect the attrs driving the rider pelvis joint so that we can re-drive it with a parent constraint.
+    attrToDisconnect = "rider_pelvis.translate"
+    connection = cmds.connectionInfo(attrToDisconnect, sfd=True)
+    cmds.disconnectAttr(connection, attrToDisconnect)
+    cmds.delete("rider_pelvis_orientConstraint1")
+
+    # Parent constrain the deform pelvis to follow the driver pelvis
+    cmds.parentConstraint("Rider:driver_pelvis", "rider_pelvis", mo=True)
+
+    cmds.select(None)
+hookUpPilot()
+
+# ------------------------------------------------------------------
+# Cleanup HyperBreach
+def cleanUp():
+    # hide controls
+    controlsToHide = ["fk_wrist_l_anim_grp", "fk_wrist_r_anim_grp"]
+    for control in controlsToHide:
+        cmds.setAttr(control + ".v", 0)
+        for attr in [".tx", ".ty", ".tz", ".rx", ".ry", ".rz", ".sx", ".sy", ".sz"]:
+            if cmds.objExists(control + attr):
+                cmds.setAttr(control + attr, keyable = False, channelBox = False)
+            
+    cmds.select(clear=True)
+cleanUp()
+
+# ------------------------------------------------------------------
+# save the AnimRig file
+cmds.file(save = True, type = "mayaBinary", force = True)
+
+# add IK Bones here so that they are included in the export file.  But dont save becasue we dont want them in the rig file.  Thats why we saved in the line before.
+addIKBones()
+
+# ------------------------------------------------------------------
+# Creeate the export file that goes into source so it can be exported to the game
+def createExportFile():
+    # Import the export file from the reference so that it can be edited.
+    cmds.file("D:/Build/usr/jeremy_ernst/MayaTools/General/ART/Projects/Orion/ExportFiles/HyperBreach_Export.mb", importReference=True)
+
+    # select the entire skeleton and make a list of it.
+    cmds.select("root", hi=True)
+    selection = cmds.ls(sl=True)
+    cmds.showHidden(a=True)
+
+    # for each object in the selection, disconnect translation and scale.
+    for each in selection:
+        attrs = ["translate", "scale"]
+        cmds.lockNode(each, lock=False)
+
+        for attr in attrs:
+            try:
+                conn = cmds.listConnections(each + "." + attr, s=True, p=True)[0]
+                cmds.disconnectAttr(conn, each + "." + attr)
+            except:
+                pass
+
+    # find all of the constraints in the skeleton and delete them
+    cmds.select("root", hi=True)
+    constraints = cmds.ls(sl=True, type="constraint")
+    cmds.select(constraints)
+    cmds.delete(constraints)
+
+    # delete the control rig so that it severs any connections between it and the geo or skeleton.
+    cmds.select("Rider:rig_grp", "rig_grp", hi=True)
+    rigs = cmds.ls(sl=True)
+    for one in rigs:
+        cmds.lockNode(one, lock=False)
+    cmds.delete("Rider:rig_grp", "rig_grp")
+
+    # Remove the Rider namespace and set the bindposes so that Creating LODs will work later on.
+    cmds.namespace(rm="Rider", mnp=True)
+    cmds.dagPose("root", bp=True, r=True) 
+    cmds.dagPose("rider_root", bp=True, r=True) 
+
+    # select the skeleton and all of the geo and export it to the export file.  Remember to go export an fbx from this.
+    cmds.select("Rabbit_Geo", "HyperBreach_Geo", "root", hi=True)
+    cmds.file("D:/Build/ArtSource/Orion/Characters/Heroes/HyperBreach/Export/HyperBreach_Export.mb", force=True, typ="mayaBinary", pr=True, es=True)
+createExportFile()
+
+# ------------------------------------------------------------------
+# ADD IK BONES to the export file
+sceneName = cmds.file(q = True, sceneName = True)
+exportFile = sceneName.partition("AnimRigs")[0] + "ExportFiles" +  sceneName.partition("AnimRigs")[2].partition(".mb")[0] + "_Export.mb"
+cmds.file(exportFile, open = True, force = True)
+
+
+addIKBones()
+
+# ------------------------------------------------------------------
+def addCustomAttrs():
+    attrs = cmds.listAttr("root")
+    exists = False
+    for one in attrs:
+        if one.find("DistanceCurve") != -1:
+            exists = True
+    if exists == False:
+        cmds.addAttr("root", ln="DistanceCurve", keyable=True)
+addCustomAttrs()
+
+# save the export file
+cmds.file(save = True, type = "mayaBinary", force = True)
+
+
