Branching

diff --git a/Engine/Extras/Maya_AnimationRiggingTools/ARTv1/General/ART/Projects/Orion/Scripts/HyperBreach_Rabbit_Post_Script.py b/Engine/Extras/Maya_AnimationRiggingTools/ARTv1/General/ART/Projects/Orion/Scripts/HyperBreach_Rabbit_Post_Script.py
new file mode 100644
--- /dev/null
+++ b/Engine/Extras/Maya_AnimationRiggingTools/ARTv1/General/ART/Projects/Orion/Scripts/HyperBreach_Rabbit_Post_Script.py
@@ -0,0 +1,152 @@
+import maya.cmds as cmds
+
+# ------------------------------------------------------------------
+def heelHelper():
+    
+    sides = ["_l", "_r"]
+    for side in sides:
+    
+        calfJoint = "calf"+side
+        footJoint = "foot"+side
+        
+        heelHelperGrp = "heel_helper"+side+"_anim_grp"
+        
+        # create the target joint
+        heelHelperTarget = cmds.duplicate(heelHelperGrp, name=heelHelperGrp+"_target", renameChildren=True)
+        heelHelperTarget[1] = cmds.rename(heelHelperTarget[1], heelHelperTarget[1]+"_target")
+        
+        if side == "_l":
+            cmds.move(0, -13.21673, 0, heelHelperTarget[0], r=True, os=True, wd=True)
+        else:
+            cmds.move(0, 13.21673, 0, heelHelperTarget[0], r=True, os=True, wd=True)
+        
+        # constrain the rig together
+        cmds.parentConstraint(calfJoint, heelHelperTarget[0], name="heelHelperTargetPrntCnst"+side, mo=True, weight=1)
+        cmds.aimConstraint(heelHelperTarget[1], heelHelperGrp, name="heelHelperDriverAimCnst"+side, mo=True, weight=1, aimVector=(0, -1, 0), upVector=(0, 0, -1), worldUpType="object", worldUpObject="foot_l")
+
+heelHelper()
+
+
+# ------------------------------------------------------------------
+def fixUpToes():
+    # change the way the Set driven keys on the heel and ball pivots work so that the leg can work more like a quadruped leg.
+    cmds.keyframe("ik_foot_heel_pivot_l_rotateZ", edit=True, index=(2,2), absolute=True, valueChange=0)
+    cmds.keyframe("ik_foot_ball_pivot_l_rotateZ", edit=True, index=(2,2), absolute=True, valueChange=90)
+
+    cmds.keyframe("ik_foot_heel_pivot_r_rotateZ", edit=True, index=(2,2), absolute=True, valueChange=0)
+    cmds.keyframe("ik_foot_ball_pivot_r_rotateZ", edit=True, index=(2,2), absolute=True, valueChange=90)
+    
+    cmds.setAttr("driver_tarsal_l_parentConstraint1.rig_tarsal_lW0", 0)
+    cmds.setAttr("driver_tarsal_r_parentConstraint1.rig_tarsal_rW0", 0)
+
+
+fixUpToes()
+
+# ------------------------------------------------------------------
+# Keep the following 2 lines to find the location of these 2 locators after moving them further away from the leg.
+# Tip for finding the best location, duplicate the leg mesh so you have one in the default position.  
+# Then just move these locators away from the body until they are plenty clear and the edges from the rigged mesh and duped one line up.
+#cmds.xform("thigh_twist_driver_l_PV", q=True, ws=True, piv=True)
+#cmds.xform("thigh_twist_driver_l_aimUp", q=True, ws=True, piv=True)
+
+# Fix the twists in the thighs so that they dont twist up so easily when the leg rotates sideways or back.
+def fixTwistInThighs():
+    sides = ["_l", "_r"]
+    for side in sides:
+        # Re-create the pole vector for the thigh IK handle
+        cmds.delete("thigh_twist_driver"+side+"_ikHandle_poleVectorConstraint1")
+        if side == "_l":
+            cmds.xform("thigh_twist_driver"+side+"_PV", ws=True, translation=(40.823986719916384, -10.167150062315951, 59.8508899397534))
+        else:
+            cmds.xform("thigh_twist_driver"+side+"_PV", ws=True, translation=(-40.823986719916384, -10.167150062315951, 59.8508899397534))
+        cmds.poleVectorConstraint("thigh_twist_driver"+side+"_PV", "thigh_twist_driver"+side+"_ikHandle")
+
+        # Move the up for the aim constraint and re-create it.
+        cmds.delete("thigh"+side+"_twist_joint_aimConstraint1")
+        if side == "_l":
+            cmds.xform("thigh_twist_driver"+side+"_aimUp", ws=True, translation=(7.615877731473773, 32.26791726115082, 32.355635500592044))
+        else:
+            cmds.xform("thigh_twist_driver"+side+"_aimUp", ws=True, translation=(-7.615877731473773, 32.26791726115082, 32.355635500592044))
+        cmds.aimConstraint("driver_calf"+side, "thigh"+side+"_twist_joint", weight = 1, aimVector = [-1, 0, 0], upVector = [0, 1, 0], worldUpType = "object", worldUpObject = "thigh_twist_driver"+side+"_aimUp")
+
+fixTwistInThighs()
+
+
+# ------------------------------------------------------------------
+def scaleControls():
+    # brows
+    cmds.select(["brow_out_l_anim.cv[0:32]", "brow_out_r_anim.cv[0:32]"])
+    cmds.scale(0.132495, 0.132495, 0.132495, r = True, ocp = True)
+    
+    cmds.select(["brow_in_l_anim.cv[0:32]", "brow_in_r_anim.cv[0:32]"])
+    cmds.scale(0.132495, 0.132495, 0.132495, r = True, ocp = True)
+
+    cmds.select(["brow_mid_anim.cv[0:32]"])
+    cmds.scale(0.132495, 0.132495, 0.132495, r = True, ocp = True)
+    
+    # eyes
+    cmds.select(["eyelid_top_l_anim.cv[0:32]", "eyelid_top_r_anim.cv[0:32]"])
+    cmds.scale(0.0483613, 0.0483613, 0.0483613, r = True, ocp = True)
+    cmds.move(0, -4.024, 0.759295, r = True)
+    
+    cmds.select(["eyelid_bott_l_anim.cv[0:32]", "eyelid_bott_r_anim.cv[0:32]"])
+    cmds.scale(0.0483613, 0.0483613, 0.0483613, r = True, ocp = True)
+    cmds.move(0, -3.759804, -1.558685, r = True)
+
+    cmds.select(["eyeball_l_anim.cv[0:32]", "eyeball_r_anim.cv[0:32]"])
+    cmds.scale(0.0483613, 0.0483613, 0.0483613, r = True, ocp = True)
+    cmds.move(0, -4.077843, -0.412389, r = True)
+
+    # jaw
+    cmds.select(["jaw_anim.cv[0:32]"])
+    cmds.scale(0.294922, 0.294922, 0.294922, r = True, ocp = True)
+    cmds.move(0, -12, -10, r = True)
+    
+    # holster
+    cmds.select(["gun_anim.cv[0:32]"])
+    cmds.scale(0.294922, 0.294922, 0.294922, r = True, ocp = True)
+    
+    # heel helpers
+    cmds.select(["heel_helper_l_anim.cv[0:32]", "heel_helper_r_anim.cv[0:32]"])
+    cmds.scale(0.294922, 0.294922, 0.294922, r = True, ocp = True)
+    
+    cmds.select(["heel_helper_l_anim1_target.cv[0:32]", "heel_helper_r_anim1_target.cv[0:32]"])
+    cmds.scale(0.2, 0.2, 0.2, r = True, ocp = True)
+    
+    # toes
+    cmds.select(["tarsal_l_anim.cv[0:32]", "tarsal_r_anim.cv[0:32]"])
+    cmds.scale(0.294922, 0.294922, 0.294922, r = True, ocp = True)
+    
+    # ears
+    cmds.select(["fk_earMain_l_01_anim.cv[0:7]", "fk_earMain_r_01_anim.cv[0:7]"])
+    cmds.scale(1.522128, 1.522128, 1.522128, r = True, ocp = False)
+
+    cmds.select(["fk_earMain_l_03_anim.cv[0:7]", "fk_earMain_r_03_anim.cv[0:7]"])
+    cmds.scale(0.696643, 0.696643, 0.696643, r = True, ocp = False)
+    
+    # face fur
+    cmds.select(["fk_headfur_mid_l_02_anim.cv[0:7]", "fk_headfur_mid_r_02_anim.cv[0:7]"])
+    cmds.scale(0.439838, 0.439838, 0.439838, r = True, ocp = True)
+    
+    # neck and head
+    cmds.select(["neck_01_fk_anim.cv[0:7]"])
+    cmds.scale(2.18565, 2.18565, 2.18565, r = True, ocp = False)
+
+    cmds.select(["neck_02_fk_anim.cv[0:7]"])
+    cmds.scale(2.18565, 2.18565, 2.18565, r = True, ocp = False)
+
+    cmds.select(["head_fk_anim.cv[0:7]"])
+    cmds.scale(2.483665, 2.483665, 2.483665, r = True, ocp = True)
+    
+    # foot
+    cmds.select(["ik_foot_anim_l.cv[0:17]"])
+    cmds.scale(1.4, 1.4, 1.4, r = True, ocp = False)
+    cmds.move(0.8, 2.63, -16, r = True)
+    
+    cmds.select(["ik_foot_anim_r.cv[0:17]"])
+    cmds.scale(1.4, 1.4, 1.4, r = True, ocp = False)
+    cmds.move(-0.8, 2.63, -16, r = True)
+
+    cmds.select(None)
+
+scaleControls()
\ No newline at end of file
diff --git a/Engine/Extras/Maya_AnimationRiggingTools/ARTv1/General/ART/Projects/Orion/Scripts/Kurohane_CreateLODs.py b/Engine/Extras/Maya_AnimationRiggingTools/ARTv1/General/ART/Projects/Orion/Scripts/Kurohane_CreateLODs.py
new file mode 100644
--- /dev/null
+++ b/Engine/Extras/Maya_AnimationRiggingTools/ARTv1/General/ART/Projects/Orion/Scripts/Kurohane_CreateLODs.py
@@ -0,0 +1,621 @@
+import maya.cmds as cmds
+import customMayaMenu as menu
+import Tools.ART_ImportMultipleWeights as impw
+import Tools.ART_TransferWeights as transfer
+import Tools.ART_ExportLODtoFBX as export
+reload(export)
+
+################################################## KUROHANE #########################################################################
+# Warn the user that this operation will save their file.
+saveResult = cmds.confirmDialog(title = "Save Warning", message="In order to continue your file must be saved.  Would you like to save it?  If yes it will be saved and after the operation is complete your file will be re-opened.", button = ["Yes", "Cancel"], defaultButton='Yes', cancelButton='Cancel', dismissString='Cancel')
+
+# If they are ok with saving the file, warn them that it will also delete all of the current files in their skin weights directory.
+if saveResult == "Yes":
+    deleteSkinFilesResult = cmds.confirmDialog(title = "SkinWeights Warning", message="In order to continue we must delete all of the files in your ..\General\ART\SkinWeights\ directory.  Would you like to continue?.", button = ["Yes", "Cancel"], defaultButton='Yes', cancelButton='Cancel', dismissString='Cancel')
+    if deleteSkinFilesResult == "Yes":
+        toolsPath = cmds.internalVar(usd = True) + "mayaTools.txt"
+        if os.path.exists(toolsPath):
+            f = open(toolsPath, 'r')
+            mayaToolsDir = f.readline()
+            f.close()
+
+        path = mayaToolsDir + "\General\ART\SkinWeights\\"
+        files = os.listdir(path)
+        #delete all weight cache files in the system folder
+        for file in files:
+            os.remove(path + file)
+            #print path + file
+
+
+
+# Kurohane LOD1 -------------------------------------------------------------------------------------------------------------
+def KurohaneLOD1():
+    if saveResult == "Yes":
+
+        cmds.file(save=True)
+
+        meshes = ["geo_dagger_mechanisms_LOD1", "geo_eyes_LOD1", "geo_helmet_LOD1", "geo_l_wrist_LOD1", "geo_legs_LOD1", "geo_lower_arm_casings_LOD1", "geo_lower_arms_LOD1", "geo_r_wrist_LOD1", "geo_shoulder_plates_LOD1", "geo_sword_l_LOD1", "geo_sword_r_LOD1", "geo_tentacles_LOD1", "geo_thruster_a_LOD1", "geo_thruster_b_LOD1", "geo_thruster_backplate_LOD1", "geo_thruster_c_LOD1", "geo_thruster_home_LOD1", "geo_thruster_vents_LOD1", "geo_thruster_wheel_LOD1", "geo_torso_LOD1", "geo_upper_arms_LOD1"]
+
+        jointsToDelete = []
+
+        # remove unneeded joints
+        # NONE
+        
+        # Delete the joints we no longer want.
+        for one in jointsToDelete:
+            if cmds.objExists(one):
+                cmds.delete(one)
+                
+        # Export the fbx file.
+        export.exportLOD("Kurohane_LOD1.fbx", "Kurohane_Geo_LOD1")
+        
+        # Re-open the file without saving.
+        fullPath = cmds.file(q = True, sceneName = True)
+        cmds.file(fullPath, open=True, f=True)
+
+        print "LOD1 Exported Successfully!!!!!!!!!!!!!!!"
+
+KurohaneLOD1()
+
+# Kurohane LOD2 -------------------------------------------------------------------------------------------------------------
+def KurohaneLOD2():
+    if saveResult == "Yes":
+
+        cmds.file(save=True)
+        
+        meshes = ["geo_dagger_mechanisms_LOD2", "geo_eyes_LOD2", "geo_helmet_LOD2", "geo_l_wrist_LOD2", "geo_legs_LOD2", "geo_lower_arm_casings_LOD2", "geo_lower_arms_LOD2", "geo_r_wrist_LOD2", "geo_shoulder_plates_LOD2", "geo_sword_l_LOD2", "geo_sword_r_LOD2", "geo_tentacles_LOD2", "geo_thruster_a_LOD2", "geo_thruster_b_LOD2", "geo_thruster_c_LOD2", "geo_thruster_home_LOD2", "geo_thruster_vents_LOD2", "geo_torso_LOD2", "geo_upper_arms_LOD2"]
+
+        # Export a skin weights file for each mesh we plan to operate on.
+        cmds.select(meshes)
+        menu.exportSelectedSkinWeights(True)
+
+        # Set the pose if needed
+        cmds.setAttr("thumb_03_l.rotateZ", 0)
+        cmds.setAttr("thumb_03_l.rotateY", -25.5635)
+        cmds.setAttr("thumb_03_l.rotateX", 0)
+        cmds.setAttr("pinky_03_l.rotateZ", 50.45)
+        cmds.setAttr("pinky_03_l.rotateY", 0)
+        cmds.setAttr("pinky_03_l.rotateX", 0)
+        cmds.setAttr("middle_03_l.rotateZ", 50.4503)
+        cmds.setAttr("middle_03_l.rotateY", 0)
+        cmds.setAttr("middle_03_l.rotateX", 0)
+        cmds.setAttr("index_03_l.rotateZ", 50.4503)
+        cmds.setAttr("index_03_l.rotateY", 0)
+        cmds.setAttr("index_03_l.rotateX", 0)
+
+        cmds.setAttr("thumb_03_r.rotateZ", 0)
+        cmds.setAttr("thumb_03_r.rotateY", -25.5635)
+        cmds.setAttr("thumb_03_r.rotateX", 0)
+        cmds.setAttr("pinky_03_r.rotateZ", 50.45)
+        cmds.setAttr("pinky_03_r.rotateY", 0)
+        cmds.setAttr("pinky_03_r.rotateX", 0)
+        cmds.setAttr("middle_03_r.rotateZ", 50.4503)
+        cmds.setAttr("middle_03_r.rotateY", 0)
+        cmds.setAttr("middle_03_r.rotateX", 0)
+        cmds.setAttr("index_03_r.rotateZ", 50.4503)
+        cmds.setAttr("index_03_r.rotateY", 0)
+        cmds.setAttr("index_03_r.rotateX", 0)
+
+        # Delete the history on the meshes so that they get baked down into their posed positions
+        cmds.delete(meshes, ch=True)
+
+        # Move the skeleton back to its bindpose (This script REQUIRES the bind pose to be the pose that we want to export the skeleton to.  This is the case for all known characters created with A.R.T.)
+        cmds.dagPose("root", bp=True, r=True) 
+
+        # Re-import the skinning we exported before onto the mesh so thats its bound in its posed position.
+        cmds.select(meshes)
+        impw.importMultipleSkinWeights("", True)
+
+        jointsToDelete = []
+
+         # remove unneeded joints
+        jointsToRemove = ["thrusterVent_l_01", "thrusterVent_l_02", "thrusterVent_r_01", "thrusterVent_r_02", "thruster_casing_l", "thruster_a_l", "thruster_b_l", "thruster_c_l", "thruster_casing_r", "thruster_a_r", "thruster_b_r", "thruster_c_r"]
+        jointToTransferTo = "spine_03"
+        transfer.LOD_transferWeights(meshes, jointsToRemove, jointToTransferTo)
+        jointsToDelete.extend(jointsToRemove)
+
+        jointsToRemove = ["shoulderflap_l"]
+        jointToTransferTo = "clavicle_l"
+        transfer.LOD_transferWeights(meshes, jointsToRemove, jointToTransferTo)
+        jointsToDelete.extend(jointsToRemove)
+
+        jointsToRemove = ["shoulderflap_r"]
+        jointToTransferTo = "clavicle_r"
+        transfer.LOD_transferWeights(meshes, jointsToRemove, jointToTransferTo)
+        jointsToDelete.extend(jointsToRemove)
+
+        jointsToRemove = ["shoulderPad_l"]
+        jointToTransferTo = "upperarm_l"
+        transfer.LOD_transferWeights(meshes, jointsToRemove, jointToTransferTo)
+        jointsToDelete.extend(jointsToRemove)
+
+        jointsToRemove = ["shoulderPad_r"]
+        jointToTransferTo = "upperarm_r"
+        transfer.LOD_transferWeights(meshes, jointsToRemove, jointToTransferTo)
+        jointsToDelete.extend(jointsToRemove)
+
+        jointsToRemove = ["dagger_c_l", "dagger_b_l", "dagger_a_l"]
+        jointToTransferTo = "dagger_base_l"
+        transfer.LOD_transferWeights(meshes, jointsToRemove, jointToTransferTo)
+        jointsToDelete.extend(jointsToRemove)
+
+        jointsToRemove = ["dagger_c_r", "dagger_b_r", "dagger_a_r"]
+        jointToTransferTo = "dagger_base_r"
+        transfer.LOD_transferWeights(meshes, jointsToRemove, jointToTransferTo)
+        jointsToDelete.extend(jointsToRemove)
+
+        jointsToRemove = ["index_03_l"]
+        jointToTransferTo = "index_02_l"
+        transfer.LOD_transferWeights(meshes, jointsToRemove, jointToTransferTo)
+        jointsToDelete.extend(jointsToRemove)
+
+        jointsToRemove = ["middle_03_l"]
+        jointToTransferTo = "middle_02_l"
+        transfer.LOD_transferWeights(meshes, jointsToRemove, jointToTransferTo)
+        jointsToDelete.extend(jointsToRemove)
+
+        jointsToRemove = ["pinky_03_l"]
+        jointToTransferTo = "pinky_02_l"
+        transfer.LOD_transferWeights(meshes, jointsToRemove, jointToTransferTo)
+        jointsToDelete.extend(jointsToRemove)
+
+        jointsToRemove = ["thumb_03_l"]
+        jointToTransferTo = "thumb_02_l"
+        transfer.LOD_transferWeights(meshes, jointsToRemove, jointToTransferTo)
+        jointsToDelete.extend(jointsToRemove)
+
+        jointsToRemove = ["index_03_r"]
+        jointToTransferTo = "index_02_r"
+        transfer.LOD_transferWeights(meshes, jointsToRemove, jointToTransferTo)
+        jointsToDelete.extend(jointsToRemove)
+
+        jointsToRemove = ["middle_03_r"]
+        jointToTransferTo = "middle_02_r"
+        transfer.LOD_transferWeights(meshes, jointsToRemove, jointToTransferTo)
+        jointsToDelete.extend(jointsToRemove)
+
+        jointsToRemove = ["pinky_03_r"]
+        jointToTransferTo = "pinky_02_r"
+        transfer.LOD_transferWeights(meshes, jointsToRemove, jointToTransferTo)
+        jointsToDelete.extend(jointsToRemove)
+
+        jointsToRemove = ["thumb_03_r"]
+        jointToTransferTo = "thumb_02_r"
+        transfer.LOD_transferWeights(meshes, jointsToRemove, jointToTransferTo)
+        jointsToDelete.extend(jointsToRemove)
+
+        jointsToRemove = ["tentacle_spring_l", "tentacle_spring_r"]
+        jointToTransferTo = "head"
+        transfer.LOD_transferWeights(meshes, jointsToRemove, jointToTransferTo)
+        jointsToDelete.extend(jointsToRemove)
+
+        # Delete the joints we no longer want.
+        for one in jointsToDelete:
+            if cmds.objExists(one):
+                cmds.delete(one)
+                
+        # Export the fbx file.
+        export.exportLOD("Kurohane_LOD2.fbx", "Kurohane_Geo_LOD2")
+
+        # Re-open the file without saving.
+        fullPath = cmds.file(q = True, sceneName = True)
+        cmds.file(fullPath, open=True, f=True)
+
+        print "LOD2 Exported Successfully!!!!!!!!!!!!!!!"
+
+KurohaneLOD2()
+
+# Kurohane LOD3 -------------------------------------------------------------------------------------------------------------
+def KurohaneLOD3():
+    if saveResult == "Yes":
+
+        cmds.file(save=True)
+    
+        meshes = ["geo_dagger_mechanisms_LOD3", "geo_helmet_LOD3", "geo_l_wrist_LOD3", "geo_legs_LOD3", "geo_lower_arm_casings_LOD3", "geo_lower_arms_LOD3", "geo_r_wrist_LOD3", "geo_shoulder_plates_LOD3", "geo_sword_l_LOD3", "geo_sword_r_LOD3", "geo_tentacles_LOD3", "geo_thruster_a_LOD3", "geo_thruster_b_LOD3", "geo_thruster_c_LOD3", "geo_thruster_home_LOD3", "geo_thruster_vents_LOD3", "geo_torso_LOD3", "geo_upper_arms_LOD3"]
+    
+        # Export a skin weights file for each mesh we plan to operate on.
+        cmds.select(meshes)
+        menu.exportSelectedSkinWeights(True)
+
+        # Set the pose if needed
+        cmds.setAttr("thumb_03_l.rotateZ", 0)
+        cmds.setAttr("thumb_03_l.rotateY", -35.5635)
+        cmds.setAttr("thumb_03_l.rotateX", 0)
+        cmds.setAttr("thumb_02_l.rotateZ", 0)
+        cmds.setAttr("thumb_02_l.rotateY", -14.1631)
+        cmds.setAttr("thumb_02_l.rotateX", 0)
+        cmds.setAttr("thumb_01_l.rotateZ", -13.26)
+        cmds.setAttr("thumb_01_l.rotateY", 38.59)
+        cmds.setAttr("thumb_01_l.rotateX", -8.362)
+        cmds.setAttr("pinky_03_l.rotateZ", 77.45)
+        cmds.setAttr("pinky_03_l.rotateY", 0)
+        cmds.setAttr("pinky_03_l.rotateX", 0)
+        cmds.setAttr("pinky_02_l.rotateZ", 120.683)
+        cmds.setAttr("pinky_02_l.rotateY", 0)
+        cmds.setAttr("pinky_02_l.rotateX", 0)
+        cmds.setAttr("pinky_01_l.rotateZ", 32.2949)
+        cmds.setAttr("pinky_01_l.rotateY", 0)
+        cmds.setAttr("pinky_01_l.rotateX", 0)
+        cmds.setAttr("middle_03_l.rotateZ", 77.4503)
+        cmds.setAttr("middle_03_l.rotateY", 0)
+        cmds.setAttr("middle_03_l.rotateX", 0)
+        cmds.setAttr("middle_02_l.rotateZ", 120.683)
+        cmds.setAttr("middle_02_l.rotateY", 0)
+        cmds.setAttr("middle_02_l.rotateX", 0)
+        cmds.setAttr("middle_01_l.rotateZ", 32.294)
+        cmds.setAttr("middle_01_l.rotateY", 0)
+        cmds.setAttr("middle_01_l.rotateX", 0)
+        cmds.setAttr("index_03_l.rotateZ", 77.4503)
+        cmds.setAttr("index_03_l.rotateY", 0)
+        cmds.setAttr("index_03_l.rotateX", 0)
+        cmds.setAttr("index_02_l.rotateZ", 120.683)
+        cmds.setAttr("index_02_l.rotateY", 0)
+        cmds.setAttr("index_02_l.rotateX", 0)
+        cmds.setAttr("index_01_l.rotateZ", 32.294)
+        cmds.setAttr("index_01_l.rotateY", 0)
+        cmds.setAttr("index_01_l.rotateX", 0)
+
+        cmds.setAttr("thumb_03_r.rotateZ", 0)
+        cmds.setAttr("thumb_03_r.rotateY", -35.5635)
+        cmds.setAttr("thumb_03_r.rotateX", 0)
+        cmds.setAttr("thumb_02_r.rotateZ", 0)
+        cmds.setAttr("thumb_02_r.rotateY", -14.1631)
+        cmds.setAttr("thumb_02_r.rotateX", 0)
+        cmds.setAttr("thumb_01_r.rotateZ", -13.26)
+        cmds.setAttr("thumb_01_r.rotateY", 38.59)
+        cmds.setAttr("thumb_01_r.rotateX", -8.362)
+        cmds.setAttr("pinky_03_r.rotateZ", 77.45)
+        cmds.setAttr("pinky_03_r.rotateY", 0)
+        cmds.setAttr("pinky_03_r.rotateX", 0)
+        cmds.setAttr("pinky_02_r.rotateZ", 120.683)
+        cmds.setAttr("pinky_02_r.rotateY", 0)
+        cmds.setAttr("pinky_02_r.rotateX", 0)
+        cmds.setAttr("pinky_01_r.rotateZ", 32.2949)
+        cmds.setAttr("pinky_01_r.rotateY", 0)
+        cmds.setAttr("pinky_01_r.rotateX", 0)
+        cmds.setAttr("middle_03_r.rotateZ", 77.4503)
+        cmds.setAttr("middle_03_r.rotateY", 0)
+        cmds.setAttr("middle_03_r.rotateX", 0)
+        cmds.setAttr("middle_02_r.rotateZ", 120.683)
+        cmds.setAttr("middle_02_r.rotateY", 0)
+        cmds.setAttr("middle_02_r.rotateX", 0)
+        cmds.setAttr("middle_01_r.rotateZ", 32.294)
+        cmds.setAttr("middle_01_r.rotateY", 0)
+        cmds.setAttr("middle_01_r.rotateX", 0)
+        cmds.setAttr("index_03_r.rotateZ", 77.4503)
+        cmds.setAttr("index_03_r.rotateY", 0)
+        cmds.setAttr("index_03_r.rotateX", 0)
+        cmds.setAttr("index_02_r.rotateZ", 120.683)
+        cmds.setAttr("index_02_r.rotateY", 0)
+        cmds.setAttr("index_02_r.rotateX", 0)
+        cmds.setAttr("index_01_r.rotateZ", 32.294)
+        cmds.setAttr("index_01_r.rotateY", 0)
+        cmds.setAttr("index_01_r.rotateX", 0)
+
+        # Delete the history on the meshes so that they get baked down into their posed positions
+        cmds.delete(meshes, ch=True)
+
+        # Move the skeleton back to its bindpose (This script REQUIRES the bind pose to be the pose that we want to export the skeleton to.  This is the case for all known characters created with A.R.T.)
+        cmds.dagPose("root", bp=True, r=True) 
+
+        # Re-import the skinning we exported before onto the mesh so thats its bound in its posed position.
+        cmds.select(meshes)
+        impw.importMultipleSkinWeights("", True)
+        
+        jointsToDelete = []
+
+        # remove unneeded joints
+        jointsToRemove = ["thrusterVent_l_01", "thrusterVent_l_02", "thrusterVent_r_01", "thrusterVent_r_02", "thruster_casing_l", "thruster_a_l", "thruster_b_l", "thruster_c_l", "thruster_casing_r", "thruster_a_r", "thruster_b_r", "thruster_c_r"]
+        jointToTransferTo = "spine_03"
+        transfer.LOD_transferWeights(meshes, jointsToRemove, jointToTransferTo)
+        jointsToDelete.extend(jointsToRemove)
+
+        jointsToRemove = ["shoulderflap_l"]
+        jointToTransferTo = "clavicle_l"
+        transfer.LOD_transferWeights(meshes, jointsToRemove, jointToTransferTo)
+        jointsToDelete.extend(jointsToRemove)
+
+        jointsToRemove = ["shoulderflap_r"]
+        jointToTransferTo = "clavicle_r"
+        transfer.LOD_transferWeights(meshes, jointsToRemove, jointToTransferTo)
+        jointsToDelete.extend(jointsToRemove)
+
+        jointsToRemove = ["upperarm_twist_01_l", "shoulderPad_l", "upperarm_twist_02_l"]
+        jointToTransferTo = "upperarm_l"
+        transfer.LOD_transferWeights(meshes, jointsToRemove, jointToTransferTo)
+        jointsToDelete.extend(jointsToRemove)
+
+        jointsToRemove = ["upperarm_twist_01_r", "shoulderPad_r", "upperarm_twist_02_r"]
+        jointToTransferTo = "upperarm_r"
+        transfer.LOD_transferWeights(meshes, jointsToRemove, jointToTransferTo)
+        jointsToDelete.extend(jointsToRemove)
+
+        jointsToRemove = ["lowerArm_bottomPlateA_l", "lowerArm_bottomPlateB_l", "lowerArm_topPlateA_l", "lowerArm_topPlateB_l"]
+        jointToTransferTo = "lowerarm_l"
+        transfer.LOD_transferWeights(meshes, jointsToRemove, jointToTransferTo)
+        jointsToDelete.extend(jointsToRemove)
+
+        jointsToRemove = ["dagger_c_l", "dagger_b_l", "dagger_a_l"]
+        jointToTransferTo = "dagger_base_l"
+        transfer.LOD_transferWeights(meshes, jointsToRemove, jointToTransferTo)
+        jointsToDelete.extend(jointsToRemove)
+
+        jointsToRemove = ["lowerArm_bottomPlateA_r", "lowerArm_bottomPlateB_r", "lowerArm_topPlateA_r", "lowerArm_topPlateB_r"]
+        jointToTransferTo = "dagger_base_r"
+        transfer.LOD_transferWeights(meshes, jointsToRemove, jointToTransferTo)
+        jointsToDelete.extend(jointsToRemove)
+
+        jointsToRemove = ["dagger_c_r", "dagger_b_r", "dagger_a_r"]
+        jointToTransferTo = "lowerarm_r"
+        transfer.LOD_transferWeights(meshes, jointsToRemove, jointToTransferTo)
+        jointsToDelete.extend(jointsToRemove)
+
+        jointsToRemove = ["index_01_l", "index_02_l", "index_03_l", "middle_01_l", "middle_02_l", "middle_03_l", "pinky_01_l", "pinky_02_l", "pinky_03_l", "thumb_01_l", "thumb_02_l", "thumb_03_l"]
+        jointToTransferTo = "hand_l"
+        transfer.LOD_transferWeights(meshes, jointsToRemove, jointToTransferTo)
+        jointsToDelete.extend(jointsToRemove)
+
+        jointsToRemove = ["index_01_r", "index_02_r", "index_03_r", "middle_01_r", "middle_02_r", "middle_03_r", "pinky_01_r", "pinky_02_r", "pinky_03_r", "thumb_01_r", "thumb_02_r", "thumb_03_r"]
+        jointToTransferTo = "hand_r"
+        transfer.LOD_transferWeights(meshes, jointsToRemove, jointToTransferTo)
+        jointsToDelete.extend(jointsToRemove)
+
+        jointsToRemove = ["sword_root_l", "sword_handle_l", "pivotA_l", "extensionA_l", "pivotB_l", "extensionB_l"]
+        jointToTransferTo = "weapon_l"
+        transfer.LOD_transferWeights(meshes, jointsToRemove, jointToTransferTo)
+        jointsToDelete.extend(jointsToRemove)
+
+        jointsToRemove = ["sword_root_r", "sword_handle_r", "pivotA_r", "extensionA_r", "pivotB_r", "extensionB_r"]
+        jointToTransferTo = "weapon_r"
+        transfer.LOD_transferWeights(meshes, jointsToRemove, jointToTransferTo)
+        jointsToDelete.extend(jointsToRemove)
+
+        jointsToRemove = ["tentacle_spring_l", "tentacle_spring_r"]
+        jointToTransferTo = "head"
+        transfer.LOD_transferWeights(meshes, jointsToRemove, jointToTransferTo)
+        jointsToDelete.extend(jointsToRemove)
+
+        jointsToRemove = ["thigh_twist_01_l", "thigh_twist_02_l"]
+        jointToTransferTo = "thigh_l"
+        transfer.LOD_transferWeights(meshes, jointsToRemove, jointToTransferTo)
+        jointsToDelete.extend(jointsToRemove)
+
+        jointsToRemove = ["thigh_twist_01_r", "thigh_twist_02_r"]
+        jointToTransferTo = "thigh_r"
+        transfer.LOD_transferWeights(meshes, jointsToRemove, jointToTransferTo)
+        jointsToDelete.extend(jointsToRemove)
+
+        jointsToRemove = ["calf_twist_01_l"]
+        jointToTransferTo = "calf_l"
+        transfer.LOD_transferWeights(meshes, jointsToRemove, jointToTransferTo)
+        jointsToDelete.extend(jointsToRemove)
+
+        jointsToRemove = ["calf_twist_01_r"]
+        jointToTransferTo = "calf_r"
+        transfer.LOD_transferWeights(meshes, jointsToRemove, jointToTransferTo)
+        jointsToDelete.extend(jointsToRemove)
+
+        # Delete the joints we no longer want.
+        for one in jointsToDelete:
+            if cmds.objExists(one):
+                cmds.delete(one)
+                
+        # Export the fbx file.
+        export.exportLOD("Kurohane_LOD3.fbx", "Kurohane_Geo_LOD3")
+        
+        # Re-open the file without saving.
+        fullPath = cmds.file(q = True, sceneName = True)
+        cmds.file(fullPath, open=True, f=True)
+
+        print "LOD3 Exported Successfully!!!!!!!!!!!!!!!"
+
+KurohaneLOD3()
+
+# Kurohane LOD4 -------------------------------------------------------------------------------------------------------------
+def KurohaneLOD4():
+    if saveResult == "Yes":
+
+        cmds.file(save=True)
+    
+        meshes = ["geo_dagger_mechanisms_LOD4", "geo_helmet_LOD4", "geo_legs_LOD4", "geo_lower_arm_casings_LOD4", "geo_lower_arms_LOD4", "geo_shoulder_plates_LOD4", "geo_sword_l_LOD4", "geo_sword_r_LOD4", "geo_thruster_vents_LOD4", "geo_torso_LOD4", "geo_upper_arms_LOD4"]
+
+
+        # Export a skin weights file for each mesh we plan to operate on.
+        cmds.select(meshes)
+        menu.exportSelectedSkinWeights(True)
+
+        # Set the pose if needed
+        cmds.setAttr("thumb_03_l.rotateZ", 0)
+        cmds.setAttr("thumb_03_l.rotateY", -35.5635)
+        cmds.setAttr("thumb_03_l.rotateX", 0)
+        cmds.setAttr("thumb_02_l.rotateZ", 0)
+        cmds.setAttr("thumb_02_l.rotateY", -14.1631)
+        cmds.setAttr("thumb_02_l.rotateX", 0)
+        cmds.setAttr("thumb_01_l.rotateZ", -13.26)
+        cmds.setAttr("thumb_01_l.rotateY", 38.59)
+        cmds.setAttr("thumb_01_l.rotateX", -8.362)
+        cmds.setAttr("pinky_03_l.rotateZ", 77.45)
+        cmds.setAttr("pinky_03_l.rotateY", 0)
+        cmds.setAttr("pinky_03_l.rotateX", 0)
+        cmds.setAttr("pinky_02_l.rotateZ", 120.683)
+        cmds.setAttr("pinky_02_l.rotateY", 0)
+        cmds.setAttr("pinky_02_l.rotateX", 0)
+        cmds.setAttr("pinky_01_l.rotateZ", 32.2949)
+        cmds.setAttr("pinky_01_l.rotateY", 0)
+        cmds.setAttr("pinky_01_l.rotateX", 0)
+        cmds.setAttr("middle_03_l.rotateZ", 77.4503)
+        cmds.setAttr("middle_03_l.rotateY", 0)
+        cmds.setAttr("middle_03_l.rotateX", 0)
+        cmds.setAttr("middle_02_l.rotateZ", 120.683)
+        cmds.setAttr("middle_02_l.rotateY", 0)
+        cmds.setAttr("middle_02_l.rotateX", 0)
+        cmds.setAttr("middle_01_l.rotateZ", 32.294)
+        cmds.setAttr("middle_01_l.rotateY", 0)
+        cmds.setAttr("middle_01_l.rotateX", 0)
+        cmds.setAttr("index_03_l.rotateZ", 77.4503)
+        cmds.setAttr("index_03_l.rotateY", 0)
+        cmds.setAttr("index_03_l.rotateX", 0)
+        cmds.setAttr("index_02_l.rotateZ", 120.683)
+        cmds.setAttr("index_02_l.rotateY", 0)
+        cmds.setAttr("index_02_l.rotateX", 0)
+        cmds.setAttr("index_01_l.rotateZ", 32.294)
+        cmds.setAttr("index_01_l.rotateY", 0)
+        cmds.setAttr("index_01_l.rotateX", 0)
+
+        cmds.setAttr("thumb_03_r.rotateZ", 0)
+        cmds.setAttr("thumb_03_r.rotateY", -35.5635)
+        cmds.setAttr("thumb_03_r.rotateX", 0)
+        cmds.setAttr("thumb_02_r.rotateZ", 0)
+        cmds.setAttr("thumb_02_r.rotateY", -14.1631)
+        cmds.setAttr("thumb_02_r.rotateX", 0)
+        cmds.setAttr("thumb_01_r.rotateZ", -13.26)
+        cmds.setAttr("thumb_01_r.rotateY", 38.59)
+        cmds.setAttr("thumb_01_r.rotateX", -8.362)
+        cmds.setAttr("pinky_03_r.rotateZ", 77.45)
+        cmds.setAttr("pinky_03_r.rotateY", 0)
+        cmds.setAttr("pinky_03_r.rotateX", 0)
+        cmds.setAttr("pinky_02_r.rotateZ", 120.683)
+        cmds.setAttr("pinky_02_r.rotateY", 0)
+        cmds.setAttr("pinky_02_r.rotateX", 0)
+        cmds.setAttr("pinky_01_r.rotateZ", 32.2949)
+        cmds.setAttr("pinky_01_r.rotateY", 0)
+        cmds.setAttr("pinky_01_r.rotateX", 0)
+        cmds.setAttr("middle_03_r.rotateZ", 77.4503)
+        cmds.setAttr("middle_03_r.rotateY", 0)
+        cmds.setAttr("middle_03_r.rotateX", 0)
+        cmds.setAttr("middle_02_r.rotateZ", 120.683)
+        cmds.setAttr("middle_02_r.rotateY", 0)
+        cmds.setAttr("middle_02_r.rotateX", 0)
+        cmds.setAttr("middle_01_r.rotateZ", 32.294)
+        cmds.setAttr("middle_01_r.rotateY", 0)
+        cmds.setAttr("middle_01_r.rotateX", 0)
+        cmds.setAttr("index_03_r.rotateZ", 77.4503)
+        cmds.setAttr("index_03_r.rotateY", 0)
+        cmds.setAttr("index_03_r.rotateX", 0)
+        cmds.setAttr("index_02_r.rotateZ", 120.683)
+        cmds.setAttr("index_02_r.rotateY", 0)
+        cmds.setAttr("index_02_r.rotateX", 0)
+        cmds.setAttr("index_01_r.rotateZ", 32.294)
+        cmds.setAttr("index_01_r.rotateY", 0)
+        cmds.setAttr("index_01_r.rotateX", 0)
+
+        # Delete the history on the meshes so that they get baked down into their posed positions
+        cmds.delete(meshes, ch=True)
+
+        # Move the skeleton back to its bindpose (This script REQUIRES the bind pose to be the pose that we want to export the skeleton to.  This is the case for all known characters created with A.R.T.)
+        cmds.dagPose("root", bp=True, r=True) 
+
+        # Re-import the skinning we exported before onto the mesh so thats its bound in its posed position.
+        cmds.select(meshes)
+        impw.importMultipleSkinWeights("", True)
+        
+        jointsToDelete = []
+
+        # remove unneeded joints
+        jointsToRemove = ["thrusterVent_l_01", "thrusterVent_l_02", "thrusterVent_r_01", "thrusterVent_r_02", "thruster_casing_l", "thruster_a_l", "thruster_b_l", "thruster_c_l", "thruster_casing_r", "thruster_a_r", "thruster_b_r", "thruster_c_r"]
+        jointToTransferTo = "spine_03"
+        transfer.LOD_transferWeights(meshes, jointsToRemove, jointToTransferTo)
+        jointsToDelete.extend(jointsToRemove)
+
+        jointsToRemove = ["shoulderflap_l"]
+        jointToTransferTo = "clavicle_l"
+        transfer.LOD_transferWeights(meshes, jointsToRemove, jointToTransferTo)
+        jointsToDelete.extend(jointsToRemove)
+
+        jointsToRemove = ["shoulderflap_r"]
+        jointToTransferTo = "clavicle_r"
+        transfer.LOD_transferWeights(meshes, jointsToRemove, jointToTransferTo)
+        jointsToDelete.extend(jointsToRemove)
+
+        jointsToRemove = ["upperarm_twist_01_l", "shoulderPad_l", "upperarm_twist_02_l"]
+        jointToTransferTo = "upperarm_l"
+        transfer.LOD_transferWeights(meshes, jointsToRemove, jointToTransferTo)
+        jointsToDelete.extend(jointsToRemove)
+
+        jointsToRemove = ["upperarm_twist_01_r", "shoulderPad_r", "upperarm_twist_02_r"]
+        jointToTransferTo = "upperarm_r"
+        transfer.LOD_transferWeights(meshes, jointsToRemove, jointToTransferTo)
+        jointsToDelete.extend(jointsToRemove)
+
+        jointsToRemove = ["lowerArm_bottomPlateA_l", "lowerArm_bottomPlateB_l", "lowerArm_topPlateA_l", "lowerArm_topPlateB_l", "dagger_c_l", "dagger_b_l", "dagger_a_l", "dagger_base_l"]
+        jointToTransferTo = "lowerarm_l"
+        transfer.LOD_transferWeights(meshes, jointsToRemove, jointToTransferTo)
+        jointsToDelete.extend(jointsToRemove)
+
+        jointsToRemove = ["lowerArm_bottomPlateA_r", "lowerArm_bottomPlateB_r", "lowerArm_topPlateA_r", "lowerArm_topPlateB_r", "dagger_c_r", "dagger_b_r", "dagger_a_r", "dagger_base_r"]
+        jointToTransferTo = "lowerarm_r"
+        transfer.LOD_transferWeights(meshes, jointsToRemove, jointToTransferTo)
+        jointsToDelete.extend(jointsToRemove)
+
+        jointsToRemove = ["index_01_l", "index_02_l", "index_03_l", "middle_01_l", "middle_02_l", "middle_03_l", "pinky_01_l", "pinky_02_l", "pinky_03_l", "thumb_01_l", "thumb_02_l", "thumb_03_l"]
+        jointToTransferTo = "hand_l"
+        transfer.LOD_transferWeights(meshes, jointsToRemove, jointToTransferTo)
+        jointsToDelete.extend(jointsToRemove)
+
+        jointsToRemove = ["index_01_r", "index_02_r", "index_03_r", "middle_01_r", "middle_02_r", "middle_03_r", "pinky_01_r", "pinky_02_r", "pinky_03_r", "thumb_01_r", "thumb_02_r", "thumb_03_r"]
+        jointToTransferTo = "hand_r"
+        transfer.LOD_transferWeights(meshes, jointsToRemove, jointToTransferTo)
+        jointsToDelete.extend(jointsToRemove)
+
+        jointsToRemove = ["sword_root_l", "sword_handle_l", "pivotA_l", "extensionA_l", "pivotB_l", "extensionB_l"]
+        jointToTransferTo = "weapon_l"
+        transfer.LOD_transferWeights(meshes, jointsToRemove, jointToTransferTo)
+        jointsToDelete.extend(jointsToRemove)
+
+        jointsToRemove = ["sword_root_r", "sword_handle_r", "pivotA_r", "extensionA_r", "pivotB_r", "extensionB_r"]
+        jointToTransferTo = "weapon_r"
+        transfer.LOD_transferWeights(meshes, jointsToRemove, jointToTransferTo)
+        jointsToDelete.extend(jointsToRemove)
+
+        jointsToRemove = ["tentacle_l_01", "tentacle_l_02", "tentacle_l_03", "tentacle_l_04", "tentacle_l_05", "tentacle_l_06", "tentacle_l_07", "tentacle_l_08", "tentacle_l_09", "tentacle_l_010", "tentacle_r_01", "tentacle_r_02", "tentacle_r_03", "tentacle_r_04", "tentacle_r_05", "tentacle_r_06", "tentacle_r_07", "tentacle_r_08", "tentacle_r_09", "tentacle_r_010", "tentacle_spring_l", "tentacle_spring_r"]
+        jointToTransferTo = "head"
+        transfer.LOD_transferWeights(meshes, jointsToRemove, jointToTransferTo)
+        jointsToDelete.extend(jointsToRemove)
+
+        jointsToRemove = ["thigh_twist_01_l", "thigh_twist_02_l"]
+        jointToTransferTo = "thigh_l"
+        transfer.LOD_transferWeights(meshes, jointsToRemove, jointToTransferTo)
+        jointsToDelete.extend(jointsToRemove)
+
+        jointsToRemove = ["thigh_twist_01_r", "thigh_twist_02_r"]
+        jointToTransferTo = "thigh_r"
+        transfer.LOD_transferWeights(meshes, jointsToRemove, jointToTransferTo)
+        jointsToDelete.extend(jointsToRemove)
+
+        jointsToRemove = ["calf_twist_01_l"]
+        jointToTransferTo = "calf_l"
+        transfer.LOD_transferWeights(meshes, jointsToRemove, jointToTransferTo)
+        jointsToDelete.extend(jointsToRemove)
+
+        jointsToRemove = ["calf_twist_01_r"]
+        jointToTransferTo = "calf_r"
+        transfer.LOD_transferWeights(meshes, jointsToRemove, jointToTransferTo)
+        jointsToDelete.extend(jointsToRemove)
+
+        jointsToRemove = ["ball_l"]
+        jointToTransferTo = "foot_l"
+        transfer.LOD_transferWeights(meshes, jointsToRemove, jointToTransferTo)
+        jointsToDelete.extend(jointsToRemove)
+
+        jointsToRemove = ["ball_r"]
+        jointToTransferTo = "foot_r"
+        transfer.LOD_transferWeights(meshes, jointsToRemove, jointToTransferTo)
+        jointsToDelete.extend(jointsToRemove)
+
+        # Delete the joints we no longer want.
+        for one in jointsToDelete:
+            if cmds.objExists(one):
+                cmds.delete(one)
+                
+        # Export the fbx file.
+        export.exportLOD("Kurohane_LOD4.fbx", "Kurohane_Geo_LOD4")
+        
+        # Re-open the file without saving.
+        fullPath = cmds.file(q = True, sceneName = True)
+        cmds.file(fullPath, open=True, f=True)
+
+        print "LOD4 Exported Successfully!!!!!!!!!!!!!!!"
+
+KurohaneLOD4()
+
+
diff --git a/Engine/Extras/Maya_AnimationRiggingTools/ARTv1/General/ART/Projects/Orion/Scripts/Kurohane_Post_Script.py b/Engine/Extras/Maya_AnimationRiggingTools/ARTv1/General/ART/Projects/Orion/Scripts/Kurohane_Post_Script.py
new file mode 100644
--- /dev/null
+++ b/Engine/Extras/Maya_AnimationRiggingTools/ARTv1/General/ART/Projects/Orion/Scripts/Kurohane_Post_Script.py
@@ -0,0 +1,926 @@
+import maya.cmds as cmds
+
+#Hide controls that will be driven
+controlsToHide = ["thruster_a_l_anim", "thruster_b_l_anim", "thruster_c_l_anim", "thruster_a_r_anim", "thruster_b_r_anim", "thruster_c_r_anim",
+"thruster_casing_l_anim", "thruster_casing_r_anim", "dagger_a_r_anim", "dagger_b_r_anim", "dagger_c_r_anim",
+"dagger_a_l_anim", "dagger_b_l_anim", "dagger_c_l_anim", "dagger_base_l_anim", "dagger_base_r_anim", "pivotA_r_anim", "pivotB_r_anim",
+"extensionA_r_anim", "extensionB_r_anim", "pivotA_l_anim", "pivotB_l_anim", "extensionA_l_anim", "extensionB_l_anim", "lowerArm_bottomPlateA_r_anim",
+"lowerArm_bottomPlateB_r_anim", "lowerArm_topPlateB_r_anim", "lowerArm_topPlateA_r_anim", "lowerArm_bottomPlateA_l_anim", "lowerArm_bottomPlateB_l_anim",
+"lowerArm_topPlateB_l_anim", "lowerArm_topPlateA_l_anim"]
+
+specialControls = ["thrusterVent_l_master_ctrl_grp", "thrusterVent_r_master_ctrl_grp"]
+
+for control in controlsToHide:
+	grp = control + "_grp"
+	cmds.setAttr(grp + ".visibility", 0)
+	
+for control in specialControls:
+	cmds.setAttr(control + ".visibility", 0)
+	
+#scale weapon and sword root controls
+weaponAnimControls = ["weapon_r_anim", "weapon_l_anim"]
+swordControls = ["sword_root_r_anim", "sword_root_l_anim"]
+
+for control in weaponAnimControls:
+	cmds.select(control + ".cv[0:32]")
+	cmds.scale(0, 1.6, 1.6, r = True)
+
+for control in swordControls:
+	cmds.select(control + ".cv[0:32]")
+	cmds.scale(0, 1.2, 1.2, r = True)
+	
+
+#ADD CUSTOM ATTRS FOR TRANSFORMATIONS
+
+#thrusters
+cmds.addAttr("body_anim", ln = "l_vent_open", dv = 0, min = 0, max = 10, keyable = True)
+cmds.addAttr("body_anim", ln = "r_vent_open", dv = 0, min = 0, max = 10, keyable = True)
+cmds.addAttr("body_anim", ln = "l_unit_deploy", dv = 0, min = 0, max = 10, keyable = True)
+cmds.addAttr("body_anim", ln = "r_unit_deploy", dv = 0, min = 0, max = 10, keyable = True)
+
+cmds.addAttr("body_anim", ln = "l_ventA_deploy", dv = 0, min = 0, max = 1, keyable = True)
+cmds.addAttr("body_anim", ln = "l_ventB_deploy", dv = 0, min = 0, max = 1, keyable = True)
+cmds.addAttr("body_anim", ln = "l_ventC_deploy", dv = 0, min = 0, max = 1, keyable = True)
+
+cmds.addAttr("body_anim", ln = "r_ventA_deploy", dv = 0, min = 0, max = 1, keyable = True)
+cmds.addAttr("body_anim", ln = "r_ventB_deploy", dv = 0, min = 0, max = 1, keyable = True)
+cmds.addAttr("body_anim", ln = "r_ventC_deploy", dv = 0, min = 0, max = 1, keyable = True)
+
+
+
+#left arm
+cmds.addAttr("weapon_l_anim", ln = "top_open", dv = 0, min = 0, max = 10, keyable = True)
+cmds.addAttr("weapon_l_anim", ln = "bottom_open", dv = 0, min = 0, max = 10, keyable = True)
+
+cmds.addAttr("weapon_l_anim", ln = "daggers_deploy", dv = 0, min = 0, max = 10, keyable = True)
+cmds.addAttr("weapon_l_anim", ln = "daggerA_swing", dv = 0, keyable = True)
+cmds.addAttr("weapon_l_anim", ln = "daggerB_swing", dv = 0, keyable = True)
+cmds.addAttr("weapon_l_anim", ln = "daggerC_swing", dv = 0, keyable = True)
+
+#right arm
+cmds.addAttr("weapon_r_anim", ln = "top_open", dv = 0, min = 0, max = 10, keyable = True)
+cmds.addAttr("weapon_r_anim", ln = "bottom_open", dv = 0, min = 0, max = 10, keyable = True)
+
+cmds.addAttr("weapon_r_anim", ln = "daggers_deploy", dv = 0, min = 0, max = 10, keyable = True)
+cmds.addAttr("weapon_r_anim", ln = "daggerA_swing", dv = 0, keyable = True)
+cmds.addAttr("weapon_r_anim", ln = "daggerB_swing", dv = 0, keyable = True)
+cmds.addAttr("weapon_r_anim", ln = "daggerC_swing", dv = 0, keyable = True)
+
+#left sword
+cmds.addAttr("sword_root_l_anim", ln = "sword_deploy", dv = 0, min = 0, max = 10, keyable = True)
+cmds.addAttr("sword_root_l_anim", ln = "sword_open", dv = 0, min = 0, max = 10, keyable = True)
+cmds.addAttr("sword_root_l_anim", ln = "sword_extend", dv = 0, min = 0, max = 10, keyable = True)
+
+#right sword
+cmds.addAttr("sword_root_r_anim", ln = "sword_deploy", dv = 0, min = 0, max = 10, keyable = True)
+cmds.addAttr("sword_root_r_anim", ln = "sword_open", dv = 0, min = 0, max = 10, keyable = True)
+cmds.addAttr("sword_root_r_anim", ln = "sword_extend", dv = 0, min = 0, max = 10, keyable = True)
+
+
+#SET DRIVEN KEYFRAMES
+
+#LEFT VENT OPEN
+cmds.setDrivenKeyframe(["fk_thrusterVent_l_01_lag_grp.tx", "fk_thrusterVent_l_01_lag_grp.ty", "fk_thrusterVent_l_01_lag_grp.tz", "fk_thrusterVent_l_01_lag_grp.rx", "fk_thrusterVent_l_01_lag_grp.ry", "fk_thrusterVent_l_01_lag_grp.rz"], cd = "body_anim.l_vent_open")
+cmds.setDrivenKeyframe(["fk_thrusterVent_l_02_lag_grp.tx", "fk_thrusterVent_l_02_lag_grp.ty", "fk_thrusterVent_l_02_lag_grp.tz", "fk_thrusterVent_l_02_lag_grp.rx", "fk_thrusterVent_l_02_lag_grp.ry", "fk_thrusterVent_l_02_lag_grp.rz"], cd = "body_anim.l_vent_open")
+
+cmds.setAttr("body_anim.l_vent_open", 5)
+cmds.setAttr("fk_thrusterVent_l_01_lag_grp.tx", 0*0.65)
+cmds.setAttr("fk_thrusterVent_l_01_lag_grp.ty", .152*0.65)
+cmds.setAttr("fk_thrusterVent_l_01_lag_grp.tz", -.09*0.65)
+cmds.setAttr("fk_thrusterVent_l_01_lag_grp.rx", 5.271)
+cmds.setAttr("fk_thrusterVent_l_01_lag_grp.ry", -.15)
+cmds.setAttr("fk_thrusterVent_l_01_lag_grp.rz", -19.4)
+
+cmds.setAttr("fk_thrusterVent_l_02_lag_grp.tx", -.653*0.65)
+cmds.setAttr("fk_thrusterVent_l_02_lag_grp.ty", 4.404*0.65)
+cmds.setAttr("fk_thrusterVent_l_02_lag_grp.tz", -.098*0.65)
+cmds.setAttr("fk_thrusterVent_l_02_lag_grp.rx", -4.033)
+cmds.setAttr("fk_thrusterVent_l_02_lag_grp.ry", -2.044)
+cmds.setAttr("fk_thrusterVent_l_02_lag_grp.rz", -19.602)
+cmds.setDrivenKeyframe(["fk_thrusterVent_l_01_lag_grp.tx", "fk_thrusterVent_l_01_lag_grp.ty", "fk_thrusterVent_l_01_lag_grp.tz", "fk_thrusterVent_l_01_lag_grp.rx", "fk_thrusterVent_l_01_lag_grp.ry", "fk_thrusterVent_l_01_lag_grp.rz"], cd = "body_anim.l_vent_open")
+cmds.setDrivenKeyframe(["fk_thrusterVent_l_02_lag_grp.tx", "fk_thrusterVent_l_02_lag_grp.ty", "fk_thrusterVent_l_02_lag_grp.tz", "fk_thrusterVent_l_02_lag_grp.rx", "fk_thrusterVent_l_02_lag_grp.ry", "fk_thrusterVent_l_02_lag_grp.rz"], cd = "body_anim.l_vent_open")
+
+
+cmds.setAttr("body_anim.l_vent_open", 10)
+cmds.setAttr("fk_thrusterVent_l_01_lag_grp.tx", 0*0.65)
+cmds.setAttr("fk_thrusterVent_l_01_lag_grp.ty", .152*0.65)
+cmds.setAttr("fk_thrusterVent_l_01_lag_grp.tz", -.09*0.65)
+cmds.setAttr("fk_thrusterVent_l_01_lag_grp.rx", 15.093)
+cmds.setAttr("fk_thrusterVent_l_01_lag_grp.ry", -.37)
+cmds.setAttr("fk_thrusterVent_l_01_lag_grp.rz", -55.082)
+
+cmds.setAttr("fk_thrusterVent_l_02_lag_grp.tx", -4.068*0.65)
+cmds.setAttr("fk_thrusterVent_l_02_lag_grp.ty", 5.659*0.65)
+cmds.setAttr("fk_thrusterVent_l_02_lag_grp.tz", .577*0.65)
+cmds.setAttr("fk_thrusterVent_l_02_lag_grp.rx", -10.697)
+cmds.setAttr("fk_thrusterVent_l_02_lag_grp.ry", -1.89)
+cmds.setAttr("fk_thrusterVent_l_02_lag_grp.rz", -15.521)
+cmds.setDrivenKeyframe(["fk_thrusterVent_l_01_lag_grp.tx", "fk_thrusterVent_l_01_lag_grp.ty", "fk_thrusterVent_l_01_lag_grp.tz", "fk_thrusterVent_l_01_lag_grp.rx", "fk_thrusterVent_l_01_lag_grp.ry", "fk_thrusterVent_l_01_lag_grp.rz"], cd = "body_anim.l_vent_open")
+cmds.setDrivenKeyframe(["fk_thrusterVent_l_02_lag_grp.tx", "fk_thrusterVent_l_02_lag_grp.ty", "fk_thrusterVent_l_02_lag_grp.tz", "fk_thrusterVent_l_02_lag_grp.rx", "fk_thrusterVent_l_02_lag_grp.ry", "fk_thrusterVent_l_02_lag_grp.rz"], cd = "body_anim.l_vent_open")
+
+cmds.setAttr("body_anim.l_vent_open", 0)
+
+
+
+#RIGHT VENT OPEN
+cmds.setDrivenKeyframe(["fk_thrusterVent_r_01_lag_grp.tx", "fk_thrusterVent_r_01_lag_grp.ty", "fk_thrusterVent_r_01_lag_grp.tz", "fk_thrusterVent_r_01_lag_grp.rx", "fk_thrusterVent_r_01_lag_grp.ry", "fk_thrusterVent_r_01_lag_grp.rz"], cd = "body_anim.r_vent_open")
+cmds.setDrivenKeyframe(["fk_thrusterVent_r_02_lag_grp.tx", "fk_thrusterVent_r_02_lag_grp.ty", "fk_thrusterVent_r_02_lag_grp.tz", "fk_thrusterVent_r_02_lag_grp.rx", "fk_thrusterVent_r_02_lag_grp.ry", "fk_thrusterVent_r_02_lag_grp.rz"], cd = "body_anim.r_vent_open")
+
+cmds.setAttr("body_anim.r_vent_open", 5)
+cmds.setAttr("fk_thrusterVent_r_01_lag_grp.tx", 0*0.65)
+cmds.setAttr("fk_thrusterVent_r_01_lag_grp.ty", -.152*0.65)
+cmds.setAttr("fk_thrusterVent_r_01_lag_grp.tz", .09*0.65)
+cmds.setAttr("fk_thrusterVent_r_01_lag_grp.rx", 5.271)
+cmds.setAttr("fk_thrusterVent_r_01_lag_grp.ry", -.15)
+cmds.setAttr("fk_thrusterVent_r_01_lag_grp.rz", -19.4)
+
+cmds.setAttr("fk_thrusterVent_r_02_lag_grp.tx", .653*0.65)
+cmds.setAttr("fk_thrusterVent_r_02_lag_grp.ty", -4.404*0.65)
+cmds.setAttr("fk_thrusterVent_r_02_lag_grp.tz", -.098*0.65)
+cmds.setAttr("fk_thrusterVent_r_02_lag_grp.rx", -4.033)
+cmds.setAttr("fk_thrusterVent_r_02_lag_grp.ry", -2.044)
+cmds.setAttr("fk_thrusterVent_r_02_lag_grp.rz", -19.602)
+cmds.setDrivenKeyframe(["fk_thrusterVent_r_01_lag_grp.tx", "fk_thrusterVent_r_01_lag_grp.ty", "fk_thrusterVent_r_01_lag_grp.tz", "fk_thrusterVent_r_01_lag_grp.rx", "fk_thrusterVent_r_01_lag_grp.ry", "fk_thrusterVent_r_01_lag_grp.rz"], cd = "body_anim.r_vent_open")
+cmds.setDrivenKeyframe(["fk_thrusterVent_r_02_lag_grp.tx", "fk_thrusterVent_r_02_lag_grp.ty", "fk_thrusterVent_r_02_lag_grp.tz", "fk_thrusterVent_r_02_lag_grp.rx", "fk_thrusterVent_r_02_lag_grp.ry", "fk_thrusterVent_r_02_lag_grp.rz"], cd = "body_anim.r_vent_open")
+
+cmds.setAttr("body_anim.r_vent_open", 10)
+cmds.setAttr("fk_thrusterVent_r_01_lag_grp.tx", 0*0.65)
+cmds.setAttr("fk_thrusterVent_r_01_lag_grp.ty", -.152*0.65)
+cmds.setAttr("fk_thrusterVent_r_01_lag_grp.tz", .09*0.65)
+cmds.setAttr("fk_thrusterVent_r_01_lag_grp.rx", 15.093)
+cmds.setAttr("fk_thrusterVent_r_01_lag_grp.ry", -.37)
+cmds.setAttr("fk_thrusterVent_r_01_lag_grp.rz", -55.082)
+
+cmds.setAttr("fk_thrusterVent_r_02_lag_grp.tx", 4.068*0.65)
+cmds.setAttr("fk_thrusterVent_r_02_lag_grp.ty", -5.659*0.65)
+cmds.setAttr("fk_thrusterVent_r_02_lag_grp.tz", -.577*0.65)
+cmds.setAttr("fk_thrusterVent_r_02_lag_grp.rx", -10.697)
+cmds.setAttr("fk_thrusterVent_r_02_lag_grp.ry", -1.89)
+cmds.setAttr("fk_thrusterVent_r_02_lag_grp.rz", -15.521)
+cmds.setDrivenKeyframe(["fk_thrusterVent_r_01_lag_grp.tx", "fk_thrusterVent_r_01_lag_grp.ty", "fk_thrusterVent_r_01_lag_grp.tz", "fk_thrusterVent_r_01_lag_grp.rx", "fk_thrusterVent_r_01_lag_grp.ry", "fk_thrusterVent_r_01_lag_grp.rz"], cd = "body_anim.r_vent_open")
+cmds.setDrivenKeyframe(["fk_thrusterVent_r_02_lag_grp.tx", "fk_thrusterVent_r_02_lag_grp.ty", "fk_thrusterVent_r_02_lag_grp.tz", "fk_thrusterVent_r_02_lag_grp.rx", "fk_thrusterVent_r_02_lag_grp.ry", "fk_thrusterVent_r_02_lag_grp.rz"], cd = "body_anim.r_vent_open")
+cmds.setAttr("body_anim.r_vent_open", 0)
+
+
+#LEFT UNIT DEPLOY
+cmds.setAttr("body_anim.l_unit_deploy", 0)
+cmds.setAttr("thruster_casing_l_anim_grp.tx", 0*0.65)
+cmds.setAttr("thruster_casing_l_anim_grp.ty", -9.046*0.65)
+cmds.setAttr("thruster_casing_l_anim_grp.tz", -1.065*0.65)
+cmds.setAttr("thruster_casing_l_anim_grp.rx", 95.625)
+cmds.setAttr("thruster_casing_l_anim_grp.ry", 0)
+cmds.setAttr("thruster_casing_l_anim_grp.rz", 0)
+cmds.setAttr("thruster_casing_l_anim_grp.sx", .398)
+cmds.setAttr("thruster_casing_l_anim_grp.sy", .398)
+cmds.setAttr("thruster_casing_l_anim_grp.sz", .398)
+cmds.setDrivenKeyframe(["thruster_casing_l_anim_grp.tx", "thruster_casing_l_anim_grp.ty", "thruster_casing_l_anim_grp.tz", "thruster_casing_l_anim_grp.rx", "thruster_casing_l_anim_grp.ry", "thruster_casing_l_anim_grp.rz", "thruster_casing_l_anim_grp.sx", "thruster_casing_l_anim_grp.sy", "thruster_casing_l_anim_grp.sz"], cd = "body_anim.l_unit_deploy",itt = "flat", ott = "flat")
+
+cmds.setAttr("body_anim.l_unit_deploy", 3.5)
+cmds.setAttr("thruster_casing_l_anim_grp.tx", 0*0.65)
+cmds.setAttr("thruster_casing_l_anim_grp.ty", .179*0.65)
+cmds.setAttr("thruster_casing_l_anim_grp.tz", -.157*0.65)
+cmds.setAttr("thruster_casing_l_anim_grp.rx", 95.625)
+cmds.setAttr("thruster_casing_l_anim_grp.ry", 0)
+cmds.setAttr("thruster_casing_l_anim_grp.rz", 0)
+cmds.setAttr("thruster_casing_l_anim_grp.sx", .8)
+cmds.setAttr("thruster_casing_l_anim_grp.sy", 1)
+cmds.setAttr("thruster_casing_l_anim_grp.sz", 1)
+cmds.setDrivenKeyframe(["thruster_casing_l_anim_grp.tx", "thruster_casing_l_anim_grp.ty", "thruster_casing_l_anim_grp.tz", "thruster_casing_l_anim_grp.rx", "thruster_casing_l_anim_grp.ry", "thruster_casing_l_anim_grp.rz", "thruster_casing_l_anim_grp.sx", "thruster_casing_l_anim_grp.sy", "thruster_casing_l_anim_grp.sz"], cd = "body_anim.l_unit_deploy",itt = "flat", ott = "flat")
+
+cmds.setAttr("body_anim.l_unit_deploy", 6.5)
+cmds.setAttr("thruster_casing_l_anim_grp.tx", 0*0.65)
+cmds.setAttr("thruster_casing_l_anim_grp.ty", .179*0.65)
+cmds.setAttr("thruster_casing_l_anim_grp.tz", -.157*0.65)
+cmds.setAttr("thruster_casing_l_anim_grp.rx", 53.119)
+cmds.setAttr("thruster_casing_l_anim_grp.ry", 0)
+cmds.setAttr("thruster_casing_l_anim_grp.rz", 0)
+cmds.setAttr("thruster_casing_l_anim_grp.sx", .933)
+cmds.setAttr("thruster_casing_l_anim_grp.sy", 1)
+cmds.setAttr("thruster_casing_l_anim_grp.sz", 1)
+cmds.setDrivenKeyframe(["thruster_casing_l_anim_grp.tx", "thruster_casing_l_anim_grp.ty", "thruster_casing_l_anim_grp.tz", "thruster_casing_l_anim_grp.rx", "thruster_casing_l_anim_grp.ry", "thruster_casing_l_anim_grp.rz", "thruster_casing_l_anim_grp.sx", "thruster_casing_l_anim_grp.sy", "thruster_casing_l_anim_grp.sz"], cd = "body_anim.l_unit_deploy",itt = "flat", ott = "flat")
+
+cmds.setAttr("body_anim.l_unit_deploy", 10)
+cmds.setAttr("thruster_casing_l_anim_grp.tx", 0*0.65)
+cmds.setAttr("thruster_casing_l_anim_grp.ty", .179*0.65)
+cmds.setAttr("thruster_casing_l_anim_grp.tz", -.157*0.65)
+cmds.setAttr("thruster_casing_l_anim_grp.rx", -.451)
+cmds.setAttr("thruster_casing_l_anim_grp.ry", 0)
+cmds.setAttr("thruster_casing_l_anim_grp.rz", 0)
+cmds.setAttr("thruster_casing_l_anim_grp.sx", 1)
+cmds.setAttr("thruster_casing_l_anim_grp.sy", 1)
+cmds.setAttr("thruster_casing_l_anim_grp.sz", 1)
+cmds.setDrivenKeyframe(["thruster_casing_l_anim_grp.tx", "thruster_casing_l_anim_grp.ty", "thruster_casing_l_anim_grp.tz", "thruster_casing_l_anim_grp.rx", "thruster_casing_l_anim_grp.ry", "thruster_casing_l_anim_grp.rz", "thruster_casing_l_anim_grp.sx", "thruster_casing_l_anim_grp.sy", "thruster_casing_l_anim_grp.sz"], cd = "body_anim.l_unit_deploy",itt = "flat", ott = "flat")
+cmds.setAttr("body_anim.l_unit_deploy", 0)
+
+
+#RIGHT UNIT DEPLOY
+cmds.setAttr("body_anim.r_unit_deploy", 0)
+cmds.setAttr("thruster_casing_r_anim_grp.tx", 0*0.65)
+cmds.setAttr("thruster_casing_r_anim_grp.ty", 9.046*0.65)
+cmds.setAttr("thruster_casing_r_anim_grp.tz", 1.065*0.65)
+cmds.setAttr("thruster_casing_r_anim_grp.rx", 95.625)
+cmds.setAttr("thruster_casing_r_anim_grp.ry", 0)
+cmds.setAttr("thruster_casing_r_anim_grp.rz", 0)
+cmds.setAttr("thruster_casing_r_anim_grp.sx", .398)
+cmds.setAttr("thruster_casing_r_anim_grp.sy", .398)
+cmds.setAttr("thruster_casing_r_anim_grp.sz", .398)
+cmds.setDrivenKeyframe(["thruster_casing_r_anim_grp.tx", "thruster_casing_r_anim_grp.ty", "thruster_casing_r_anim_grp.tz", "thruster_casing_r_anim_grp.rx", "thruster_casing_r_anim_grp.ry", "thruster_casing_r_anim_grp.rz", "thruster_casing_r_anim_grp.sx", "thruster_casing_r_anim_grp.sy", "thruster_casing_r_anim_grp.sz"], cd = "body_anim.r_unit_deploy",itt = "flat", ott = "flat")
+
+cmds.setAttr("body_anim.r_unit_deploy", 3.5)
+cmds.setAttr("thruster_casing_r_anim_grp.tx", 0*0.65)
+cmds.setAttr("thruster_casing_r_anim_grp.ty", -.179*0.65)
+cmds.setAttr("thruster_casing_r_anim_grp.tz", .157*0.65)
+cmds.setAttr("thruster_casing_r_anim_grp.rx", 95.625)
+cmds.setAttr("thruster_casing_r_anim_grp.ry", 0)
+cmds.setAttr("thruster_casing_r_anim_grp.rz", 0)
+cmds.setAttr("thruster_casing_r_anim_grp.sx", .8)
+cmds.setAttr("thruster_casing_r_anim_grp.sy", 1)
+cmds.setAttr("thruster_casing_r_anim_grp.sz", 1)
+cmds.setDrivenKeyframe(["thruster_casing_r_anim_grp.tx", "thruster_casing_r_anim_grp.ty", "thruster_casing_r_anim_grp.tz", "thruster_casing_r_anim_grp.rx", "thruster_casing_r_anim_grp.ry", "thruster_casing_r_anim_grp.rz", "thruster_casing_r_anim_grp.sx", "thruster_casing_r_anim_grp.sy", "thruster_casing_r_anim_grp.sz"], cd = "body_anim.r_unit_deploy",itt = "flat", ott = "flat")
+
+cmds.setAttr("body_anim.r_unit_deploy", 6.5)
+cmds.setAttr("thruster_casing_r_anim_grp.tx", 0*0.65)
+cmds.setAttr("thruster_casing_r_anim_grp.ty", -.179*0.65)
+cmds.setAttr("thruster_casing_r_anim_grp.tz", .157*0.65)
+cmds.setAttr("thruster_casing_r_anim_grp.rx", 53.119)
+cmds.setAttr("thruster_casing_r_anim_grp.ry", 0)
+cmds.setAttr("thruster_casing_r_anim_grp.rz", 0)
+cmds.setAttr("thruster_casing_r_anim_grp.sx", .933)
+cmds.setAttr("thruster_casing_r_anim_grp.sy", 1)
+cmds.setAttr("thruster_casing_r_anim_grp.sz", 1)
+cmds.setDrivenKeyframe(["thruster_casing_r_anim_grp.tx", "thruster_casing_r_anim_grp.ty", "thruster_casing_r_anim_grp.tz", "thruster_casing_r_anim_grp.rx", "thruster_casing_r_anim_grp.ry", "thruster_casing_r_anim_grp.rz", "thruster_casing_r_anim_grp.sx", "thruster_casing_r_anim_grp.sy", "thruster_casing_r_anim_grp.sz"], cd = "body_anim.r_unit_deploy",itt = "flat", ott = "flat")
+
+cmds.setAttr("body_anim.r_unit_deploy", 10)
+cmds.setAttr("thruster_casing_r_anim_grp.tx", 0*0.65)
+cmds.setAttr("thruster_casing_r_anim_grp.ty", -.179*0.65)
+cmds.setAttr("thruster_casing_r_anim_grp.tz", .157*0.65)
+cmds.setAttr("thruster_casing_r_anim_grp.rx", -.451)
+cmds.setAttr("thruster_casing_r_anim_grp.ry", 0)
+cmds.setAttr("thruster_casing_r_anim_grp.rz", 0)
+cmds.setAttr("thruster_casing_r_anim_grp.sx", 1)
+cmds.setAttr("thruster_casing_r_anim_grp.sy", 1)
+cmds.setAttr("thruster_casing_r_anim_grp.sz", 1)
+cmds.setDrivenKeyframe(["thruster_casing_r_anim_grp.tx", "thruster_casing_r_anim_grp.ty", "thruster_casing_r_anim_grp.tz", "thruster_casing_r_anim_grp.rx", "thruster_casing_r_anim_grp.ry", "thruster_casing_r_anim_grp.rz", "thruster_casing_r_anim_grp.sx", "thruster_casing_r_anim_grp.sy", "thruster_casing_r_anim_grp.sz"], cd = "body_anim.r_unit_deploy",itt = "flat", ott = "flat")
+cmds.setAttr("body_anim.r_unit_deploy", 0)
+
+
+
+#LEFT THRUSTER A 
+cmds.setAttr("body_anim.l_ventA_deploy", 0)
+cmds.setAttr("thruster_a_l_anim_grp.tx", -.038*0.65)
+cmds.setAttr("thruster_a_l_anim_grp.ty", -5.089*0.65)
+cmds.setAttr("thruster_a_l_anim_grp.tz", -0.79*0.65)
+cmds.setAttr("thruster_a_l_anim_grp.rx", -.405)
+cmds.setAttr("thruster_a_l_anim_grp.ry", -.105)
+cmds.setAttr("thruster_a_l_anim_grp.rz", -.167)
+cmds.setAttr("thruster_a_l_anim_grp.sx", 1)
+cmds.setAttr("thruster_a_l_anim_grp.sy", .434)
+cmds.setAttr("thruster_a_l_anim_grp.sz", 1)
+cmds.setDrivenKeyframe(["thruster_a_l_anim_grp.tx", "thruster_a_l_anim_grp.ty", "thruster_a_l_anim_grp.tz", "thruster_a_l_anim_grp.rx", "thruster_a_l_anim_grp.ry", "thruster_a_l_anim_grp.rz", "thruster_a_l_anim_grp.sx", "thruster_a_l_anim_grp.sy", "thruster_a_l_anim_grp.sz"], cd = "body_anim.l_ventA_deploy")
+
+cmds.setAttr("body_anim.l_ventA_deploy", 1)
+cmds.setAttr("thruster_a_l_anim_grp.tx", -.019*0.65)
+cmds.setAttr("thruster_a_l_anim_grp.ty", .261*0.65)
+cmds.setAttr("thruster_a_l_anim_grp.tz", -0.116*0.65)
+cmds.setAttr("thruster_a_l_anim_grp.rx", -.405)
+cmds.setAttr("thruster_a_l_anim_grp.ry", -.105)
+cmds.setAttr("thruster_a_l_anim_grp.rz", -.167)
+cmds.setAttr("thruster_a_l_anim_grp.sx", 1)
+cmds.setAttr("thruster_a_l_anim_grp.sy", 1)
+cmds.setAttr("thruster_a_l_anim_grp.sz", 1)
+cmds.setDrivenKeyframe(["thruster_a_l_anim_grp.tx", "thruster_a_l_anim_grp.ty", "thruster_a_l_anim_grp.tz", "thruster_a_l_anim_grp.rx", "thruster_a_l_anim_grp.ry", "thruster_a_l_anim_grp.rz", "thruster_a_l_anim_grp.sx", "thruster_a_l_anim_grp.sy", "thruster_a_l_anim_grp.sz"], cd = "body_anim.l_ventA_deploy")
+cmds.setAttr("body_anim.l_ventA_deploy", 0)
+
+
+#LEFT THRUSTER B 
+cmds.setAttr("body_anim.l_ventB_deploy", 0)
+cmds.setAttr("thruster_b_l_anim_grp.tx", -.038*0.65)
+cmds.setAttr("thruster_b_l_anim_grp.ty", -5.089*0.65)
+cmds.setAttr("thruster_b_l_anim_grp.tz", -0.79*0.65)
+cmds.setAttr("thruster_b_l_anim_grp.rx", -.405)
+cmds.setAttr("thruster_b_l_anim_grp.ry", -.105)
+cmds.setAttr("thruster_b_l_anim_grp.rz", -.167)
+cmds.setAttr("thruster_b_l_anim_grp.sx", 1)
+cmds.setAttr("thruster_b_l_anim_grp.sy", .434)
+cmds.setAttr("thruster_b_l_anim_grp.sz", 1)
+cmds.setDrivenKeyframe(["thruster_b_l_anim_grp.tx", "thruster_b_l_anim_grp.ty", "thruster_b_l_anim_grp.tz", "thruster_b_l_anim_grp.rx", "thruster_b_l_anim_grp.ry", "thruster_b_l_anim_grp.rz", "thruster_b_l_anim_grp.sx", "thruster_b_l_anim_grp.sy", "thruster_b_l_anim_grp.sz"], cd = "body_anim.l_ventB_deploy")
+
+cmds.setAttr("body_anim.l_ventB_deploy", 1)
+cmds.setAttr("thruster_b_l_anim_grp.tx", -.019*0.65)
+cmds.setAttr("thruster_b_l_anim_grp.ty", .261*0.65)
+cmds.setAttr("thruster_b_l_anim_grp.tz", -0.116*0.65)
+cmds.setAttr("thruster_b_l_anim_grp.rx", -.405)
+cmds.setAttr("thruster_b_l_anim_grp.ry", -.105)
+cmds.setAttr("thruster_b_l_anim_grp.rz", -.167)
+cmds.setAttr("thruster_b_l_anim_grp.sx", 1)
+cmds.setAttr("thruster_b_l_anim_grp.sy", 1)
+cmds.setAttr("thruster_b_l_anim_grp.sz", 1)
+cmds.setDrivenKeyframe(["thruster_b_l_anim_grp.tx", "thruster_b_l_anim_grp.ty", "thruster_b_l_anim_grp.tz", "thruster_b_l_anim_grp.rx", "thruster_b_l_anim_grp.ry", "thruster_b_l_anim_grp.rz", "thruster_b_l_anim_grp.sx", "thruster_b_l_anim_grp.sy", "thruster_b_l_anim_grp.sz"], cd = "body_anim.l_ventB_deploy")
+cmds.setAttr("body_anim.l_ventB_deploy", 0)
+
+#LEFT THRUSTER C 
+cmds.setAttr("body_anim.l_ventC_deploy", 0)
+cmds.setAttr("thruster_c_l_anim_grp.tx", -.038*0.65)
+cmds.setAttr("thruster_c_l_anim_grp.ty", -5.089*0.65)
+cmds.setAttr("thruster_c_l_anim_grp.tz", -0.79*0.65)
+cmds.setAttr("thruster_c_l_anim_grp.rx", -.405)
+cmds.setAttr("thruster_c_l_anim_grp.ry", -.105)
+cmds.setAttr("thruster_c_l_anim_grp.rz", -.167)
+cmds.setAttr("thruster_c_l_anim_grp.sx", 1)
+cmds.setAttr("thruster_c_l_anim_grp.sy", .434)
+cmds.setAttr("thruster_c_l_anim_grp.sz", 1)
+cmds.setDrivenKeyframe(["thruster_c_l_anim_grp.tx", "thruster_c_l_anim_grp.ty", "thruster_c_l_anim_grp.tz", "thruster_c_l_anim_grp.rx", "thruster_c_l_anim_grp.ry", "thruster_c_l_anim_grp.rz", "thruster_c_l_anim_grp.sx", "thruster_c_l_anim_grp.sy", "thruster_c_l_anim_grp.sz"], cd = "body_anim.l_ventC_deploy")
+
+cmds.setAttr("body_anim.l_ventC_deploy", 1)
+cmds.setAttr("thruster_c_l_anim_grp.tx", -.019*0.65)
+cmds.setAttr("thruster_c_l_anim_grp.ty", .261*0.65)
+cmds.setAttr("thruster_c_l_anim_grp.tz", -0.116*0.65)
+cmds.setAttr("thruster_c_l_anim_grp.rx", -.405)
+cmds.setAttr("thruster_c_l_anim_grp.ry", -.105)
+cmds.setAttr("thruster_c_l_anim_grp.rz", -.167)
+cmds.setAttr("thruster_c_l_anim_grp.sx", 1)
+cmds.setAttr("thruster_c_l_anim_grp.sy", 1)
+cmds.setAttr("thruster_c_l_anim_grp.sz", 1)
+cmds.setDrivenKeyframe(["thruster_c_l_anim_grp.tx", "thruster_c_l_anim_grp.ty", "thruster_c_l_anim_grp.tz", "thruster_c_l_anim_grp.rx", "thruster_c_l_anim_grp.ry", "thruster_c_l_anim_grp.rz", "thruster_c_l_anim_grp.sx", "thruster_c_l_anim_grp.sy", "thruster_c_l_anim_grp.sz"], cd = "body_anim.l_ventC_deploy")
+cmds.setAttr("body_anim.l_ventC_deploy", 0)
+
+
+
+
+#RIGHT THRUSTER A 
+cmds.setAttr("body_anim.r_ventA_deploy", 0)
+cmds.setAttr("thruster_a_r_anim_grp.tx", .038*0.65)
+cmds.setAttr("thruster_a_r_anim_grp.ty", 5.089*0.65)
+cmds.setAttr("thruster_a_r_anim_grp.tz", 0.79*0.65)
+cmds.setAttr("thruster_a_r_anim_grp.rx", -.405)
+cmds.setAttr("thruster_a_r_anim_grp.ry", -.105)
+cmds.setAttr("thruster_a_r_anim_grp.rz", -.167)
+cmds.setAttr("thruster_a_r_anim_grp.sx", 1)
+cmds.setAttr("thruster_a_r_anim_grp.sy", .434)
+cmds.setAttr("thruster_a_r_anim_grp.sz", 1)
+cmds.setDrivenKeyframe(["thruster_a_r_anim_grp.tx", "thruster_a_r_anim_grp.ty", "thruster_a_r_anim_grp.tz", "thruster_a_r_anim_grp.rx", "thruster_a_r_anim_grp.ry", "thruster_a_r_anim_grp.rz", "thruster_a_r_anim_grp.sx", "thruster_a_r_anim_grp.sy", "thruster_a_r_anim_grp.sz"], cd = "body_anim.r_ventA_deploy")
+
+cmds.setAttr("body_anim.r_ventA_deploy", 1)
+cmds.setAttr("thruster_a_r_anim_grp.tx", .019*0.65)
+cmds.setAttr("thruster_a_r_anim_grp.ty", -.261*0.65)
+cmds.setAttr("thruster_a_r_anim_grp.tz", 0.116*0.65)
+cmds.setAttr("thruster_a_r_anim_grp.rx", -.405)
+cmds.setAttr("thruster_a_r_anim_grp.ry", -.105)
+cmds.setAttr("thruster_a_r_anim_grp.rz", -.167)
+cmds.setAttr("thruster_a_r_anim_grp.sx", 1)
+cmds.setAttr("thruster_a_r_anim_grp.sy", 1)
+cmds.setAttr("thruster_a_r_anim_grp.sz", 1)
+cmds.setDrivenKeyframe(["thruster_a_r_anim_grp.tx", "thruster_a_r_anim_grp.ty", "thruster_a_r_anim_grp.tz", "thruster_a_r_anim_grp.rx", "thruster_a_r_anim_grp.ry", "thruster_a_r_anim_grp.rz", "thruster_a_r_anim_grp.sx", "thruster_a_r_anim_grp.sy", "thruster_a_r_anim_grp.sz"], cd = "body_anim.r_ventA_deploy")
+cmds.setAttr("body_anim.r_ventA_deploy", 0)
+
+
+#RIGHT THRUSTER B 
+cmds.setAttr("body_anim.r_ventB_deploy", 0)
+cmds.setAttr("thruster_b_r_anim_grp.tx", .038*0.65)
+cmds.setAttr("thruster_b_r_anim_grp.ty", 5.089*0.65)
+cmds.setAttr("thruster_b_r_anim_grp.tz", 0.79*0.65)
+cmds.setAttr("thruster_b_r_anim_grp.rx", -.405)
+cmds.setAttr("thruster_b_r_anim_grp.ry", -.105)
+cmds.setAttr("thruster_b_r_anim_grp.rz", -.167)
+cmds.setAttr("thruster_b_r_anim_grp.sx", 1)
+cmds.setAttr("thruster_b_r_anim_grp.sy", .434)
+cmds.setAttr("thruster_b_r_anim_grp.sz", 1)
+cmds.setDrivenKeyframe(["thruster_b_r_anim_grp.tx", "thruster_b_r_anim_grp.ty", "thruster_b_r_anim_grp.tz", "thruster_b_r_anim_grp.rx", "thruster_b_r_anim_grp.ry", "thruster_b_r_anim_grp.rz", "thruster_b_r_anim_grp.sx", "thruster_b_r_anim_grp.sy", "thruster_b_r_anim_grp.sz"], cd = "body_anim.r_ventB_deploy")
+
+cmds.setAttr("body_anim.r_ventB_deploy", 1)
+cmds.setAttr("thruster_b_r_anim_grp.tx", .019*0.65)
+cmds.setAttr("thruster_b_r_anim_grp.ty", -.261*0.65)
+cmds.setAttr("thruster_b_r_anim_grp.tz", 0.116*0.65)
+cmds.setAttr("thruster_b_r_anim_grp.rx", -.405)
+cmds.setAttr("thruster_b_r_anim_grp.ry", -.105)
+cmds.setAttr("thruster_b_r_anim_grp.rz", -.167)
+cmds.setAttr("thruster_b_r_anim_grp.sx", 1)
+cmds.setAttr("thruster_b_r_anim_grp.sy", 1)
+cmds.setAttr("thruster_b_r_anim_grp.sz", 1)
+cmds.setDrivenKeyframe(["thruster_b_r_anim_grp.tx", "thruster_b_r_anim_grp.ty", "thruster_b_r_anim_grp.tz", "thruster_b_r_anim_grp.rx", "thruster_b_r_anim_grp.ry", "thruster_b_r_anim_grp.rz", "thruster_b_r_anim_grp.sx", "thruster_b_r_anim_grp.sy", "thruster_b_r_anim_grp.sz"], cd = "body_anim.r_ventB_deploy")
+cmds.setAttr("body_anim.r_ventB_deploy", 0)
+
+#RIGHT THRUSTER C 
+cmds.setAttr("body_anim.r_ventC_deploy", 0)
+cmds.setAttr("thruster_c_r_anim_grp.tx", .038*0.65)
+cmds.setAttr("thruster_c_r_anim_grp.ty", 5.089*0.65)
+cmds.setAttr("thruster_c_r_anim_grp.tz", 0.79*0.65)
+cmds.setAttr("thruster_c_r_anim_grp.rx", -.405)
+cmds.setAttr("thruster_c_r_anim_grp.ry", -.105)
+cmds.setAttr("thruster_c_r_anim_grp.rz", -.167)
+cmds.setAttr("thruster_c_r_anim_grp.sx", 1)
+cmds.setAttr("thruster_c_r_anim_grp.sy", .434)
+cmds.setAttr("thruster_c_r_anim_grp.sz", 1)
+cmds.setDrivenKeyframe(["thruster_c_r_anim_grp.tx", "thruster_c_r_anim_grp.ty", "thruster_c_r_anim_grp.tz", "thruster_c_r_anim_grp.rx", "thruster_c_r_anim_grp.ry", "thruster_c_r_anim_grp.rz", "thruster_c_r_anim_grp.sx", "thruster_c_r_anim_grp.sy", "thruster_c_r_anim_grp.sz"], cd = "body_anim.r_ventC_deploy")
+
+cmds.setAttr("body_anim.r_ventC_deploy", 1)
+cmds.setAttr("thruster_c_r_anim_grp.tx", .019*0.65)
+cmds.setAttr("thruster_c_r_anim_grp.ty", -.261*0.65)
+cmds.setAttr("thruster_c_r_anim_grp.tz", 0.116*0.65)
+cmds.setAttr("thruster_c_r_anim_grp.rx", -.405)
+cmds.setAttr("thruster_c_r_anim_grp.ry", -.105)
+cmds.setAttr("thruster_c_r_anim_grp.rz", -.167)
+cmds.setAttr("thruster_c_r_anim_grp.sx", 1)
+cmds.setAttr("thruster_c_r_anim_grp.sy", 1)
+cmds.setAttr("thruster_c_r_anim_grp.sz", 1)
+cmds.setDrivenKeyframe(["thruster_c_r_anim_grp.tx", "thruster_c_r_anim_grp.ty", "thruster_c_r_anim_grp.tz", "thruster_c_r_anim_grp.rx", "thruster_c_r_anim_grp.ry", "thruster_c_r_anim_grp.rz", "thruster_c_r_anim_grp.sx", "thruster_c_r_anim_grp.sy", "thruster_c_r_anim_grp.sz"], cd = "body_anim.r_ventC_deploy")
+cmds.setAttr("body_anim.r_ventC_deploy", 0)
+
+
+
+#LEFT ARM TOP DEPLOY
+cmds.setAttr("weapon_l_anim.top_open", 0)
+cmds.setDrivenKeyframe(["lowerArm_topPlateA_l_anim_grp.tx", "lowerArm_topPlateA_l_anim_grp.ty", "lowerArm_topPlateA_l_anim_grp.tz", "lowerArm_topPlateA_l_anim_grp.rx", "lowerArm_topPlateA_l_anim_grp.ry", "lowerArm_topPlateA_l_anim_grp.rz"], cd = "weapon_l_anim.top_open",itt = "flat", ott = "flat")
+cmds.setDrivenKeyframe(["lowerArm_topPlateB_l_anim_grp.tx", "lowerArm_topPlateB_l_anim_grp.ty", "lowerArm_topPlateB_l_anim_grp.tz", "lowerArm_topPlateB_l_anim_grp.rx", "lowerArm_topPlateB_l_anim_grp.ry", "lowerArm_topPlateB_l_anim_grp.rz"], cd = "weapon_l_anim.top_open",itt = "flat", ott = "flat")
+
+cmds.setAttr("weapon_l_anim.top_open", 5)
+cmds.setAttr("lowerArm_topPlateA_l_anim_grp.ty", -1.408*0.65)
+cmds.setAttr("lowerArm_topPlateB_l_anim_grp.ty", -1.408*0.65)
+cmds.setDrivenKeyframe(["lowerArm_topPlateA_l_anim_grp.tx", "lowerArm_topPlateA_l_anim_grp.ty", "lowerArm_topPlateA_l_anim_grp.tz", "lowerArm_topPlateA_l_anim_grp.rx", "lowerArm_topPlateA_l_anim_grp.ry", "lowerArm_topPlateA_l_anim_grp.rz"], cd = "weapon_l_anim.top_open",itt = "flat", ott = "flat")
+cmds.setDrivenKeyframe(["lowerArm_topPlateB_l_anim_grp.tx", "lowerArm_topPlateB_l_anim_grp.ty", "lowerArm_topPlateB_l_anim_grp.tz", "lowerArm_topPlateB_l_anim_grp.rx", "lowerArm_topPlateB_l_anim_grp.ry", "lowerArm_topPlateB_l_anim_grp.rz"], cd = "weapon_l_anim.top_open",itt = "flat", ott = "flat")
+
+cmds.setAttr("weapon_l_anim.top_open", 10)
+cmds.setAttr("lowerArm_topPlateA_l_anim_grp.tx", 4.631*0.65)
+cmds.setAttr("lowerArm_topPlateB_l_anim_grp.tx", 4.631*0.65)
+cmds.setAttr("lowerArm_topPlateA_l_anim_grp.ty", -1.408*0.65)
+cmds.setAttr("lowerArm_topPlateB_l_anim_grp.ty", -1.408*0.65)
+cmds.setAttr("lowerArm_topPlateA_l_anim_grp.rz", -30.5)
+cmds.setAttr("lowerArm_topPlateB_l_anim_grp.rz", -30.5)
+cmds.setDrivenKeyframe(["lowerArm_topPlateA_l_anim_grp.tx", "lowerArm_topPlateA_l_anim_grp.ty", "lowerArm_topPlateA_l_anim_grp.tz", "lowerArm_topPlateA_l_anim_grp.rx", "lowerArm_topPlateA_l_anim_grp.ry", "lowerArm_topPlateA_l_anim_grp.rz"], cd = "weapon_l_anim.top_open",itt = "flat", ott = "flat")
+cmds.setDrivenKeyframe(["lowerArm_topPlateB_l_anim_grp.tx", "lowerArm_topPlateB_l_anim_grp.ty", "lowerArm_topPlateB_l_anim_grp.tz", "lowerArm_topPlateB_l_anim_grp.rx", "lowerArm_topPlateB_l_anim_grp.ry", "lowerArm_topPlateB_l_anim_grp.rz"], cd = "weapon_l_anim.top_open",itt = "flat", ott = "flat")
+cmds.setAttr("weapon_l_anim.top_open", 0)
+
+
+#LEFT ARM BOTTOM DEPLOY
+cmds.setAttr("weapon_l_anim.bottom_open", 0)
+cmds.setDrivenKeyframe(["lowerArm_bottomPlateA_l_anim_grp.tx", "lowerArm_bottomPlateA_l_anim_grp.ty", "lowerArm_bottomPlateA_l_anim_grp.tz", "lowerArm_bottomPlateA_l_anim_grp.rx", "lowerArm_bottomPlateA_l_anim_grp.ry", "lowerArm_bottomPlateA_l_anim_grp.rz"], cd = "weapon_l_anim.bottom_open",itt = "flat", ott = "flat")
+cmds.setDrivenKeyframe(["lowerArm_bottomPlateB_l_anim_grp.tx", "lowerArm_bottomPlateB_l_anim_grp.ty", "lowerArm_bottomPlateB_l_anim_grp.tz", "lowerArm_bottomPlateB_l_anim_grp.rx", "lowerArm_bottomPlateB_l_anim_grp.ry", "lowerArm_bottomPlateB_l_anim_grp.rz"], cd = "weapon_l_anim.bottom_open",itt = "flat", ott = "flat")
+
+cmds.setAttr("weapon_l_anim.bottom_open", 5)
+cmds.setAttr("lowerArm_bottomPlateA_l_anim_grp.ty", -2*0.65)
+cmds.setAttr("lowerArm_bottomPlateB_l_anim_grp.ty", -2*0.65)
+cmds.setDrivenKeyframe(["lowerArm_bottomPlateA_l_anim_grp.tx", "lowerArm_bottomPlateA_l_anim_grp.ty", "lowerArm_bottomPlateA_l_anim_grp.tz", "lowerArm_bottomPlateA_l_anim_grp.rx", "lowerArm_bottomPlateA_l_anim_grp.ry", "lowerArm_bottomPlateA_l_anim_grp.rz"], cd = "weapon_l_anim.bottom_open",itt = "flat", ott = "flat")
+cmds.setDrivenKeyframe(["lowerArm_bottomPlateB_l_anim_grp.tx", "lowerArm_bottomPlateB_l_anim_grp.ty", "lowerArm_bottomPlateB_l_anim_grp.tz", "lowerArm_bottomPlateB_l_anim_grp.rx", "lowerArm_bottomPlateB_l_anim_grp.ry", "lowerArm_bottomPlateB_l_anim_grp.rz"], cd = "weapon_l_anim.bottom_open",itt = "flat", ott = "flat")
+
+cmds.setAttr("weapon_l_anim.bottom_open", 10)
+cmds.setAttr("lowerArm_bottomPlateA_l_anim_grp.tx", 5.34*0.65)
+cmds.setAttr("lowerArm_bottomPlateB_l_anim_grp.tx", 5.34*0.65)
+cmds.setAttr("lowerArm_bottomPlateA_l_anim_grp.ty", -2*0.65)
+cmds.setAttr("lowerArm_bottomPlateB_l_anim_grp.ty", -2*0.65)
+cmds.setAttr("lowerArm_bottomPlateA_l_anim_grp.rx", -12.291)
+cmds.setAttr("lowerArm_bottomPlateB_l_anim_grp.rx", 12.291)
+cmds.setAttr("lowerArm_bottomPlateA_l_anim_grp.ry", 7.644)
+cmds.setAttr("lowerArm_bottomPlateB_l_anim_grp.ry", -7.644)
+cmds.setAttr("lowerArm_bottomPlateA_l_anim_grp.rz", -63.777)
+cmds.setAttr("lowerArm_bottomPlateB_l_anim_grp.rz", -63.777)
+cmds.setDrivenKeyframe(["lowerArm_bottomPlateA_l_anim_grp.tx", "lowerArm_bottomPlateA_l_anim_grp.ty", "lowerArm_bottomPlateA_l_anim_grp.tz", "lowerArm_bottomPlateA_l_anim_grp.rx", "lowerArm_bottomPlateA_l_anim_grp.ry", "lowerArm_bottomPlateA_l_anim_grp.rz"], cd = "weapon_l_anim.bottom_open",itt = "flat", ott = "flat")
+cmds.setDrivenKeyframe(["lowerArm_bottomPlateB_l_anim_grp.tx", "lowerArm_bottomPlateB_l_anim_grp.ty", "lowerArm_bottomPlateB_l_anim_grp.tz", "lowerArm_bottomPlateB_l_anim_grp.rx", "lowerArm_bottomPlateB_l_anim_grp.ry", "lowerArm_bottomPlateB_l_anim_grp.rz"], cd = "weapon_l_anim.bottom_open",itt = "flat", ott = "flat")
+cmds.setAttr("weapon_l_anim.bottom_open", 0)
+
+
+
+
+#RIGHT ARM TOP DEPLOY
+cmds.setAttr("weapon_r_anim.top_open", 0)
+cmds.setDrivenKeyframe(["lowerArm_topPlateA_r_anim_grp.tx", "lowerArm_topPlateA_r_anim_grp.ty", "lowerArm_topPlateA_r_anim_grp.tz", "lowerArm_topPlateA_r_anim_grp.rx", "lowerArm_topPlateA_r_anim_grp.ry", "lowerArm_topPlateA_r_anim_grp.rz"], cd = "weapon_r_anim.top_open",itt = "flat", ott = "flat")
+cmds.setDrivenKeyframe(["lowerArm_topPlateB_r_anim_grp.tx", "lowerArm_topPlateB_r_anim_grp.ty", "lowerArm_topPlateB_r_anim_grp.tz", "lowerArm_topPlateB_r_anim_grp.rx", "lowerArm_topPlateB_r_anim_grp.ry", "lowerArm_topPlateB_r_anim_grp.rz"], cd = "weapon_r_anim.top_open",itt = "flat", ott = "flat")
+
+cmds.setAttr("weapon_r_anim.top_open", 5)
+cmds.setAttr("lowerArm_topPlateA_r_anim_grp.ty", 1.408*0.65)
+cmds.setAttr("lowerArm_topPlateB_r_anim_grp.ty", 1.408*0.65)
+cmds.setDrivenKeyframe(["lowerArm_topPlateA_r_anim_grp.tx", "lowerArm_topPlateA_r_anim_grp.ty", "lowerArm_topPlateA_r_anim_grp.tz", "lowerArm_topPlateA_r_anim_grp.rx", "lowerArm_topPlateA_r_anim_grp.ry", "lowerArm_topPlateA_r_anim_grp.rz"], cd = "weapon_r_anim.top_open",itt = "flat", ott = "flat")
+cmds.setDrivenKeyframe(["lowerArm_topPlateB_r_anim_grp.tx", "lowerArm_topPlateB_r_anim_grp.ty", "lowerArm_topPlateB_r_anim_grp.tz", "lowerArm_topPlateB_r_anim_grp.rx", "lowerArm_topPlateB_r_anim_grp.ry", "lowerArm_topPlateB_r_anim_grp.rz"], cd = "weapon_r_anim.top_open",itt = "flat", ott = "flat")
+
+cmds.setAttr("weapon_r_anim.top_open", 10)
+cmds.setAttr("lowerArm_topPlateA_r_anim_grp.tx", -4.631*0.65)
+cmds.setAttr("lowerArm_topPlateB_r_anim_grp.tx", -4.631*0.65)
+cmds.setAttr("lowerArm_topPlateA_r_anim_grp.ty", 1.408*0.65)
+cmds.setAttr("lowerArm_topPlateB_r_anim_grp.ty", 1.408*0.65)
+cmds.setAttr("lowerArm_topPlateA_r_anim_grp.rz", -30.5)
+cmds.setAttr("lowerArm_topPlateB_r_anim_grp.rz", -30.5)
+cmds.setDrivenKeyframe(["lowerArm_topPlateA_r_anim_grp.tx", "lowerArm_topPlateA_r_anim_grp.ty", "lowerArm_topPlateA_r_anim_grp.tz", "lowerArm_topPlateA_r_anim_grp.rx", "lowerArm_topPlateA_r_anim_grp.ry", "lowerArm_topPlateA_r_anim_grp.rz"], cd = "weapon_r_anim.top_open",itt = "flat", ott = "flat")
+cmds.setDrivenKeyframe(["lowerArm_topPlateB_r_anim_grp.tx", "lowerArm_topPlateB_r_anim_grp.ty", "lowerArm_topPlateB_r_anim_grp.tz", "lowerArm_topPlateB_r_anim_grp.rx", "lowerArm_topPlateB_r_anim_grp.ry", "lowerArm_topPlateB_r_anim_grp.rz"], cd = "weapon_r_anim.top_open",itt = "flat", ott = "flat")
+cmds.setAttr("weapon_r_anim.top_open", 0)
+
+
+#RIGHT ARM BOTTOM DEPLOY
+cmds.setAttr("weapon_r_anim.bottom_open", 0)
+cmds.setDrivenKeyframe(["lowerArm_bottomPlateA_r_anim_grp.tx", "lowerArm_bottomPlateA_r_anim_grp.ty", "lowerArm_bottomPlateA_r_anim_grp.tz", "lowerArm_bottomPlateA_r_anim_grp.rx", "lowerArm_bottomPlateA_r_anim_grp.ry", "lowerArm_bottomPlateA_r_anim_grp.rz"], cd = "weapon_r_anim.bottom_open",itt = "flat", ott = "flat")
+cmds.setDrivenKeyframe(["lowerArm_bottomPlateB_r_anim_grp.tx", "lowerArm_bottomPlateB_r_anim_grp.ty", "lowerArm_bottomPlateB_r_anim_grp.tz", "lowerArm_bottomPlateB_r_anim_grp.rx", "lowerArm_bottomPlateB_r_anim_grp.ry", "lowerArm_bottomPlateB_r_anim_grp.rz"], cd = "weapon_r_anim.bottom_open",itt = "flat", ott = "flat")
+
+cmds.setAttr("weapon_r_anim.bottom_open", 5)
+cmds.setAttr("lowerArm_bottomPlateA_r_anim_grp.ty", 2*0.65)
+cmds.setAttr("lowerArm_bottomPlateB_r_anim_grp.ty", 2*0.65)
+cmds.setDrivenKeyframe(["lowerArm_bottomPlateA_r_anim_grp.tx", "lowerArm_bottomPlateA_r_anim_grp.ty", "lowerArm_bottomPlateA_r_anim_grp.tz", "lowerArm_bottomPlateA_r_anim_grp.rx", "lowerArm_bottomPlateA_r_anim_grp.ry", "lowerArm_bottomPlateA_r_anim_grp.rz"], cd = "weapon_r_anim.bottom_open",itt = "flat", ott = "flat")
+cmds.setDrivenKeyframe(["lowerArm_bottomPlateB_r_anim_grp.tx", "lowerArm_bottomPlateB_r_anim_grp.ty", "lowerArm_bottomPlateB_r_anim_grp.tz", "lowerArm_bottomPlateB_r_anim_grp.rx", "lowerArm_bottomPlateB_r_anim_grp.ry", "lowerArm_bottomPlateB_r_anim_grp.rz"], cd = "weapon_r_anim.bottom_open",itt = "flat", ott = "flat")
+
+cmds.setAttr("weapon_r_anim.bottom_open", 10)
+cmds.setAttr("lowerArm_bottomPlateA_r_anim_grp.tx", -5.34*0.65)
+cmds.setAttr("lowerArm_bottomPlateB_r_anim_grp.tx", -5.34*0.65)
+cmds.setAttr("lowerArm_bottomPlateA_r_anim_grp.ty", 2*0.65)
+cmds.setAttr("lowerArm_bottomPlateB_r_anim_grp.ty", 2*0.65)
+cmds.setAttr("lowerArm_bottomPlateA_r_anim_grp.rx", -12.291)
+cmds.setAttr("lowerArm_bottomPlateB_r_anim_grp.rx", 12.291)
+cmds.setAttr("lowerArm_bottomPlateA_r_anim_grp.ry", 7.644)
+cmds.setAttr("lowerArm_bottomPlateB_r_anim_grp.ry", -7.644)
+cmds.setAttr("lowerArm_bottomPlateA_r_anim_grp.rz", -63.777)
+cmds.setAttr("lowerArm_bottomPlateB_r_anim_grp.rz", -63.777)
+cmds.setDrivenKeyframe(["lowerArm_bottomPlateA_r_anim_grp.tx", "lowerArm_bottomPlateA_r_anim_grp.ty", "lowerArm_bottomPlateA_r_anim_grp.tz", "lowerArm_bottomPlateA_r_anim_grp.rx", "lowerArm_bottomPlateA_r_anim_grp.ry", "lowerArm_bottomPlateA_r_anim_grp.rz"], cd = "weapon_r_anim.bottom_open",itt = "flat", ott = "flat")
+cmds.setDrivenKeyframe(["lowerArm_bottomPlateB_r_anim_grp.tx", "lowerArm_bottomPlateB_r_anim_grp.ty", "lowerArm_bottomPlateB_r_anim_grp.tz", "lowerArm_bottomPlateB_r_anim_grp.rx", "lowerArm_bottomPlateB_r_anim_grp.ry", "lowerArm_bottomPlateB_r_anim_grp.rz"], cd = "weapon_r_anim.bottom_open",itt = "flat", ott = "flat")
+cmds.setAttr("weapon_r_anim.bottom_open", 0)
+
+
+#LEFT DAGGERS
+cmds.setAttr("weapon_l_anim.daggers_deploy", 0)
+cmds.setDrivenKeyframe("dagger_base_l_anim_grp.tz", cd = "weapon_l_anim.daggers_deploy")
+
+cmds.setAttr("weapon_l_anim.daggers_deploy", 10)
+cmds.setAttr("dagger_base_l_anim_grp.tz", -30*0.65)
+cmds.setDrivenKeyframe("dagger_base_l_anim_grp.tz", cd = "weapon_l_anim.daggers_deploy")
+cmds.setAttr("weapon_l_anim.daggers_deploy", 0)
+
+cmds.connectAttr("weapon_l_anim.daggerA_swing", "dagger_a_l_anim_grp.rx")
+cmds.connectAttr("weapon_l_anim.daggerB_swing", "dagger_b_l_anim_grp.rx")
+cmds.connectAttr("weapon_l_anim.daggerC_swing", "dagger_c_l_anim_grp.rx")
+
+
+#RIGHT DAGGERS
+cmds.setAttr("weapon_r_anim.daggers_deploy", 0)
+cmds.setDrivenKeyframe("dagger_base_r_anim_grp.tz", cd = "weapon_r_anim.daggers_deploy")
+
+cmds.setAttr("weapon_r_anim.daggers_deploy", 10)
+cmds.setAttr("dagger_base_r_anim_grp.tz", 30*0.65)
+cmds.setDrivenKeyframe("dagger_base_r_anim_grp.tz", cd = "weapon_r_anim.daggers_deploy")
+cmds.setAttr("weapon_r_anim.daggers_deploy", 0)
+
+cmds.connectAttr("weapon_r_anim.daggerA_swing", "dagger_a_r_anim_grp.rx")
+cmds.connectAttr("weapon_r_anim.daggerB_swing", "dagger_b_r_anim_grp.rx")
+cmds.connectAttr("weapon_r_anim.daggerC_swing", "dagger_c_r_anim_grp.rx")
+
+
+#LEFT SWORD DEPLOY
+cmds.setAttr("sword_root_l_anim.sword_deploy", 0)
+cmds.setDrivenKeyframe("sword_root_l_anim_grp.tx", cd = "sword_root_l_anim.sword_deploy")
+
+cmds.setAttr("sword_root_l_anim.sword_deploy", 10)
+cmds.setAttr("sword_root_l_anim_grp.tx", 43*0.65)
+cmds.setDrivenKeyframe("sword_root_l_anim_grp.tx", cd = "sword_root_l_anim.sword_deploy")
+
+cmds.setAttr("sword_root_l_anim.sword_deploy", 0)
+
+
+#LEFT SWORD OPEN
+cmds.setAttr("sword_root_l_anim.sword_open", 0)
+cmds.setDrivenKeyframe(["pivotA_l_anim_grp.ry", "pivotB_l_anim_grp.ry", "pivotA_l_anim_grp.rz", "pivotB_l_anim_grp.rz"], cd = "sword_root_l_anim.sword_open")
+
+cmds.setAttr("sword_root_l_anim.sword_open", 10)
+cmds.setAttr("pivotA_l_anim_grp.ry", 179.7)
+cmds.setAttr("pivotB_l_anim_grp.ry", 179.7)
+cmds.setAttr("pivotA_l_anim_grp.rz", 0.1)
+cmds.setAttr("pivotB_l_anim_grp.rz", 0.1)
+cmds.setDrivenKeyframe(["pivotA_l_anim_grp.ry", "pivotB_l_anim_grp.ry", "pivotA_l_anim_grp.rz", "pivotB_l_anim_grp.rz"], cd = "sword_root_l_anim.sword_open")
+
+cmds.setAttr("sword_root_l_anim.sword_open", 0)
+
+#LEFT SWORD EXTEND
+cmds.setAttr("sword_root_l_anim.sword_extend", 0)
+cmds.setDrivenKeyframe(["extensionA_l_anim_grp.tx", "extensionA_l_anim_grp.ty", "extensionB_l_anim_grp.tx", "extensionB_l_anim_grp.ty"], cd = "sword_root_l_anim.sword_extend")
+
+cmds.setAttr("sword_root_l_anim.sword_extend", 10)
+cmds.setAttr("extensionA_l_anim_grp.tx", -16.469)
+cmds.setAttr("extensionA_l_anim_grp.ty", 0.0738)
+cmds.setAttr("extensionB_l_anim_grp.tx", -16.469)
+cmds.setAttr("extensionB_l_anim_grp.ty", -0.0794)
+cmds.setDrivenKeyframe(["extensionA_l_anim_grp.tx", "extensionA_l_anim_grp.ty", "extensionB_l_anim_grp.tx", "extensionB_l_anim_grp.ty"], cd = "sword_root_l_anim.sword_extend")
+
+cmds.setAttr("sword_root_l_anim.sword_extend", 0)
+
+
+
+#RIGHT SWORD DEPLOY
+cmds.setAttr("sword_root_r_anim.sword_deploy", 0)
+cmds.setDrivenKeyframe("sword_root_r_anim_grp.tx", cd = "sword_root_r_anim.sword_deploy")
+
+cmds.setAttr("sword_root_r_anim.sword_deploy", 10)
+cmds.setAttr("sword_root_r_anim_grp.tx", -43*0.65)
+cmds.setDrivenKeyframe("sword_root_r_anim_grp.tx", cd = "sword_root_r_anim.sword_deploy")
+
+cmds.setAttr("sword_root_r_anim.sword_deploy", 0)
+
+
+#RIGHT SWORD OPEN
+cmds.setAttr("sword_root_r_anim.sword_open", 0)
+cmds.setDrivenKeyframe(["pivotA_r_anim_grp.ry", "pivotB_r_anim_grp.ry", "pivotA_r_anim_grp.rz", "pivotB_r_anim_grp.rz"], cd = "sword_root_r_anim.sword_open")
+
+cmds.setAttr("sword_root_r_anim.sword_open", 10)
+cmds.setAttr("pivotA_r_anim_grp.ry", 179.7)
+cmds.setAttr("pivotB_r_anim_grp.ry", 179.7)
+cmds.setAttr("pivotA_r_anim_grp.rz", 0.1)
+cmds.setAttr("pivotB_r_anim_grp.rz", 0.1)
+cmds.setDrivenKeyframe(["pivotA_r_anim_grp.ry", "pivotB_r_anim_grp.ry", "pivotA_r_anim_grp.rz", "pivotB_r_anim_grp.rz"], cd = "sword_root_r_anim.sword_open")
+
+cmds.setAttr("sword_root_r_anim.sword_open", 0)
+
+#RIGHT SWORD EXTEND
+cmds.setAttr("sword_root_r_anim.sword_extend", 0)
+cmds.setDrivenKeyframe(["extensionA_r_anim_grp.tx", "extensionA_r_anim_grp.ty", "extensionB_r_anim_grp.tx", "extensionB_r_anim_grp.ty"], cd = "sword_root_r_anim.sword_extend")
+
+cmds.setAttr("sword_root_r_anim.sword_extend", 10)
+cmds.setAttr("extensionA_r_anim_grp.tx", 16.469)
+cmds.setAttr("extensionA_r_anim_grp.ty", 0.101)
+cmds.setAttr("extensionB_r_anim_grp.tx", 16.469)
+cmds.setAttr("extensionB_r_anim_grp.ty", -0.0965)
+cmds.setDrivenKeyframe(["extensionA_r_anim_grp.tx", "extensionA_r_anim_grp.ty", "extensionB_r_anim_grp.tx", "extensionB_r_anim_grp.ty"], cd = "sword_root_r_anim.sword_extend")
+
+cmds.setAttr("sword_root_r_anim.sword_extend", 0)
+
+
+# MAKE THE GRP NODES SCALE DRIVE THE CONTROLS SCALE
+for control in controlsToHide:
+	grp = control + "_grp"
+	cmds.connectAttr(grp+".sx", control+".sx", f=True)
+	cmds.connectAttr(grp+".sy", control+".sy", f=True)
+	cmds.connectAttr(grp+".sz", control+".sz", f=True)
+	cmds.setAttr(control + ".sx", l=True)
+	cmds.setAttr(control + ".sy", l=True)
+	cmds.setAttr(control + ".sz", l=True)
+    
+
+
+#CREATE SPACES FOR WEAPONS ANIMS
+pairs = [["weapon_l_anim", "driver_hand_l"], ["weapon_r_anim", "driver_hand_r"]]
+
+for pair in pairs:
+    control = pair[0]
+    spaceObj = pair[1]
+
+    spaceSwitchFollow = control + "_space_switcher_follow"
+    spaceSwitchNode = control + "_space_switcher"
+
+    #add attr to the space switcher node
+    cmds.select(spaceSwitchNode)
+    cmds.addAttr(ln = "space_" + spaceObj, minValue = 0, maxValue = 1, dv = 0, keyable = True)
+
+    #add constraint to the new object on the follow node
+    constraint = cmds.parentConstraint(spaceObj, spaceSwitchFollow, mo = True)[0]
+
+    #hook up connections
+    targets = cmds.parentConstraint(constraint, q = True, targetList = True)
+    weight = 0
+    for i in range(int(len(targets))):
+        if targets[i].find(spaceObj) != -1:
+            weight = i
+            
+    cmds.connectAttr(spaceSwitchNode + ".space_" + spaceObj, constraint + "." + spaceObj + "W" + str(weight))
+
+    #lockNode on space object so it cannot be deleted by the user (if node is not a referenced node)
+    if spaceObj.find(":") == -1:
+        cmds.lockNode(spaceObj)
+	    
+
+const = cmds.parentConstraint("fk_tentacle_l_01_anim", "tentacle_spring_l_anim_grp", mo=False)
+cmds.move(0, 17, 0, "tentacle_spring_l_anim", os=True, wd=True)
+const = cmds.parentConstraint("fk_tentacle_r_01_anim", "tentacle_spring_r_anim_grp", mo=False)
+cmds.move(0, -17, 0, "tentacle_spring_r_anim", os=True, wd=True)
+
+	    
+#SET SPACES
+cmds.setAttr("weapon_l_anim_space_switcher.space_driver_hand_l", 1)
+cmds.setAttr("weapon_r_anim_space_switcher.space_driver_hand_r", 1)
+
+
+#SET ROTATE ORDER ON SWORD HANDLE
+cmds.setAttr("sword_handle_l_anim.rotateOrder", 3)
+cmds.setAttr("sword_handle_r_anim.rotateOrder", 3)
+
+#DEPLOY SWORD
+cmds.setAttr("sword_root_l_anim.sword_deploy", 10)
+cmds.setAttr("sword_root_l_anim.sword_open", 10)
+cmds.setAttr("sword_root_l_anim.sword_extend", 10)
+
+cmds.setAttr("sword_root_r_anim.sword_deploy", 10)
+cmds.setAttr("sword_root_r_anim.sword_open", 10)
+cmds.setAttr("sword_root_r_anim.sword_extend", 10)
+
+
+#TURN SWORD
+cmds.setAttr("sword_handle_l_anim.rotateY", 90)
+cmds.setAttr("sword_handle_r_anim.rotateY", 90)
+
+#TURN OFF SHOULDER PADS
+cmds.setAttr("shoulderPad_l_anim.chainStartEnvelope", 0)
+cmds.setAttr("shoulderPad_r_anim.chainStartEnvelope", 0)
+
+
+#SWORD HANDLE 
+cmds.setAttr("sword_handle_l_anim.tx", -7.5*0.65)
+cmds.setAttr("sword_handle_r_anim.tx", 7.5*0.65)
+cmds.setAttr("sword_handle_l_anim.ty", 1*0.65)
+cmds.setAttr("sword_handle_r_anim.ty", 1*0.65)
+
+# ------------------------------------------------------------------
+def shoulderRig():
+    # Left Shoulder
+    cmds.setAttr("shoulderflap_l_anim_grp.translate", -0.03089, 0, 0.30622, type="double3")
+    cmds.setAttr("shoulderflap_l_anim_grp.rotate", -0.45756, -8.57763, -3.24899, type="double3")
+    cmds.setDrivenKeyframe(["shoulderflap_l_anim_grp.translate", "shoulderflap_l_anim_grp.rotate"], cd = "driver_clavicle_l.rotateY", itt = "spline", ott = "spline")
+
+    cmds.setAttr("clavicle_l_anim.translateZ", 6.5)
+    cmds.setAttr("shoulderflap_l_anim_grp.translate", -0.09202, 0.05725, 0.27608, type="double3")
+    cmds.setAttr("shoulderflap_l_anim_grp.rotate", -8.73922, -16.32952, -7.35928, type="double3")
+    cmds.setDrivenKeyframe(["shoulderflap_l_anim_grp.translate", "shoulderflap_l_anim_grp.rotate"], cd = "driver_clavicle_l.rotateY", itt = "spline", ott = "spline")
+
+    cmds.setAttr("clavicle_l_anim.translateZ", -8)
+    cmds.setAttr("shoulderflap_l_anim_grp.translate", -1.36702, -0.35488, 2.09206, type="double3")
+    cmds.setAttr("shoulderflap_l_anim_grp.rotate", 12.18621, -2.83324, 3.00603, type="double3")
+    cmds.setDrivenKeyframe(["shoulderflap_l_anim_grp.translate", "shoulderflap_l_anim_grp.rotate"], cd = "driver_clavicle_l.rotateY", itt = "spline", ott = "spline")
+
+    cmds.setAttr("clavicle_l_anim.translateZ", 0)
+
+    # Right Shoulder
+    cmds.setAttr("shoulderflap_r_anim_grp.translate", 0.03089, 0, -0.30622, type="double3")
+    cmds.setAttr("shoulderflap_r_anim_grp.rotate", -0.45756, -8.57763, -3.24899, type="double3")
+    cmds.setDrivenKeyframe(["shoulderflap_r_anim_grp.translate", "shoulderflap_r_anim_grp.rotate"], cd = "driver_clavicle_r.rotateY", itt = "spline", ott = "spline")
+
+    cmds.setAttr("clavicle_r_anim.translateZ", 6.5)
+    cmds.setAttr("shoulderflap_r_anim_grp.translate", 0.09202, -0.05725, -0.27608, type="double3")
+    cmds.setAttr("shoulderflap_r_anim_grp.rotate", -8.73922, -16.32952, -7.35928, type="double3")
+    cmds.setDrivenKeyframe(["shoulderflap_r_anim_grp.translate", "shoulderflap_r_anim_grp.rotate"], cd = "driver_clavicle_r.rotateY", itt = "spline", ott = "spline")
+
+    cmds.setAttr("clavicle_r_anim.translateZ", -8)
+    cmds.setAttr("shoulderflap_r_anim_grp.translate", 1.36702, 0.35488, -2.09206, type="double3")
+    cmds.setAttr("shoulderflap_r_anim_grp.rotate", 12.18621, -2.83324, 3.00603, type="double3")
+    cmds.setDrivenKeyframe(["shoulderflap_r_anim_grp.translate", "shoulderflap_r_anim_grp.rotate"], cd = "driver_clavicle_r.rotateY", itt = "spline", ott = "spline")
+
+    cmds.setAttr("clavicle_r_anim.translateZ", 0)
+
+shoulderRig()
+
+
+# ------------------------------------------------------------------
+def dynamicsOff():
+    for one in ["shoulderPad_l_anim", "shoulderPad_r_anim", "tentacle_l_dyn_anim", "tentacle_r_dyn_anim", "thrusterVent_l_dyn_anim", "thrusterVent_r_dyn_anim"]:
+        cmds.setAttr(one+".chainStartEnvelope", 0)
+        cmds.setAttr(one+".chainStartEnvelope", l=True)
+dynamicsOff()
+
+# ------------------------------------------------------------------
+def legMorphsRig():
+    cmds.setAttr("Rig_Settings.lLegMode", 0)
+    cmds.setAttr("Rig_Settings.rLegMode", 0)
+    # SDKS
+    for side in ["_l", "_r"]:
+        cmds.setAttr("driver_calf"+side+".rotateOrder", 2)
+        cmds.setDrivenKeyframe(["Legs_geo_blendShape.Legs_geo_knee_bend"+side], cd = "driver_calf"+side+".rz", itt = "linear", ott = "linear")
+        cmds.setAttr("fk_calf"+side+"_anim.rz", 100)
+        cmds.setDrivenKeyframe(["Legs_geo_blendShape.Legs_geo_knee_bend"+side], cd = "driver_calf"+side+".rz", itt = "linear", ott = "linear")
+        cmds.setAttr("fk_calf"+side+"_anim.rz", 140)
+        cmds.setAttr("Legs_geo_blendShape.Legs_geo_knee_bend"+side, 1)
+        cmds.setDrivenKeyframe(["Legs_geo_blendShape.Legs_geo_knee_bend"+side], cd = "driver_calf"+side+".rz", itt = "linear", ott = "linear")
+        cmds.setAttr("fk_calf"+side+"_anim.rz", 0)
+    cmds.setAttr("Rig_Settings.lLegMode", 1)
+    cmds.setAttr("Rig_Settings.rLegMode", 1)
+
+legMorphsRig()
+# ------------------------------------------------------------------
+
+# ------------------------------------------------------------------
+def createGameCam():
+    gameCam = cmds.camera(centerOfInterest=5, focalLength=18, lensSqueezeRatio=1, cameraScale=1, horizontalFilmAperture=1.41732, horizontalFilmOffset=0, verticalFilmAperture=0.94488, verticalFilmOffset=0, filmFit="Overscan", overscan=1.05, motionBlur=0, shutterAngle=144, nearClipPlane=1.0, farClipPlane=10000, orthographic=0, orthographicWidth=30, panZoomEnabled=0, horizontalPan=0, verticalPan=0, zoom=1, displayFilmGate=True)[0]
+    print gameCam
+    gameCam = cmds.rename(gameCam, "GameCam")
+    cmds.camera(gameCam, e=True, horizontalFilmAperture=1.6724376, focalLength=21.2399857011)
+    cmds.setAttr(gameCam+".displayGateMaskOpacity", 1)
+    cmds.setAttr(gameCam+".displayGateMaskColor", 0, 0, 0, type="double3")
+    cmds.setAttr(gameCam+".r", 90, 0, 180, type="double3")
+    
+    # Create the control
+    control = cmds.circle(name=gameCam+"_Ctrl", ch=True, o=True, r=15)[0]
+    cmds.makeIdentity(control, a=True, t=True, r=True, s=True)
+    # Create groups for the sdk and the constraints
+    const_grp = cmds.group(control, name=control+"_grp")
+    if cmds.objExists("rig_grp"):
+        cmds.parent(const_grp, "rig_grp")
+    # Parent Constrain the cam to the control
+    cmds.parentConstraint(control, gameCam, mo=True)
+    
+    cmds.connectAttr(control+".visibility", gameCam+".visibility", f=True)
+    cmds.setAttr(const_grp+".v", 0)
+    
+    cmds.setAttr(const_grp+".t", 0, 500, 200, type="double3")
+createGameCam()
+# ------------------------------------------------------------------
+
+##################################################################################
+##################################################################################
+
+def addCapsule(radius, halfHeight):
+    trueHalfHeight = (halfHeight*2)-(radius*2)
+    capsule = cmds.polyCylinder(name="collision_capsule", r=radius, h=trueHalfHeight, sx=18, sy=1, sz=5, ax=[0,0,1], rcp=True, cuv=3, ch=True)
+    cmds.setAttr(capsule[0]+".tz", halfHeight+2.2)
+    cmds.setAttr(capsule[0]+".rz", 90)
+
+    cmds.addAttr(capsule[0], ln="Radius", dv=radius, keyable=True)
+    cmds.addAttr(capsule[0], ln="HalfHeight", dv=halfHeight, keyable=True)
+
+    expr = capsule[1]+".radius = "+capsule[0]+".Radius;\n"+capsule[1]+".height = ("+capsule[0]+".HalfHeight*2) - ("+capsule[0]+".Radius*2);"
+    cmds.expression(s=expr, o=capsule[1], ae=1, uc="all")
+
+    cmds.setAttr(capsule[0]+".overrideEnabled", 1)
+    cmds.setAttr(capsule[0]+".overrideColor", 18)
+    cmds.setAttr(capsule[0]+".overrideShading", 0)
+
+    cmds.addAttr(capsule[0], ln="Shaded", dv=0, min=0, max=1, keyable=True)
+    cmds.connectAttr(capsule[0]+".Shaded", capsule[0]+".overrideShading")
+
+    cmds.parentConstraint("root", capsule[0], mo=True)
+    cmds.select(capsule[0], r=True)
+    capsuleDisplayLayer = cmds.createDisplayLayer(nr=True, name="Collision Capsule")
+    cmds.setAttr(capsuleDisplayLayer+".displayType", 2)
+    cmds.setAttr(capsuleDisplayLayer+".v", 0)
+    
+    for attr in ["tx", "ty", "tz", "rx", "ry", "rz", "sx", "sy", "sz"]:
+        cmds.setAttr(capsule[0]+"."+attr, lock=True, keyable=False, channelBox=False)
+
+    cmds.select(clear=True)
+addCapsule(46, 72)
+
+# ------------------------------------------------------------------
+# add ik bones Function that is called later after we open the export file.
+def addIKBones():
+    try:
+		# create the joints
+		cmds.select(clear = True)
+		ikFootRoot = cmds.joint(name = "ik_foot_root")
+		cmds.select(clear = True)
+		
+		cmds.select(clear = True)
+		ikFootLeft = cmds.joint(name = "ik_foot_l")
+		cmds.select(clear = True)
+		
+		cmds.select(clear = True)
+		ikFootRight = cmds.joint(name = "ik_foot_r")
+		cmds.select(clear = True)
+		
+		cmds.select(clear = True)
+		ikHandRoot = cmds.joint(name = "ik_hand_root")
+		cmds.select(clear = True)
+		
+		cmds.select(clear = True)
+		ikHandGun = cmds.joint(name = "ik_hand_gun")
+		cmds.select(clear = True)
+		
+		cmds.select(clear = True)
+		ikHandLeft = cmds.joint(name = "ik_hand_l")
+		cmds.select(clear = True)
+		
+		cmds.select(clear = True)
+		ikHandRight = cmds.joint(name = "ik_hand_r")
+		cmds.select(clear = True)
+		
+		# create hierarhcy
+		cmds.parent(ikFootRoot, "root")
+		cmds.parent(ikHandRoot, "root")
+		cmds.parent(ikFootLeft, ikFootRoot)
+		cmds.parent(ikFootRight, ikFootRoot)
+		cmds.parent(ikHandGun, ikHandRoot)
+		cmds.parent(ikHandLeft, ikHandGun)
+		cmds.parent(ikHandRight, ikHandGun)
+		
+		# constrain the joints
+		leftHandConstraint = cmds.parentConstraint("hand_l", ikHandLeft)[0]
+		rightHandGunConstraint = cmds.parentConstraint("hand_r", ikHandGun)[0]
+		rightHandConstraint = cmds.parentConstraint("hand_r", ikHandRight)[0]
+		leftFootConstraint = cmds.parentConstraint("foot_l", ikFootLeft)[0]
+		rightFootConstraint = cmds.parentConstraint("foot_r", ikFootRight)[0]
+	
+    except:
+		cmds.warning("Something went wrong")
+
+# ------------------------------------------------------------------
+# ADD the IK bones to the Export file.
+#save the AnimRig file before opening the export file.
+cmds.file(save = True, type = "mayaBinary", force = True)
+
+#ADD IK BONES
+#open export file
+sceneName = cmds.file(q = True, sceneName = True)
+exportFile = sceneName.partition("AnimRigs")[0] + "ExportFiles" +  sceneName.partition("AnimRigs")[2].partition(".mb")[0] + "_Export.mb"
+cmds.file(exportFile, open = True, force = True)
+
+#add ik bones
+addIKBones()
+
+# ------------------------------------------------------------------
+def addCustomAttrs():
+    attrs = cmds.listAttr("root")
+    exists = False
+    for one in attrs:
+        if one.find("DistanceCurve") != -1:
+            exists = True
+    if exists == False:
+        cmds.addAttr("root", ln="DistanceCurve", keyable=True)
+addCustomAttrs()
+
+#save the export file.
+cmds.file(save = True, type = "mayaBinary", force = True)
\ No newline at end of file
