Branching

diff --git a/Engine/Extras/Maya_AnimationRiggingTools/ARTv1/General/ART/Projects/Orion/Scripts/Minion_Leader_Post_Script.py b/Engine/Extras/Maya_AnimationRiggingTools/ARTv1/General/ART/Projects/Orion/Scripts/Minion_Leader_Post_Script.py
new file mode 100644
--- /dev/null
+++ b/Engine/Extras/Maya_AnimationRiggingTools/ARTv1/General/ART/Projects/Orion/Scripts/Minion_Leader_Post_Script.py
@@ -0,0 +1,194 @@
+import maya.cmds as cmds
+import maya.mel as mel
+import math
+import Tools.ART_addSpaceSwitching as space
+reload(space)
+
+# NOTE: This character has special steps to run during Rig Pose creation.  See the following commented out script.
+#########################################
+# RUN THIS COMMAND DURING THE RIG POSE CREATION STEP TO GET THE SHIELD TO THE CORRECT LOCATION.  LEAVE IT COMMENTED OUT SO THAT IT DOESNT RUN DURING RIGGING.
+'''
+# The following 2 lines are used to find the locations and then input below.
+#cmds.xform(q=True, ws=True, rp=True)
+#cmds.xform(q=True, ro=True)
+
+
+cmds.xform("calf_mover_l", ws=True, t=[19.79804629, 3.337048167, 40.39003741])
+cmds.xform("calf_mover_l", ro=[160.36136150000002, -71.79676151, -144.3110928])
+
+cmds.xform("calf_mover_r", ws=True, t=[-19.79807484, 3.337052443, 40.39004631])
+cmds.xform("calf_mover_r", ro=[-19.638638459999992, 71.79676151, 144.3110928])
+
+cmds.xform("foot_mover_l", ws=True, t=[26.80774865, 8.872520219, 12.46154229])
+cmds.xform("foot_mover_l", ro=[0.0, -90.0, 0.0])
+
+cmds.xform("foot_mover_r", ws=True, t=[-26.80775027, 8.872528025, 12.46154483])
+cmds.xform("foot_mover_r", ro=[90.0, 90.0, -90.0])
+'''
+#########################################
+
+# ------------------------------------------------------------------
+# scale all of the controls to be more visible with the mesh.
+def scaleControls():
+    for side in ["_l", "_r"]:
+        # feet
+        cmds.select(["ik_foot_anim"+side+".cv[0:17]"])
+        cmds.scale(1.5, 1, 1, r = True, ocp = False)
+        cmds.move(0, 6, -5.468891, r = True)
+
+        # hands
+        cmds.select(["ik_wrist"+side+"_anim.cv[0:7]"])
+        cmds.scale(1.9, 1.9, 1.9, r = True, ocp = False)
+
+        # other
+        cmds.select(["weapon"+side+"_anim.cv[0:32]"])
+        cmds.scale(0.238522, 0.238522, 0.238522, r = True, ocp = True)
+        cmds.move(-5.6, 0, 0, r = True, os=True, wd=True)
+
+        cmds.select(["clavicle"+side+"_anim.cv[0:32]"])
+        cmds.move(0, 0, 11, r = True, os=True, wd=True)
+
+    for one in ["shoulder_blade_A_l_anim", "shoulder_blade_B_l_anim", "shoulder_blade_A_r_anim", "shoulder_blade_B_r_anim"]:
+        cmds.select([one+".cv[*]"])
+        cmds.scale(0.3, 0.3, 0.3, r = True, ocp = True)
+        cmds.move(0, 0, 23, r = True, os=False, wd=True)
+
+    cmds.select(None)
+
+scaleControls()
+
+# ------------------------------------------------------------------
+def shoulderRig():
+    cmds.setDrivenKeyframe(["shoulder_blade_A_l_anim_grp.translate", "shoulder_blade_B_l_anim_grp.translate"], cd = "driver_clavicle_l.rotateY", itt = "spline", ott = "spline")
+    cmds.setAttr("clavicle_l_anim.translateZ", 12)
+    cmds.setAttr("shoulder_blade_A_l_anim_grp.translate", 1.9051, 0, -5.01986, type="double3")
+    cmds.setAttr("shoulder_blade_B_l_anim_grp.translate", 0, 0, -2.74608, type="double3")
+    cmds.setDrivenKeyframe(["shoulder_blade_A_l_anim_grp.translate", "shoulder_blade_B_l_anim_grp.translate"], cd = "driver_clavicle_l.rotateY", itt = "spline", ott = "spline")
+    cmds.setAttr("clavicle_l_anim.translateZ", 0)
+
+    cmds.setAttr("clavicle_l_anim.translateZ", -10)
+    cmds.setAttr("shoulder_blade_A_l_anim_grp.translate", -3.99968, 0, 5.47598, type="double3")
+    cmds.setAttr("shoulder_blade_B_l_anim_grp.translate", 0, 0, 4.31874, type="double3")
+    cmds.setDrivenKeyframe(["shoulder_blade_A_l_anim_grp.translate", "shoulder_blade_B_l_anim_grp.translate"], cd = "driver_clavicle_l.rotateY", itt = "spline", ott = "spline")
+    cmds.setAttr("clavicle_l_anim.translateZ", 0)
+    
+    
+    
+    
+    
+    cmds.setDrivenKeyframe(["shoulder_blade_A_r_anim_grp.translate", "shoulder_blade_B_r_anim_grp.translate"], cd = "driver_clavicle_r.rotateY", itt = "spline", ott = "spline")
+    cmds.setAttr("clavicle_r_anim.translateZ", 12)
+    cmds.setAttr("shoulder_blade_A_r_anim_grp.translate", -1.9051, 0, 5.01986, type="double3")
+    cmds.setAttr("shoulder_blade_B_r_anim_grp.translate", 0, 0, 2.74608, type="double3")
+    cmds.setDrivenKeyframe(["shoulder_blade_A_r_anim_grp.translate", "shoulder_blade_B_r_anim_grp.translate"], cd = "driver_clavicle_r.rotateY", itt = "spline", ott = "spline")
+    cmds.setAttr("clavicle_r_anim.translateZ", 0)
+
+    cmds.setAttr("clavicle_r_anim.translateZ", -10)
+    cmds.setAttr("shoulder_blade_A_r_anim_grp.translate", 3.99968, 0, -5.47598, type="double3")
+    cmds.setAttr("shoulder_blade_B_r_anim_grp.translate", 0, 0, -4.31874, type="double3")
+    cmds.setDrivenKeyframe(["shoulder_blade_A_r_anim_grp.translate", "shoulder_blade_B_r_anim_grp.translate"], cd = "driver_clavicle_r.rotateY", itt = "spline", ott = "spline")
+    cmds.setAttr("clavicle_r_anim.translateZ", 0)
+shoulderRig()
+
+##################################################################################
+##################################################################################
+
+def addCapsule(radius, halfHeight):
+    trueHalfHeight = (halfHeight*2)-(radius*2)
+    capsule = cmds.polyCylinder(name="collision_capsule", r=radius, h=trueHalfHeight, sx=18, sy=1, sz=5, ax=[0,0,1], rcp=True, cuv=3, ch=True)
+    cmds.setAttr(capsule[0]+".tz", halfHeight+2.2)
+    cmds.setAttr(capsule[0]+".rz", 90)
+
+    cmds.addAttr(capsule[0], ln="Radius", dv=radius, keyable=True)
+    cmds.addAttr(capsule[0], ln="HalfHeight", dv=halfHeight, keyable=True)
+
+    expr = capsule[1]+".radius = "+capsule[0]+".Radius;\n"+capsule[1]+".height = ("+capsule[0]+".HalfHeight*2) - ("+capsule[0]+".Radius*2);"
+    cmds.expression(s=expr, o=capsule[1], ae=1, uc="all")
+
+    cmds.setAttr(capsule[0]+".overrideEnabled", 1)
+    cmds.setAttr(capsule[0]+".overrideColor", 18)
+    cmds.setAttr(capsule[0]+".overrideShading", 0)
+
+    cmds.addAttr(capsule[0], ln="Shaded", dv=0, min=0, max=1, keyable=True)
+    cmds.connectAttr(capsule[0]+".Shaded", capsule[0]+".overrideShading")
+
+    cmds.parentConstraint("root", capsule[0], mo=True)
+    cmds.select(capsule[0], r=True)
+    capsuleDisplayLayer = cmds.createDisplayLayer(nr=True, name="Collision Capsule")
+    cmds.setAttr(capsuleDisplayLayer+".displayType", 2)
+    cmds.setAttr(capsuleDisplayLayer+".v", 0)
+    
+    for attr in ["tx", "ty", "tz", "rx", "ry", "rz", "sx", "sy", "sz"]:
+        cmds.setAttr(capsule[0]+"."+attr, lock=True, keyable=False, channelBox=False)
+
+    cmds.select(clear=True)
+addCapsule(60, 135)
+
+# ------------------------------------------------------------------
+# add ik bones Function that is called later after we open the export file.
+def addIKBones():
+    try:
+		# create the joints
+		cmds.select(clear = True)
+		ikFootRoot = cmds.joint(name = "ik_foot_root")
+		cmds.select(clear = True)
+		
+		cmds.select(clear = True)
+		ikFootLeft = cmds.joint(name = "ik_foot_l")
+		cmds.select(clear = True)
+		
+		cmds.select(clear = True)
+		ikFootRight = cmds.joint(name = "ik_foot_r")
+		cmds.select(clear = True)
+		
+		cmds.select(clear = True)
+		ikHandRoot = cmds.joint(name = "ik_hand_root")
+		cmds.select(clear = True)
+		
+		cmds.select(clear = True)
+		ikHandGun = cmds.joint(name = "ik_hand_gun")
+		cmds.select(clear = True)
+		
+		cmds.select(clear = True)
+		ikHandLeft = cmds.joint(name = "ik_hand_l")
+		cmds.select(clear = True)
+		
+		cmds.select(clear = True)
+		ikHandRight = cmds.joint(name = "ik_hand_r")
+		cmds.select(clear = True)
+		
+		# create hierarhcy
+		cmds.parent(ikFootRoot, "root")
+		cmds.parent(ikHandRoot, "root")
+		cmds.parent(ikFootLeft, ikFootRoot)
+		cmds.parent(ikFootRight, ikFootRoot)
+		cmds.parent(ikHandGun, ikHandRoot)
+		cmds.parent(ikHandLeft, ikHandGun)
+		cmds.parent(ikHandRight, ikHandGun)
+		
+		# constrain the joints
+		leftHandConstraint = cmds.parentConstraint("hand_l", ikHandLeft)[0]
+		rightHandGunConstraint = cmds.parentConstraint("hand_r", ikHandGun)[0]
+		rightHandConstraint = cmds.parentConstraint("hand_r", ikHandRight)[0]
+		leftFootConstraint = cmds.parentConstraint("foot_l", ikFootLeft)[0]
+		rightFootConstraint = cmds.parentConstraint("foot_r", ikFootRight)[0]
+	
+    except:
+		cmds.warning("Something went wrong")
+
+# ------------------------------------------------------------------
+# ADD the IK bones to the Export file.
+#save the AnimRig file before opening the export file.
+cmds.file(save = True, type = "mayaBinary", force = True)
+
+#ADD IK BONES
+#open export file
+sceneName = cmds.file(q = True, sceneName = True)
+exportFile = sceneName.partition("AnimRigs")[0] + "ExportFiles" +  sceneName.partition("AnimRigs")[2].partition(".mb")[0] + "_Export.mb"
+cmds.file(exportFile, open = True, force = True)
+
+#add ik bones
+addIKBones()
+
+#save the export file.
+cmds.file(save = True, type = "mayaBinary", force = True)
diff --git a/Engine/Extras/Maya_AnimationRiggingTools/ARTv1/General/ART/Projects/Orion/Scripts/Minion_Melee_Post_Script.py b/Engine/Extras/Maya_AnimationRiggingTools/ARTv1/General/ART/Projects/Orion/Scripts/Minion_Melee_Post_Script.py
new file mode 100644
--- /dev/null
+++ b/Engine/Extras/Maya_AnimationRiggingTools/ARTv1/General/ART/Projects/Orion/Scripts/Minion_Melee_Post_Script.py
@@ -0,0 +1,230 @@
+import maya.cmds as cmds
+
+# NOTE: This character has special steps to run during Rig Pose creation.  See the following commented out script.
+#########################################
+# RUN THIS COMMAND DURING THE RIG POSE CREATION STEP TO GET THE SHIELD TO THE CORRECT LOCATION.  LEAVE IT COMMENTED OUT SO THAT IT DOESNT RUN DURING RIGGING.
+'''
+# The following 2 lines are used to find the locations and then input below.
+#cmds.xform(q=True, ws=True, rp=True)
+#cmds.xform(q=True, ro=True)
+
+
+cmds.xform("calf_mover_l", ws=True, t=[15.209881960000002, -1.3946871069999982, 32.64604213])
+cmds.xform("calf_mover_l", ro=[-1.6927803365230126e-08, -7.90540084577015e-10, 4.610907371818708e-09])
+
+cmds.xform("calf_mover_r", ws=True, t=[-15.209908050000001, -1.3946831929999999, 32.64604689])
+cmds.xform("calf_mover_r", ro=[-1.692620235351736e-08, -7.906350042557556e-10, 4.611578468856501e-09])
+
+cmds.xform("foot_mover_l", ws=True, t=[17.394773850000004, 4.334110128000001, 11.327026350000002])
+cmds.xform("foot_mover_l", ro=[-27.333076877633392, 10.809974591780163, -0.40751877939002024])
+
+cmds.xform("foot_mover_r", ws=True, t=[-17.394775560000003, 4.334116807000003, 11.32702872])
+cmds.xform("foot_mover_r", ro=[-27.33307687763361, 10.809974591780179, -0.40751877938934844])
+'''
+#########################################
+
+# ------------------------------------------------------------------
+# scale all of the controls to be more visible with the mesh.
+def scaleControls():
+    for side in ["_l", "_r"]:
+        # feet
+        cmds.select(["ik_foot_anim"+side+".cv[0:17]"])
+        cmds.scale(1.5, 1, 1, r = True, ocp = False)
+        cmds.move(0, 6, -5.468891, r = True)
+
+        # hands
+        cmds.select(["ik_wrist"+side+"_anim.cv[0:7]"])
+        cmds.scale(1.9, 1.9, 1.9, r = True, ocp = False)
+
+        # other
+        cmds.select(["weapon"+side+"_anim.cv[0:32]"])
+        cmds.scale(0.238522, 0.238522, 0.238522, r = True, ocp = True)
+        cmds.move(-5.6, 0, 0, r = True, os=True, wd=True)
+
+        cmds.select(["clavicle"+side+"_anim.cv[0:32]"])
+        cmds.move(0, 0, 11, r = True, os=True, wd=True)
+
+    for one in ["shoulder_blade_A_l_anim", "shoulder_blade_C_l_anim", "shoulder_blade_B_l_anim", "shoulder_blade_A_r_anim", "shoulder_blade_C_r_anim", "shoulder_blade_B_r_anim"]:
+        cmds.select([one+".cv[*]"])
+        cmds.scale(0.3, 0.3, 0.3, r = True, ocp = True)
+        cmds.move(0, 0, 0, r = True, os=True, wd=True)
+
+    cmds.select(None)
+
+scaleControls()
+
+# ------------------------------------------------------------------
+def shoulderRig():
+    cmds.setDrivenKeyframe(["shoulder_blade_A_l_anim_grp.translate", "shoulder_blade_B_l_anim_grp.translate", "shoulder_blade_C_l_anim_grp.translate"], cd = "driver_clavicle_l.rotateY", itt = "flat", ott = "flat")
+    cmds.setAttr("clavicle_l_anim.translateZ", 12)
+    cmds.setAttr("shoulder_blade_A_l_anim_grp.translate", 0, 0, -1.23994, type="double3")
+    cmds.setAttr("shoulder_blade_B_l_anim_grp.translate", 0, 0, -2.74608, type="double3")
+    cmds.setAttr("shoulder_blade_C_l_anim_grp.translate", 0, 0, -4.31946, type="double3")
+    cmds.setDrivenKeyframe(["shoulder_blade_A_l_anim_grp.translate", "shoulder_blade_B_l_anim_grp.translate", "shoulder_blade_C_l_anim_grp.translate"], cd = "driver_clavicle_l.rotateY", itt = "flat", ott = "flat")
+    cmds.setAttr("clavicle_l_anim.translateZ", 0)
+
+
+    cmds.setDrivenKeyframe(["shoulder_blade_A_r_anim_grp.translate", "shoulder_blade_B_r_anim_grp.translate", "shoulder_blade_C_r_anim_grp.translate"], cd = "driver_clavicle_r.rotateY", itt = "flat", ott = "flat")
+    cmds.setAttr("clavicle_r_anim.translateZ", 12)
+    cmds.setAttr("shoulder_blade_A_r_anim_grp.translate", 0, 0, 1.23994, type="double3")
+    cmds.setAttr("shoulder_blade_B_r_anim_grp.translate", 0, 0, 2.74608, type="double3")
+    cmds.setAttr("shoulder_blade_C_r_anim_grp.translate", 0, 0, 4.31946, type="double3")
+    cmds.setDrivenKeyframe(["shoulder_blade_A_r_anim_grp.translate", "shoulder_blade_B_r_anim_grp.translate", "shoulder_blade_C_r_anim_grp.translate"], cd = "driver_clavicle_r.rotateY", itt = "flat", ott = "flat")
+    cmds.setAttr("clavicle_r_anim.translateZ", 0)
+
+shoulderRig()
+
+# ------------------------------------------------------------------
+def shoulderPadRig():
+    shoulderpad_l_anim_drv = cmds.group("shoulderpad_l_anim", name="shoulderpad_l_anim_drv")
+    cmds.aimConstraint("driver_upperarm_l", shoulderpad_l_anim_drv, aimVector = [1, 0, 0], upVector = [0, 0, 1], worldUpType="objectrotation", worldUpVector=[0, 1, 0], worldUpObject="driver_clavicle_l", mo=True, skip = ["x", "z"])
+    cmds.transformLimits("shoulderpad_l_anim_drv", ry=[0, 1.92], ery=[0, 1])
+    
+    shoulderpad_r_anim_drv = cmds.group("shoulderpad_r_anim", name="shoulderpad_r_anim_drv")
+    cmds.aimConstraint("driver_upperarm_r", shoulderpad_r_anim_drv, aimVector = [-1, 0, 0], upVector = [0, 0, -1], worldUpType="objectrotation", worldUpVector=[0, 1, 0], worldUpObject="driver_clavicle_r", mo=True, skip = ["x", "z"])
+    cmds.transformLimits("shoulderpad_r_anim_drv", ry=[0, 1.92], ery=[0, 1])
+
+    shoulderblade_l_anim_drv = cmds.group("shoulderblade_l_anim", name="shoulderblade_l_anim_drv")
+    cmds.aimConstraint("driver_upperarm_l", shoulderblade_l_anim_drv, aimVector = [1, 0, 0], upVector = [0, 0, 1], worldUpType="objectrotation", worldUpVector=[0, 1, 0], worldUpObject="driver_clavicle_l", mo=True, skip = ["x", "z"])
+    cmds.transformLimits("shoulderblade_l_anim_drv", ry=[0, 1.92], ery=[0, 1])
+    
+    shoulderblade_r_anim_drv = cmds.group("shoulderblade_r_anim", name="shoulderblade_r_anim_drv")
+    cmds.aimConstraint("driver_upperarm_r", shoulderblade_r_anim_drv, aimVector = [-1, 0, 0], upVector = [0, 0, -1], worldUpType="objectrotation", worldUpVector=[0, 1, 0], worldUpObject="driver_clavicle_r", mo=True, skip = ["x", "z"])
+    cmds.transformLimits("shoulderblade_r_anim_drv", ry=[0, 1.92], ery=[0, 1])
+
+#shoulderPadRig()
+
+##################################################################################
+##################################################################################
+
+def addCapsule(radius, halfHeight):
+    trueHalfHeight = (halfHeight*2)-(radius*2)
+    capsule = cmds.polyCylinder(name="collision_capsule", r=radius, h=trueHalfHeight, sx=18, sy=1, sz=5, ax=[0,0,1], rcp=True, cuv=3, ch=True)
+    cmds.setAttr(capsule[0]+".tz", halfHeight+2.2)
+    cmds.setAttr(capsule[0]+".rz", 90)
+
+    cmds.addAttr(capsule[0], ln="Radius", dv=radius, keyable=True)
+    cmds.addAttr(capsule[0], ln="HalfHeight", dv=halfHeight, keyable=True)
+
+    expr = capsule[1]+".radius = "+capsule[0]+".Radius;\n"+capsule[1]+".height = ("+capsule[0]+".HalfHeight*2) - ("+capsule[0]+".Radius*2);"
+    cmds.expression(s=expr, o=capsule[1], ae=1, uc="all")
+
+    cmds.setAttr(capsule[0]+".overrideEnabled", 1)
+    cmds.setAttr(capsule[0]+".overrideColor", 18)
+    cmds.setAttr(capsule[0]+".overrideShading", 0)
+
+    cmds.addAttr(capsule[0], ln="Shaded", dv=0, min=0, max=1, keyable=True)
+    cmds.connectAttr(capsule[0]+".Shaded", capsule[0]+".overrideShading")
+
+    cmds.parentConstraint("root", capsule[0], mo=True)
+    cmds.select(capsule[0], r=True)
+    capsuleDisplayLayer = cmds.createDisplayLayer(nr=True, name="Collision Capsule")
+    cmds.setAttr(capsuleDisplayLayer+".displayType", 2)
+    cmds.setAttr(capsuleDisplayLayer+".v", 0)
+    
+    for attr in ["tx", "ty", "tz", "rx", "ry", "rz", "sx", "sy", "sz"]:
+        cmds.setAttr(capsule[0]+"."+attr, lock=True, keyable=False, channelBox=False)
+
+    cmds.select(clear=True)
+addCapsule(60, 85)
+
+# ------------------------------------------------------------------
+# add ik bones Function that is called later after we open the export file.
+def addIKBones():
+    try:
+		# create the joints
+		cmds.select(clear = True)
+		ikFootRoot = cmds.joint(name = "ik_foot_root")
+		cmds.select(clear = True)
+		
+		cmds.select(clear = True)
+		ikFootLeft = cmds.joint(name = "ik_foot_l")
+		cmds.select(clear = True)
+		
+		cmds.select(clear = True)
+		ikFootRight = cmds.joint(name = "ik_foot_r")
+		cmds.select(clear = True)
+		
+		cmds.select(clear = True)
+		ikHandRoot = cmds.joint(name = "ik_hand_root")
+		cmds.select(clear = True)
+		
+		cmds.select(clear = True)
+		ikHandGun = cmds.joint(name = "ik_hand_gun")
+		cmds.select(clear = True)
+		
+		cmds.select(clear = True)
+		ikHandLeft = cmds.joint(name = "ik_hand_l")
+		cmds.select(clear = True)
+		
+		cmds.select(clear = True)
+		ikHandRight = cmds.joint(name = "ik_hand_r")
+		cmds.select(clear = True)
+		
+		# create hierarhcy
+		cmds.parent(ikFootRoot, "root")
+		cmds.parent(ikHandRoot, "root")
+		cmds.parent(ikFootLeft, ikFootRoot)
+		cmds.parent(ikFootRight, ikFootRoot)
+		cmds.parent(ikHandGun, ikHandRoot)
+		cmds.parent(ikHandLeft, ikHandGun)
+		cmds.parent(ikHandRight, ikHandGun)
+		
+		# constrain the joints
+		leftHandConstraint = cmds.parentConstraint("hand_l", ikHandLeft)[0]
+		rightHandGunConstraint = cmds.parentConstraint("hand_r", ikHandGun)[0]
+		rightHandConstraint = cmds.parentConstraint("hand_r", ikHandRight)[0]
+		leftFootConstraint = cmds.parentConstraint("foot_l", ikFootLeft)[0]
+		rightFootConstraint = cmds.parentConstraint("foot_r", ikFootRight)[0]
+	
+    except:
+		cmds.warning("Something went wrong")
+
+# ------------------------------------------------------------------
+# ADD the IK bones to the Export file.
+#save the AnimRig file before opening the export file.
+cmds.file(save = True, type = "mayaBinary", force = True)
+
+#ADD IK BONES
+#open export file
+sceneName = cmds.file(q = True, sceneName = True)
+exportFile = sceneName.partition("AnimRigs")[0] + "ExportFiles" +  sceneName.partition("AnimRigs")[2].partition(".mb")[0] + "_Export.mb"
+cmds.file(exportFile, open = True, force = True)
+
+#add ik bones
+addIKBones()
+
+#save the export file.
+cmds.file(save = True, type = "mayaBinary", force = True)
+
+
+
+
+
+
+
+
+
+
+
+
+
+
+
+
+
+
+
+
+
+
+
+
+
+
+
+
+
+
+
+
diff --git a/Engine/Extras/Maya_AnimationRiggingTools/ARTv1/General/ART/Projects/Orion/Scripts/Minion_Ranged_CreateLODs.py b/Engine/Extras/Maya_AnimationRiggingTools/ARTv1/General/ART/Projects/Orion/Scripts/Minion_Ranged_CreateLODs.py
new file mode 100644
--- /dev/null
+++ b/Engine/Extras/Maya_AnimationRiggingTools/ARTv1/General/ART/Projects/Orion/Scripts/Minion_Ranged_CreateLODs.py
@@ -0,0 +1,242 @@
+import maya.cmds as cmds
+import customMayaMenu as menu
+import Tools.ART_ImportMultipleWeights as impw
+import Tools.ART_TransferWeights as transfer
+import Tools.ART_ExportLODtoFBX as export
+reload(export)
+
+################################################## MINION RANGED #########################################################################
+# Warn the user that this operation will save their file.
+saveResult = cmds.confirmDialog(title = "Save Warning", message="In order to continue your file must be saved.  Would you like to save it?  If yes it will be saved and after the operation is complete your file will be re-opened.", button = ["Yes", "Cancel"], defaultButton='Yes', cancelButton='Cancel', dismissString='Cancel')
+
+# If they are ok with saving the file, warn them that it will also delete all of the current files in their skin weights directory.
+if saveResult == "Yes":
+    deleteSkinFilesResult = cmds.confirmDialog(title = "SkinWeights Warning", message="In order to continue we must delete all of the files in your ..\General\ART\SkinWeights\ directory.  Would you like to continue?.", button = ["Yes", "Cancel"], defaultButton='Yes', cancelButton='Cancel', dismissString='Cancel')
+    if deleteSkinFilesResult == "Yes":
+        toolsPath = cmds.internalVar(usd = True) + "mayaTools.txt"
+        if os.path.exists(toolsPath):
+            f = open(toolsPath, 'r')
+            mayaToolsDir = f.readline()
+            f.close()
+
+        path = mayaToolsDir + "\General\ART\SkinWeights\\"
+        files = os.listdir(path)
+        #delete all weight cache files in the system folder
+        for file in files:
+            os.remove(path + file)
+            #print path + file
+
+
+
+# Minion_Ranged LOD1 -------------------------------------------------------------------------------------------------------------
+def Minion_RangedLOD1():
+    if saveResult == "Yes":
+
+        cmds.file(save=True)
+
+        meshes = ["geo_ankle_pivots_LOD1", "geo_calfs_LOD1", "geo_feet_LOD1", "geo_fingers_LOD1", "geo_hands_LOD1", "geo_head_LOD1", "geo_hip_ball_pivots_LOD1", "geo_knee_pivots_LOD1", "geo_knees_LOD1", "geo_lower_arms_LOD1", "geo_neck_LOD1", "geo_pelvis_LOD1", "geo_shoulder_blades_LOD1", "geo_shoulder_plates_LOD1", "geo_thighs_LOD1", "geo_upper_arms_LOD1", "geo_upper_torso_LOD1", "geo_weapon_LOD1"]
+        
+        jointsToDelete = []
+
+        # remove unneeded joints
+        jointsToRemove = ["shoulder_blade_A_l", "shoulder_blade_B_l", "shoulder_blade_C_l"]
+        jointToTransferTo = "clavicle_l"
+        transfer.LOD_transferWeights(meshes, jointsToRemove, jointToTransferTo)
+        jointsToDelete.extend(jointsToRemove)
+
+        jointsToRemove = ["shoulder_blade_A_r", "shoulder_blade_B_r", "shoulder_blade_C_r"]
+        jointToTransferTo = "clavicle_r"
+        transfer.LOD_transferWeights(meshes, jointsToRemove, jointToTransferTo)
+        jointsToDelete.extend(jointsToRemove)
+
+        jointsToRemove = ["lowerarm_twist_01_l"]
+        jointToTransferTo = "lowerarm_l"
+        transfer.LOD_transferWeights(meshes, jointsToRemove, jointToTransferTo)
+        jointsToDelete.extend(jointsToRemove)
+
+        jointsToRemove = ["lowerarm_twist_01_r"]
+        jointToTransferTo = "lowerarm_r"
+        transfer.LOD_transferWeights(meshes, jointsToRemove, jointToTransferTo)
+        jointsToDelete.extend(jointsToRemove)
+
+        jointsToRemove = ["index_01_l", "index_02_l", "middle_01_l", "middle_02_l", "pinky_01_l", "pinky_02_l", "ring_01_l", "ring_02_l"]
+        jointToTransferTo = "hand_l"
+        transfer.LOD_transferWeights(meshes, jointsToRemove, jointToTransferTo)
+        jointsToDelete.extend(jointsToRemove)
+
+        jointsToRemove = ["index_01_r", "index_02_r", "middle_01_r", "middle_02_r", "pinky_01_r", "pinky_02_r", "ring_01_r", "ring_02_r"]
+        jointToTransferTo = "hand_r"
+        transfer.LOD_transferWeights(meshes, jointsToRemove, jointToTransferTo)
+        jointsToDelete.extend(jointsToRemove)
+
+        # Delete the joints we no longer want.
+        for one in jointsToDelete:
+            if cmds.objExists(one):
+                cmds.delete(one)
+                
+        # Export the fbx file.
+        export.exportLOD("Minion_Ranged_LOD1.fbx", "Minion_Ranged_Geo_LOD1")
+        
+        # Re-open the file without saving.
+        fullPath = cmds.file(q = True, sceneName = True)
+        cmds.file(fullPath, open=True, f=True)
+
+        print "LOD1 Exported Successfully!!!!!!!!!!!!!!!"
+
+Minion_RangedLOD1()
+
+# Minion_Ranged LOD2 -------------------------------------------------------------------------------------------------------------
+def Minion_RangedLOD2():
+    if saveResult == "Yes":
+
+        cmds.file(save=True)
+        
+        meshes = ["geo_ankle_pivots_LOD2", "geo_calfs_LOD2", "geo_feet_LOD2", "geo_fingers_LOD2", "geo_hands_LOD2", "geo_head_LOD2", "geo_knees_LOD2", "geo_lower_arms_LOD2", "geo_neck_LOD2", "geo_pelvis_LOD2", "geo_shoulder_blades_LOD2", "geo_shoulder_plates_LOD2", "geo_thighs_LOD2", "geo_upper_arms_LOD2", "geo_upper_torso_LOD2", "geo_weapon_LOD2"]
+
+        jointsToDelete = []
+
+         # remove unneeded joints
+        jointsToRemove = ["shoulder_blade_A_l", "shoulder_blade_B_l", "shoulder_blade_C_l"]
+        jointToTransferTo = "clavicle_l"
+        transfer.LOD_transferWeights(meshes, jointsToRemove, jointToTransferTo)
+        jointsToDelete.extend(jointsToRemove)
+
+        jointsToRemove = ["shoulder_blade_A_r", "shoulder_blade_B_r", "shoulder_blade_C_r"]
+        jointToTransferTo = "clavicle_r"
+        transfer.LOD_transferWeights(meshes, jointsToRemove, jointToTransferTo)
+        jointsToDelete.extend(jointsToRemove)
+
+        jointsToRemove = ["upperarm_twist_01_l"]
+        jointToTransferTo = "upperarm_l"
+        transfer.LOD_transferWeights(meshes, jointsToRemove, jointToTransferTo)
+        jointsToDelete.extend(jointsToRemove)
+
+        jointsToRemove = ["upperarm_twist_01_r"]
+        jointToTransferTo = "upperarm_r"
+        transfer.LOD_transferWeights(meshes, jointsToRemove, jointToTransferTo)
+        jointsToDelete.extend(jointsToRemove)
+
+        jointsToRemove = ["lowerarm_twist_01_l"]
+        jointToTransferTo = "lowerarm_l"
+        transfer.LOD_transferWeights(meshes, jointsToRemove, jointToTransferTo)
+        jointsToDelete.extend(jointsToRemove)
+
+        jointsToRemove = ["lowerarm_twist_01_r"]
+        jointToTransferTo = "lowerarm_r"
+        transfer.LOD_transferWeights(meshes, jointsToRemove, jointToTransferTo)
+        jointsToDelete.extend(jointsToRemove)
+
+        jointsToRemove = ["index_01_l", "index_02_l", "middle_01_l", "middle_02_l", "pinky_01_l", "pinky_02_l", "ring_01_l", "ring_02_l"]
+        jointToTransferTo = "hand_l"
+        transfer.LOD_transferWeights(meshes, jointsToRemove, jointToTransferTo)
+        jointsToDelete.extend(jointsToRemove)
+
+        jointsToRemove = ["index_01_r", "index_02_r", "middle_01_r", "middle_02_r", "pinky_01_r", "pinky_02_r", "ring_01_r", "ring_02_r"]
+        jointToTransferTo = "hand_r"
+        transfer.LOD_transferWeights(meshes, jointsToRemove, jointToTransferTo)
+        jointsToDelete.extend(jointsToRemove)
+
+        jointsToRemove = ["foot_l"]
+        jointToTransferTo = "calf_l"
+        transfer.LOD_transferWeights(meshes, jointsToRemove, jointToTransferTo)
+        jointsToDelete.extend(jointsToRemove)
+
+        jointsToRemove = ["foot_r"]
+        jointToTransferTo = "calf_r"
+        transfer.LOD_transferWeights(meshes, jointsToRemove, jointToTransferTo)
+        jointsToDelete.extend(jointsToRemove)
+
+
+        # Delete the joints we no longer want.
+        for one in jointsToDelete:
+            if cmds.objExists(one):
+                cmds.delete(one)
+                
+        # Export the fbx file.
+        export.exportLOD("Minion_Ranged_LOD2.fbx", "Minion_Ranged_Geo_LOD2")
+
+        # Re-open the file without saving.
+        fullPath = cmds.file(q = True, sceneName = True)
+        cmds.file(fullPath, open=True, f=True)
+
+        print "LOD2 Exported Successfully!!!!!!!!!!!!!!!"
+
+Minion_RangedLOD2()
+
+# Minion_Ranged LOD3 -------------------------------------------------------------------------------------------------------------
+def Minion_RangedLOD3():
+    if saveResult == "Yes":
+
+        cmds.file(save=True)
+    
+        meshes = ["geo_calfs_LOD3", "geo_feet_LOD3", "geo_head_LOD3", "geo_lower_arms_LOD3", "geo_pelvis_LOD3", "geo_shoulder_plates_LOD3", "geo_thighs_LOD3", "geo_upper_arms_LOD3", "geo_upper_torso_LOD3", "geo_weapon_LOD3"]
+
+        jointsToDelete = []
+
+        # remove unneeded joints
+        jointsToRemove = ["shoulder_blade_A_l", "shoulder_blade_B_l", "shoulder_blade_C_l"]
+        jointToTransferTo = "clavicle_l"
+        transfer.LOD_transferWeights(meshes, jointsToRemove, jointToTransferTo)
+        jointsToDelete.extend(jointsToRemove)
+
+        jointsToRemove = ["shoulder_blade_A_r", "shoulder_blade_B_r", "shoulder_blade_C_r"]
+        jointToTransferTo = "clavicle_r"
+        transfer.LOD_transferWeights(meshes, jointsToRemove, jointToTransferTo)
+        jointsToDelete.extend(jointsToRemove)
+
+        jointsToRemove = ["upperarm_twist_01_l"]
+        jointToTransferTo = "upperarm_l"
+        transfer.LOD_transferWeights(meshes, jointsToRemove, jointToTransferTo)
+        jointsToDelete.extend(jointsToRemove)
+
+        jointsToRemove = ["upperarm_twist_01_r"]
+        jointToTransferTo = "upperarm_r"
+        transfer.LOD_transferWeights(meshes, jointsToRemove, jointToTransferTo)
+        jointsToDelete.extend(jointsToRemove)
+
+        jointsToRemove = ["lowerarm_twist_01_l"]
+        jointToTransferTo = "lowerarm_l"
+        transfer.LOD_transferWeights(meshes, jointsToRemove, jointToTransferTo)
+        jointsToDelete.extend(jointsToRemove)
+
+        jointsToRemove = ["lowerarm_twist_01_r"]
+        jointToTransferTo = "lowerarm_r"
+        transfer.LOD_transferWeights(meshes, jointsToRemove, jointToTransferTo)
+        jointsToDelete.extend(jointsToRemove)
+
+        jointsToRemove = ["index_01_l", "index_02_l", "middle_01_l", "middle_02_l", "pinky_01_l", "pinky_02_l", "ring_01_l", "ring_02_l"]
+        jointToTransferTo = "hand_l"
+        transfer.LOD_transferWeights(meshes, jointsToRemove, jointToTransferTo)
+        jointsToDelete.extend(jointsToRemove)
+
+        jointsToRemove = ["index_01_r", "index_02_r", "middle_01_r", "middle_02_r", "pinky_01_r", "pinky_02_r", "ring_01_r", "ring_02_r"]
+        jointToTransferTo = "hand_r"
+        transfer.LOD_transferWeights(meshes, jointsToRemove, jointToTransferTo)
+        jointsToDelete.extend(jointsToRemove)
+
+        jointsToRemove = ["foot_l"]
+        jointToTransferTo = "calf_l"
+        transfer.LOD_transferWeights(meshes, jointsToRemove, jointToTransferTo)
+        jointsToDelete.extend(jointsToRemove)
+
+        jointsToRemove = ["foot_r"]
+        jointToTransferTo = "calf_r"
+        transfer.LOD_transferWeights(meshes, jointsToRemove, jointToTransferTo)
+        jointsToDelete.extend(jointsToRemove)
+
+        # Delete the joints we no longer want.
+        for one in jointsToDelete:
+            if cmds.objExists(one):
+                cmds.delete(one)
+                
+        # Export the fbx file.
+        export.exportLOD("Minion_Ranged_LOD3.fbx", "Minion_Ranged_Geo_LOD3")
+        
+        # Re-open the file without saving.
+        fullPath = cmds.file(q = True, sceneName = True)
+        cmds.file(fullPath, open=True, f=True)
+
+        print "LOD3 Exported Successfully!!!!!!!!!!!!!!!"
+
+Minion_RangedLOD3()
+
diff --git a/Engine/Extras/Maya_AnimationRiggingTools/ARTv1/General/ART/Projects/Orion/Scripts/Minion_Ranged_Post_Script.py b/Engine/Extras/Maya_AnimationRiggingTools/ARTv1/General/ART/Projects/Orion/Scripts/Minion_Ranged_Post_Script.py
new file mode 100644
--- /dev/null
+++ b/Engine/Extras/Maya_AnimationRiggingTools/ARTv1/General/ART/Projects/Orion/Scripts/Minion_Ranged_Post_Script.py
@@ -0,0 +1,178 @@
+import maya.cmds as cmds
+
+# NOTE: This character has special steps to run during Rig Pose creation.  See the following commented out script.
+#########################################
+# RUN THIS COMMAND DURING THE RIG POSE CREATION STEP TO GET THE SHIELD TO THE CORRECT LOCATION.  LEAVE IT COMMENTED OUT SO THAT IT DOESNT RUN DURING RIGGING.
+'''
+# The following 2 lines are used to find the locations and then input below.
+#cmds.xform(q=True, ws=True, rp=True)
+#cmds.xform(q=True, ro=True)
+
+
+cmds.xform("calf_mover_l", ws=True, t=[15.209881960000002, -1.3946871069999982, 32.64604213])
+cmds.xform("calf_mover_l", ro=[-1.6927803365230126e-08, -7.90540084577015e-10, 4.610907371818708e-09])
+
+cmds.xform("calf_mover_r", ws=True, t=[-15.209908050000001, -1.3946831929999999, 32.64604689])
+cmds.xform("calf_mover_r", ro=[-1.692620235351736e-08, -7.906350042557556e-10, 4.611578468856501e-09])
+
+cmds.xform("foot_mover_l", ws=True, t=[17.394773850000004, 4.334110128000001, 11.327026350000002])
+cmds.xform("foot_mover_l", ro=[-27.333076877633392, 10.809974591780163, -0.40751877939002024])
+
+cmds.xform("foot_mover_r", ws=True, t=[-17.394775560000003, 4.334116807000003, 11.32702872])
+cmds.xform("foot_mover_r", ro=[-27.33307687763361, 10.809974591780179, -0.40751877938934844])
+'''
+#########################################
+
+# ------------------------------------------------------------------
+# scale all of the controls to be more visible with the mesh.
+def scaleControls():
+    for side in ["_l", "_r"]:
+        # feet
+        cmds.select(["ik_foot_anim"+side+".cv[0:17]"])
+        cmds.scale(1.5, 1, 1, r = True, ocp = False)
+        cmds.move(0, 6, -5.468891, r = True)
+
+        # hands
+        cmds.select(["ik_wrist"+side+"_anim.cv[0:7]"])
+        cmds.scale(1.9, 1.9, 1.9, r = True, ocp = False)
+
+        # other
+        cmds.select(["weapon"+side+"_anim.cv[0:32]"])
+        cmds.scale(0.238522, 0.238522, 0.238522, r = True, ocp = True)
+        cmds.move(-5.6, 0, 0, r = True, os=True, wd=True)
+
+        cmds.select(["clavicle"+side+"_anim.cv[0:32]"])
+        cmds.move(0, 0, 11, r = True, os=True, wd=True)
+
+    for one in ["shoulder_blade_A_l_anim", "shoulder_blade_C_l_anim", "shoulder_blade_B_l_anim", "shoulder_blade_A_r_anim", "shoulder_blade_C_r_anim", "shoulder_blade_B_r_anim"]:
+        cmds.select([one+".cv[*]"])
+        cmds.scale(0.3, 0.3, 0.3, r = True, ocp = True)
+        cmds.move(0, 0, 0, r = True, os=True, wd=True)
+
+    cmds.select(None)
+
+scaleControls()
+
+# ------------------------------------------------------------------
+def shoulderRig():
+    cmds.setDrivenKeyframe(["shoulder_blade_A_l_anim_grp.translate", "shoulder_blade_B_l_anim_grp.translate", "shoulder_blade_C_l_anim_grp.translate"], cd = "driver_clavicle_l.rotateY", itt = "flat", ott = "flat")
+    cmds.setAttr("clavicle_l_anim.translateZ", 12)
+    cmds.setAttr("shoulder_blade_A_l_anim_grp.translate", 0, 0, -1.23994, type="double3")
+    cmds.setAttr("shoulder_blade_B_l_anim_grp.translate", 0, 0, -2.74608, type="double3")
+    cmds.setAttr("shoulder_blade_C_l_anim_grp.translate", 0, 0, -4.31946, type="double3")
+    cmds.setDrivenKeyframe(["shoulder_blade_A_l_anim_grp.translate", "shoulder_blade_B_l_anim_grp.translate", "shoulder_blade_C_l_anim_grp.translate"], cd = "driver_clavicle_l.rotateY", itt = "flat", ott = "flat")
+    cmds.setAttr("clavicle_l_anim.translateZ", 0)
+
+
+    cmds.setDrivenKeyframe(["shoulder_blade_A_r_anim_grp.translate", "shoulder_blade_B_r_anim_grp.translate", "shoulder_blade_C_r_anim_grp.translate"], cd = "driver_clavicle_r.rotateY", itt = "flat", ott = "flat")
+    cmds.setAttr("clavicle_r_anim.translateZ", 12)
+    cmds.setAttr("shoulder_blade_A_r_anim_grp.translate", 0, 0, 1.23994, type="double3")
+    cmds.setAttr("shoulder_blade_B_r_anim_grp.translate", 0, 0, 2.74608, type="double3")
+    cmds.setAttr("shoulder_blade_C_r_anim_grp.translate", 0, 0, 4.31946, type="double3")
+    cmds.setDrivenKeyframe(["shoulder_blade_A_r_anim_grp.translate", "shoulder_blade_B_r_anim_grp.translate", "shoulder_blade_C_r_anim_grp.translate"], cd = "driver_clavicle_r.rotateY", itt = "flat", ott = "flat")
+    cmds.setAttr("clavicle_r_anim.translateZ", 0)
+
+shoulderRig()
+
+##################################################################################
+##################################################################################
+
+def addCapsule(radius, halfHeight):
+    trueHalfHeight = (halfHeight*2)-(radius*2)
+    capsule = cmds.polyCylinder(name="collision_capsule", r=radius, h=trueHalfHeight, sx=18, sy=1, sz=5, ax=[0,0,1], rcp=True, cuv=3, ch=True)
+    cmds.setAttr(capsule[0]+".tz", halfHeight+2.2)
+    cmds.setAttr(capsule[0]+".rz", 90)
+
+    cmds.addAttr(capsule[0], ln="Radius", dv=radius, keyable=True)
+    cmds.addAttr(capsule[0], ln="HalfHeight", dv=halfHeight, keyable=True)
+
+    expr = capsule[1]+".radius = "+capsule[0]+".Radius;\n"+capsule[1]+".height = ("+capsule[0]+".HalfHeight*2) - ("+capsule[0]+".Radius*2);"
+    cmds.expression(s=expr, o=capsule[1], ae=1, uc="all")
+
+    cmds.setAttr(capsule[0]+".overrideEnabled", 1)
+    cmds.setAttr(capsule[0]+".overrideColor", 18)
+    cmds.setAttr(capsule[0]+".overrideShading", 0)
+
+    cmds.addAttr(capsule[0], ln="Shaded", dv=0, min=0, max=1, keyable=True)
+    cmds.connectAttr(capsule[0]+".Shaded", capsule[0]+".overrideShading")
+
+    cmds.parentConstraint("root", capsule[0], mo=True)
+    cmds.select(capsule[0], r=True)
+    capsuleDisplayLayer = cmds.createDisplayLayer(nr=True, name="Collision Capsule")
+    cmds.setAttr(capsuleDisplayLayer+".displayType", 2)
+    cmds.setAttr(capsuleDisplayLayer+".v", 0)
+    
+    for attr in ["tx", "ty", "tz", "rx", "ry", "rz", "sx", "sy", "sz"]:
+        cmds.setAttr(capsule[0]+"."+attr, lock=True, keyable=False, channelBox=False)
+
+    cmds.select(clear=True)
+addCapsule(42, 85)
+
+# ------------------------------------------------------------------
+# add ik bones Function that is called later after we open the export file.
+def addIKBones():
+    try:
+		# create the joints
+		cmds.select(clear = True)
+		ikFootRoot = cmds.joint(name = "ik_foot_root")
+		cmds.select(clear = True)
+		
+		cmds.select(clear = True)
+		ikFootLeft = cmds.joint(name = "ik_foot_l")
+		cmds.select(clear = True)
+		
+		cmds.select(clear = True)
+		ikFootRight = cmds.joint(name = "ik_foot_r")
+		cmds.select(clear = True)
+		
+		cmds.select(clear = True)
+		ikHandRoot = cmds.joint(name = "ik_hand_root")
+		cmds.select(clear = True)
+		
+		cmds.select(clear = True)
+		ikHandGun = cmds.joint(name = "ik_hand_gun")
+		cmds.select(clear = True)
+		
+		cmds.select(clear = True)
+		ikHandLeft = cmds.joint(name = "ik_hand_l")
+		cmds.select(clear = True)
+		
+		cmds.select(clear = True)
+		ikHandRight = cmds.joint(name = "ik_hand_r")
+		cmds.select(clear = True)
+		
+		# create hierarhcy
+		cmds.parent(ikFootRoot, "root")
+		cmds.parent(ikHandRoot, "root")
+		cmds.parent(ikFootLeft, ikFootRoot)
+		cmds.parent(ikFootRight, ikFootRoot)
+		cmds.parent(ikHandGun, ikHandRoot)
+		cmds.parent(ikHandLeft, ikHandGun)
+		cmds.parent(ikHandRight, ikHandGun)
+		
+		# constrain the joints
+		leftHandConstraint = cmds.parentConstraint("hand_l", ikHandLeft)[0]
+		rightHandGunConstraint = cmds.parentConstraint("hand_r", ikHandGun)[0]
+		rightHandConstraint = cmds.parentConstraint("hand_r", ikHandRight)[0]
+		leftFootConstraint = cmds.parentConstraint("foot_l", ikFootLeft)[0]
+		rightFootConstraint = cmds.parentConstraint("foot_r", ikFootRight)[0]
+	
+    except:
+		cmds.warning("Something went wrong")
+
+# ------------------------------------------------------------------
+# ADD the IK bones to the Export file.
+#save the AnimRig file before opening the export file.
+cmds.file(save = True, type = "mayaBinary", force = True)
+
+#ADD IK BONES
+#open export file
+sceneName = cmds.file(q = True, sceneName = True)
+exportFile = sceneName.partition("AnimRigs")[0] + "ExportFiles" +  sceneName.partition("AnimRigs")[2].partition(".mb")[0] + "_Export.mb"
+cmds.file(exportFile, open = True, force = True)
+
+#add ik bones
+addIKBones()
+
+#save the export file.
+cmds.file(save = True, type = "mayaBinary", force = True)
diff --git a/Engine/Extras/Maya_AnimationRiggingTools/ARTv1/General/ART/Projects/Orion/Scripts/Minion_Super_Post_Script.py b/Engine/Extras/Maya_AnimationRiggingTools/ARTv1/General/ART/Projects/Orion/Scripts/Minion_Super_Post_Script.py
new file mode 100644
--- /dev/null
+++ b/Engine/Extras/Maya_AnimationRiggingTools/ARTv1/General/ART/Projects/Orion/Scripts/Minion_Super_Post_Script.py
@@ -0,0 +1,180 @@
+import maya.cmds as cmds
+import maya.mel as mel
+import math
+import Tools.ART_addSpaceSwitching as space
+reload(space)
+
+# ------------------------------------------------------------------
+# scale all of the controls to be more visible with the mesh.
+def scaleControls():
+    for side in ["_l", "_r"]:
+        # feet
+        cmds.select(["ik_foot_anim"+side+".cv[0:17]"])
+        cmds.scale(1.5, 1, 1, r = True, ocp = False)
+        cmds.move(0, 6, 2.818713, r = True)
+
+        # hands
+        cmds.select(["ik_wrist"+side+"_anim.cv[0:7]"])
+        cmds.scale(1.9, 1.9, 1.9, r = True, ocp = False)
+
+        # other
+        cmds.select(["weapon"+side+"_anim.cv[0:32]"])
+        cmds.scale(0.238522, 0.238522, 0.238522, r = True, ocp = True)
+        cmds.move(-5.6, 0, 0, r = True, os=True, wd=True)
+
+        cmds.select(["clavicle"+side+"_anim.cv[0:32]"])
+        cmds.move(0, 0, 11, r = True, os=True, wd=True)
+
+    for one in ["shoulder_blade_A_l_anim", "shoulder_blade_B_l_anim", "shoulder_blade_A_r_anim", "shoulder_blade_B_r_anim", "shoulder_blade_C_l_anim", "shoulder_blade_C_r_anim"]:
+        cmds.select([one+".cv[*]"])
+        cmds.scale(0.3, 0.3, 0.3, r = True, ocp = True)
+        cmds.move(0, 0, 17, r = True, os=False, wd=True)
+
+    cmds.select(None)
+
+scaleControls()
+
+# ------------------------------------------------------------------
+def shoulderRig():
+    # Left Shoulder
+    cmds.setAttr("Rig_Settings.lArmMode", 0)
+
+    cmds.setAttr("shoulder_blade_A_l_anim_grp.rotate", 0, -5, 0, type="double3")
+    cmds.setAttr("shoulder_blade_B_l_anim_grp.rotate", 0, -13, 0, type="double3")
+    cmds.setDrivenKeyframe(["shoulder_blade_A_l_anim_grp.rotate", "shoulder_blade_B_l_anim_grp.rotate"], cd = "driver_upperarm_l.rotateY", itt = "spline", ott = "spline")
+
+    cmds.setAttr("fk_arm_l_anim.rotateY", -26)
+    cmds.setAttr("shoulder_blade_A_l_anim_grp.rotate", 0, -14, 0, type="double3")
+    cmds.setAttr("shoulder_blade_B_l_anim_grp.rotate", 0, -18, 0, type="double3")
+    cmds.setDrivenKeyframe(["shoulder_blade_A_l_anim_grp.rotate", "shoulder_blade_B_l_anim_grp.rotate"], cd = "driver_upperarm_l.rotateY", itt = "spline", ott = "spline")
+
+    cmds.setAttr("fk_arm_l_anim.rotateY", 51)
+    cmds.setAttr("shoulder_blade_A_l_anim_grp.rotate", 0, 0, 0, type="double3")
+    cmds.setAttr("shoulder_blade_B_l_anim_grp.rotate", 0, 0, 0, type="double3")
+    cmds.setDrivenKeyframe(["shoulder_blade_A_l_anim_grp.rotate", "shoulder_blade_B_l_anim_grp.rotate"], cd = "driver_upperarm_l.rotateY", itt = "spline", ott = "spline")
+
+    cmds.setAttr("Rig_Settings.lArmMode", 1)
+
+    # Right Shoulder
+    cmds.setAttr("Rig_Settings.rArmMode", 0)
+
+    cmds.setAttr("shoulder_blade_A_r_anim_grp.rotate", 0, -5, 0, type="double3")
+    cmds.setAttr("shoulder_blade_B_r_anim_grp.rotate", 0, -13, 0, type="double3")
+    cmds.setDrivenKeyframe(["shoulder_blade_A_r_anim_grp.rotate", "shoulder_blade_B_r_anim_grp.rotate"], cd = "driver_upperarm_r.rotateY", itt = "spline", ott = "spline")
+
+    cmds.setAttr("fk_arm_r_anim.rotateY", -26)
+    cmds.setAttr("shoulder_blade_A_r_anim_grp.rotate", 0, -14, 0, type="double3")
+    cmds.setAttr("shoulder_blade_B_r_anim_grp.rotate", 0, -18, 0, type="double3")
+    cmds.setDrivenKeyframe(["shoulder_blade_A_r_anim_grp.rotate", "shoulder_blade_B_r_anim_grp.rotate"], cd = "driver_upperarm_r.rotateY", itt = "spline", ott = "spline")
+
+    cmds.setAttr("fk_arm_r_anim.rotateY", 51)
+    cmds.setAttr("shoulder_blade_A_r_anim_grp.rotate", 0, 0, 0, type="double3")
+    cmds.setAttr("shoulder_blade_B_r_anim_grp.rotate", 0, 0, 0, type="double3")
+    cmds.setDrivenKeyframe(["shoulder_blade_A_r_anim_grp.rotate", "shoulder_blade_B_r_anim_grp.rotate"], cd = "driver_upperarm_r.rotateY", itt = "spline", ott = "spline")
+
+    cmds.setAttr("Rig_Settings.rArmMode", 1)
+
+shoulderRig()
+
+##################################################################################
+##################################################################################
+
+def addCapsule(radius, halfHeight):
+    trueHalfHeight = (halfHeight*2)-(radius*2)
+    capsule = cmds.polyCylinder(name="collision_capsule", r=radius, h=trueHalfHeight, sx=18, sy=1, sz=5, ax=[0,0,1], rcp=True, cuv=3, ch=True)
+    cmds.setAttr(capsule[0]+".tz", halfHeight+2.2)
+    cmds.setAttr(capsule[0]+".rz", 90)
+
+    cmds.addAttr(capsule[0], ln="Radius", dv=radius, keyable=True)
+    cmds.addAttr(capsule[0], ln="HalfHeight", dv=halfHeight, keyable=True)
+
+    expr = capsule[1]+".radius = "+capsule[0]+".Radius;\n"+capsule[1]+".height = ("+capsule[0]+".HalfHeight*2) - ("+capsule[0]+".Radius*2);"
+    cmds.expression(s=expr, o=capsule[1], ae=1, uc="all")
+
+    cmds.setAttr(capsule[0]+".overrideEnabled", 1)
+    cmds.setAttr(capsule[0]+".overrideColor", 18)
+    cmds.setAttr(capsule[0]+".overrideShading", 0)
+
+    cmds.addAttr(capsule[0], ln="Shaded", dv=0, min=0, max=1, keyable=True)
+    cmds.connectAttr(capsule[0]+".Shaded", capsule[0]+".overrideShading")
+
+    cmds.parentConstraint("root", capsule[0], mo=True)
+    cmds.select(capsule[0], r=True)
+    capsuleDisplayLayer = cmds.createDisplayLayer(nr=True, name="Collision Capsule")
+    cmds.setAttr(capsuleDisplayLayer+".displayType", 2)
+    cmds.setAttr(capsuleDisplayLayer+".v", 0)
+    
+    for attr in ["tx", "ty", "tz", "rx", "ry", "rz", "sx", "sy", "sz"]:
+        cmds.setAttr(capsule[0]+"."+attr, lock=True, keyable=False, channelBox=False)
+
+    cmds.select(clear=True)
+addCapsule(60, 135)
+
+# ------------------------------------------------------------------
+# add ik bones Function that is called later after we open the export file.
+def addIKBones():
+    try:
+		# create the joints
+		cmds.select(clear = True)
+		ikFootRoot = cmds.joint(name = "ik_foot_root")
+		cmds.select(clear = True)
+		
+		cmds.select(clear = True)
+		ikFootLeft = cmds.joint(name = "ik_foot_l")
+		cmds.select(clear = True)
+		
+		cmds.select(clear = True)
+		ikFootRight = cmds.joint(name = "ik_foot_r")
+		cmds.select(clear = True)
+		
+		cmds.select(clear = True)
+		ikHandRoot = cmds.joint(name = "ik_hand_root")
+		cmds.select(clear = True)
+		
+		cmds.select(clear = True)
+		ikHandGun = cmds.joint(name = "ik_hand_gun")
+		cmds.select(clear = True)
+		
+		cmds.select(clear = True)
+		ikHandLeft = cmds.joint(name = "ik_hand_l")
+		cmds.select(clear = True)
+		
+		cmds.select(clear = True)
+		ikHandRight = cmds.joint(name = "ik_hand_r")
+		cmds.select(clear = True)
+		
+		# create hierarhcy
+		cmds.parent(ikFootRoot, "root")
+		cmds.parent(ikHandRoot, "root")
+		cmds.parent(ikFootLeft, ikFootRoot)
+		cmds.parent(ikFootRight, ikFootRoot)
+		cmds.parent(ikHandGun, ikHandRoot)
+		cmds.parent(ikHandLeft, ikHandGun)
+		cmds.parent(ikHandRight, ikHandGun)
+		
+		# constrain the joints
+		leftHandConstraint = cmds.parentConstraint("hand_l", ikHandLeft)[0]
+		rightHandGunConstraint = cmds.parentConstraint("hand_r", ikHandGun)[0]
+		rightHandConstraint = cmds.parentConstraint("hand_r", ikHandRight)[0]
+		leftFootConstraint = cmds.parentConstraint("foot_l", ikFootLeft)[0]
+		rightFootConstraint = cmds.parentConstraint("foot_r", ikFootRight)[0]
+	
+    except:
+		cmds.warning("Something went wrong")
+
+# ------------------------------------------------------------------
+# ADD the IK bones to the Export file.
+#save the AnimRig file before opening the export file.
+cmds.file(save = True, type = "mayaBinary", force = True)
+
+#ADD IK BONES
+#open export file
+sceneName = cmds.file(q = True, sceneName = True)
+exportFile = sceneName.partition("AnimRigs")[0] + "ExportFiles" +  sceneName.partition("AnimRigs")[2].partition(".mb")[0] + "_Export.mb"
+cmds.file(exportFile, open = True, force = True)
+
+#add ik bones
+addIKBones()
+
+#save the export file.
+cmds.file(save = True, type = "mayaBinary", force = True)
diff --git a/Engine/Extras/Maya_AnimationRiggingTools/ARTv1/General/ART/Projects/Orion/Scripts/Muriel_CreateLODs.py b/Engine/Extras/Maya_AnimationRiggingTools/ARTv1/General/ART/Projects/Orion/Scripts/Muriel_CreateLODs.py
new file mode 100644
--- /dev/null
+++ b/Engine/Extras/Maya_AnimationRiggingTools/ARTv1/General/ART/Projects/Orion/Scripts/Muriel_CreateLODs.py
@@ -0,0 +1,556 @@
+import maya.cmds as cmds
+import customMayaMenu as menu
+import Tools.ART_ImportMultipleWeights as impw
+import Tools.ART_TransferWeights as transfer
+import Tools.ART_ExportLODtoFBX as export
+reload(export)
+
+################################################## MURIEL #########################################################################
+# Warn the user that this operation will save their file.
+saveResult = cmds.confirmDialog(title = "Save Warning", message="In order to continue your file must be saved.  Would you like to save it?  If yes it will be saved and after the operation is complete your file will be re-opened.", button = ["Yes", "Cancel"], defaultButton='Yes', cancelButton='Cancel', dismissString='Cancel')
+
+# If they are ok with saving the file, warn them that it will also delete all of the current files in their skin weights directory.
+if saveResult == "Yes":
+    deleteSkinFilesResult = cmds.confirmDialog(title = "SkinWeights Warning", message="In order to continue we must delete all of the files in your ..\General\ART\SkinWeights\ directory.  Would you like to continue?.", button = ["Yes", "Cancel"], defaultButton='Yes', cancelButton='Cancel', dismissString='Cancel')
+    if deleteSkinFilesResult == "Yes":
+        toolsPath = cmds.internalVar(usd = True) + "mayaTools.txt"
+        if os.path.exists(toolsPath):
+            f = open(toolsPath, 'r')
+            mayaToolsDir = f.readline()
+            f.close()
+
+        path = mayaToolsDir + "\General\ART\SkinWeights\\"
+        files = os.listdir(path)
+        #delete all weight cache files in the system folder
+        for file in files:
+            os.remove(path + file)
+            #print path + file
+
+
+
+# Muriel LOD1 -------------------------------------------------------------------------------------------------------------
+def MurielLOD1():
+    if saveResult == "Yes":
+
+        cmds.file(save=True)
+
+        meshes = ["corneas_geo_LOD1", "eye_shadow_geo_LOD1", "eyeballs_geo_LOD1", "Head_geo_LOD1", "Legs_Shell_geo_LOD1", "lowerteeth_geo_LOD1", "Soft_geo_LOD1", "tongue_geo_LOD1", "Torso_Shell_geo_LOD1", "Upperteeth_geo_LOD1", "Wings_Blade_A_geo_LOD1", "Wings_Blade_B_geo_LOD1", "Wings_Blade_C_geo_LOD1", "Wings_Blade_D_geo_LOD1", "Wings_Top_Geo_LOD1"]
+
+        jointsToDelete = []
+
+        # remove unneeded joints
+        jointsToRemove = ["chest_ring"]
+        jointToTransferTo = "spine_03"
+        transfer.LOD_transferWeights(meshes, jointsToRemove, jointToTransferTo)
+        jointsToDelete.extend(jointsToRemove)
+        
+        jointsToRemove = ["eyeball_l", "eyeball_r", "eyelid_bott_l", "eyelid_bott_r", "eyelid_top_l", "eyelid_top_r", "jaw"]
+        jointToTransferTo = "head"
+        transfer.LOD_transferWeights(meshes, jointsToRemove, jointToTransferTo)
+        jointsToDelete.extend(jointsToRemove)
+        
+        # Delete the joints we no longer want.
+        for one in jointsToDelete:
+            if cmds.objExists(one):
+                cmds.delete(one)
+                
+        # Export the fbx file.
+        export.exportLOD("Muriel_LOD1.fbx", "LOD1")
+        
+        # Re-open the file without saving.
+        fullPath = cmds.file(q = True, sceneName = True)
+        cmds.file(fullPath, open=True, f=True)
+
+        print "LOD1 Exported Successfully!!!!!!!!!!!!!!!"
+
+MurielLOD1()
+
+# Muriel LOD2 -------------------------------------------------------------------------------------------------------------
+def MurielLOD2():
+    if saveResult == "Yes":
+
+        cmds.file(save=True)
+        
+        meshes = ["corneas_geo_LOD2", "eye_shadow_geo_LOD2", "eyeballs_geo_LOD2", "Head_geo_LOD2", "Legs_Shell_geo_LOD2", "lowerteeth_geo_LOD2", "Soft_geo_LOD2", "tongue_geo_LOD2", "Torso_Shell_geo_LOD2", "Upperteeth_geo_LOD2", "Wings_Blade_A_geo_LOD2", "Wings_Blade_B_geo_LOD2", "Wings_Blade_C_geo_LOD2", "Wings_Blade_D_geo_LOD2", "Wings_Top_Geo_LOD2"]
+
+        # Export a skin weights file for each mesh we plan to operate on.
+        cmds.select(meshes)
+        menu.exportSelectedSkinWeights(True)
+
+        # Set the pose if needed
+        cmds.setAttr("thumb_03_l.rotateZ", 23.02290762)
+        cmds.setAttr("ring_03_l.rotateZ", 29.40776524)
+        cmds.setAttr("pinky_03_l.rotateZ", 39.74899767)
+        cmds.setAttr("middle_03_l.rotateZ", 23.93256607)
+        cmds.setAttr("index_03_l.rotateZ", 16.68034218)
+
+        cmds.setAttr("thumb_03_r.rotateZ", 23.02290762)
+        cmds.setAttr("ring_03_r.rotateZ", 29.40776524)
+        cmds.setAttr("pinky_03_r.rotateZ", 39.74899767)
+        cmds.setAttr("middle_03_r.rotateZ", 23.93256607)
+        cmds.setAttr("index_03_r.rotateZ", 16.68034218)
+
+        # Delete the history on the meshes so that they get baked down into their posed positions
+        cmds.delete(meshes, ch=True)
+
+        # Move the skeleton back to its bindpose (This script REQUIRES the bind pose to be the pose that we want to export the skeleton to.  This is the case for all known characters created with A.R.T.)
+        cmds.dagPose("root", bp=True, r=True) 
+
+        # Re-import the skinning we exported before onto the mesh so thats its bound in its posed position.
+        cmds.select(meshes)
+        impw.importMultipleSkinWeights("", True)
+
+        jointsToDelete = []
+
+         # remove unneeded joints
+        jointsToRemove = ["chest_ring"]
+        jointToTransferTo = "spine_03"
+        transfer.LOD_transferWeights(meshes, jointsToRemove, jointToTransferTo)
+        jointsToDelete.extend(jointsToRemove)
+
+        jointsToRemove = ["upperarm_twist_01_l"]
+        jointToTransferTo = "upperarm_l"
+        transfer.LOD_transferWeights(meshes, jointsToRemove, jointToTransferTo)
+        jointsToDelete.extend(jointsToRemove)
+
+        jointsToRemove = ["upperarm_twist_01_r"]
+        jointToTransferTo = "upperarm_r"
+        transfer.LOD_transferWeights(meshes, jointsToRemove, jointToTransferTo)
+        jointsToDelete.extend(jointsToRemove)
+
+        jointsToRemove = ["hand_flap_l"]
+        jointToTransferTo = "lowerarm_l"
+        transfer.LOD_transferWeights(meshes, jointsToRemove, jointToTransferTo)
+        jointsToDelete.extend(jointsToRemove)
+
+        jointsToRemove = ["hand_flap_r"]
+        jointToTransferTo = "lowerarm_r"
+        transfer.LOD_transferWeights(meshes, jointsToRemove, jointToTransferTo)
+        jointsToDelete.extend(jointsToRemove)
+
+        jointsToRemove = ["index_03_l"]
+        jointToTransferTo = "index_02_l"
+        transfer.LOD_transferWeights(meshes, jointsToRemove, jointToTransferTo)
+        jointsToDelete.extend(jointsToRemove)
+
+        jointsToRemove = ["middle_03_l"]
+        jointToTransferTo = "middle_02_l"
+        transfer.LOD_transferWeights(meshes, jointsToRemove, jointToTransferTo)
+        jointsToDelete.extend(jointsToRemove)
+
+        jointsToRemove = ["pinky_03_l"]
+        jointToTransferTo = "pinky_02_l"
+        transfer.LOD_transferWeights(meshes, jointsToRemove, jointToTransferTo)
+        jointsToDelete.extend(jointsToRemove)
+
+        jointsToRemove = ["ring_03_l"]
+        jointToTransferTo = "ring_02_l"
+        transfer.LOD_transferWeights(meshes, jointsToRemove, jointToTransferTo)
+        jointsToDelete.extend(jointsToRemove)
+
+        jointsToRemove = ["thumb_03_l"]
+        jointToTransferTo = "thumb_02_l"
+        transfer.LOD_transferWeights(meshes, jointsToRemove, jointToTransferTo)
+        jointsToDelete.extend(jointsToRemove)
+
+        jointsToRemove = ["index_03_r"]
+        jointToTransferTo = "index_02_r"
+        transfer.LOD_transferWeights(meshes, jointsToRemove, jointToTransferTo)
+        jointsToDelete.extend(jointsToRemove)
+
+        jointsToRemove = ["middle_03_r"]
+        jointToTransferTo = "middle_02_r"
+        transfer.LOD_transferWeights(meshes, jointsToRemove, jointToTransferTo)
+        jointsToDelete.extend(jointsToRemove)
+
+        jointsToRemove = ["pinky_03_r"]
+        jointToTransferTo = "pinky_02_r"
+        transfer.LOD_transferWeights(meshes, jointsToRemove, jointToTransferTo)
+        jointsToDelete.extend(jointsToRemove)
+
+        jointsToRemove = ["ring_03_r"]
+        jointToTransferTo = "ring_02_r"
+        transfer.LOD_transferWeights(meshes, jointsToRemove, jointToTransferTo)
+        jointsToDelete.extend(jointsToRemove)
+
+        jointsToRemove = ["thumb_02_r", "thumb_03_r"]
+        jointToTransferTo = "thumb_01_r"
+        transfer.LOD_transferWeights(meshes, jointsToRemove, jointToTransferTo)
+        jointsToDelete.extend(jointsToRemove)
+        
+        jointsToRemove = ["eyeball_l", "eyeball_r", "eyelid_bott_l", "eyelid_bott_r", "eyelid_top_l", "eyelid_top_r", "jaw"]
+        jointToTransferTo = "head"
+        transfer.LOD_transferWeights(meshes, jointsToRemove, jointToTransferTo)
+        jointsToDelete.extend(jointsToRemove)
+
+        jointsToRemove = ["thigh_twist_01_l"]
+        jointToTransferTo = "thigh_l"
+        transfer.LOD_transferWeights(meshes, jointsToRemove, jointToTransferTo)
+        jointsToDelete.extend(jointsToRemove)
+
+        jointsToRemove = ["thigh_twist_01_r"]
+        jointToTransferTo = "thigh_r"
+        transfer.LOD_transferWeights(meshes, jointsToRemove, jointToTransferTo)
+        jointsToDelete.extend(jointsToRemove)
+
+        jointsToRemove = ["foot_flap_l"]
+        jointToTransferTo = "calf_l"
+        transfer.LOD_transferWeights(meshes, jointsToRemove, jointToTransferTo)
+        jointsToDelete.extend(jointsToRemove)
+
+        jointsToRemove = ["foot_flap_r"]
+        jointToTransferTo = "calf_r"
+        transfer.LOD_transferWeights(meshes, jointsToRemove, jointToTransferTo)
+        jointsToDelete.extend(jointsToRemove)
+
+
+        # Delete the joints we no longer want.
+        for one in jointsToDelete:
+            if cmds.objExists(one):
+                cmds.delete(one)
+                
+        # Export the fbx file.
+        export.exportLOD("Muriel_LOD2.fbx", "LOD2")
+
+        # Re-open the file without saving.
+        fullPath = cmds.file(q = True, sceneName = True)
+        cmds.file(fullPath, open=True, f=True)
+
+        print "LOD2 Exported Successfully!!!!!!!!!!!!!!!"
+
+MurielLOD2()
+
+# Muriel LOD3 -------------------------------------------------------------------------------------------------------------
+def MurielLOD3():
+    if saveResult == "Yes":
+
+        cmds.file(save=True)
+    
+        meshes = ["corneas_geo_LOD3", "eye_shadow_geo_LOD3", "eyeballs_geo_LOD3", "Head_geo_LOD3", "Legs_Shell_geo_LOD3", "lowerteeth_geo_LOD3", "Soft_geo_LOD3", "tongue_geo_LOD3", "Torso_Shell_geo_LOD3", "Upperteeth_geo_LOD3", "Wings_Blade_A_geo_LOD3", "Wings_Blade_B_geo_LOD3", "Wings_Blade_C_geo_LOD3", "Wings_Blade_D_geo_LOD3", "Wings_Top_Geo_LOD3"]
+    
+        # Export a skin weights file for each mesh we plan to operate on.
+        cmds.select(meshes)
+        menu.exportSelectedSkinWeights(True)
+
+        # Set the pose if needed
+        cmds.setAttr("thumb_03_l.rotateZ", 23.02290762)
+        cmds.setAttr("thumb_02_l.rotateZ", 10.06960964)
+        cmds.setAttr("ring_03_l.rotateZ", 29.40776524)
+        cmds.setAttr("ring_02_l.rotateZ", 29.40776524)
+        cmds.setAttr("pinky_03_l.rotateZ", 39.74899767)
+        cmds.setAttr("pinky_02_l.rotateZ", 39.74899767)
+        cmds.setAttr("middle_03_l.rotateZ", 23.93256607)
+        cmds.setAttr("middle_02_l.rotateZ", 23.93256607)
+        cmds.setAttr("index_03_l.rotateZ", 16.68034218)
+        cmds.setAttr("index_02_l.rotateZ", 16.68034218)
+
+        cmds.setAttr("thumb_03_r.rotateZ", 23.02290762)
+        cmds.setAttr("thumb_02_r.rotateZ", 10.06960964)
+        cmds.setAttr("ring_03_r.rotateZ", 29.40776524)
+        cmds.setAttr("ring_02_r.rotateZ", 29.40776524)
+        cmds.setAttr("pinky_03_r.rotateZ", 39.74899767)
+        cmds.setAttr("pinky_02_r.rotateZ", 39.74899767)
+        cmds.setAttr("middle_03_r.rotateZ", 23.93256607)
+        cmds.setAttr("middle_02_r.rotateZ", 23.93256607)
+        cmds.setAttr("index_03_r.rotateZ", 16.68034218)
+        cmds.setAttr("index_02_r.rotateZ", 16.68034218)
+
+        # Delete the history on the meshes so that they get baked down into their posed positions
+        cmds.delete(meshes, ch=True)
+
+        # Move the skeleton back to its bindpose (This script REQUIRES the bind pose to be the pose that we want to export the skeleton to.  This is the case for all known characters created with A.R.T.)
+        cmds.dagPose("root", bp=True, r=True) 
+
+        # Re-import the skinning we exported before onto the mesh so thats its bound in its posed position.
+        cmds.select(meshes)
+        impw.importMultipleSkinWeights("", True)
+
+        jointsToDelete = []
+
+        # remove unneeded joints
+        jointsToRemove = ["chest_ring"]
+        jointToTransferTo = "spine_03"
+        transfer.LOD_transferWeights(meshes, jointsToRemove, jointToTransferTo)
+        jointsToDelete.extend(jointsToRemove)
+
+        jointsToRemove = ["wing_blade_D_l", "wing_blade_A_l", "wing_blade_B_l", "wing_blade_C_l"]
+        jointToTransferTo = "wing_blade_track_l"
+        transfer.LOD_transferWeights(meshes, jointsToRemove, jointToTransferTo)
+        jointsToDelete.extend(jointsToRemove)
+
+        jointsToRemove = ["wing_blade_A_r", "wing_blade_B_r", "wing_blade_C_r", "wing_blade_D_r"]
+        jointToTransferTo = "wing_blade_track_r"
+        transfer.LOD_transferWeights(meshes, jointsToRemove, jointToTransferTo)
+        jointsToDelete.extend(jointsToRemove)
+
+        jointsToRemove = ["upperarm_twist_01_l", "upperarm_twist_02_l"]
+        jointToTransferTo = "upperarm_l"
+        transfer.LOD_transferWeights(meshes, jointsToRemove, jointToTransferTo)
+        jointsToDelete.extend(jointsToRemove)
+
+        jointsToRemove = ["upperarm_twist_01_r", "upperarm_twist_02_r"]
+        jointToTransferTo = "upperarm_r"
+        transfer.LOD_transferWeights(meshes, jointsToRemove, jointToTransferTo)
+        jointsToDelete.extend(jointsToRemove)
+
+        jointsToRemove = ["hand_flap_l"]
+        jointToTransferTo = "lowerarm_l"
+        transfer.LOD_transferWeights(meshes, jointsToRemove, jointToTransferTo)
+        jointsToDelete.extend(jointsToRemove)
+
+        jointsToRemove = ["hand_flap_r"]
+        jointToTransferTo = "lowerarm_r"
+        transfer.LOD_transferWeights(meshes, jointsToRemove, jointToTransferTo)
+        jointsToDelete.extend(jointsToRemove)
+
+        jointsToRemove = ["index_02_l", "index_03_l"]
+        jointToTransferTo = "index_01_l"
+        transfer.LOD_transferWeights(meshes, jointsToRemove, jointToTransferTo)
+        jointsToDelete.extend(jointsToRemove)
+
+        jointsToRemove = ["middle_02_l", "middle_03_l"]
+        jointToTransferTo = "middle_01_l"
+        transfer.LOD_transferWeights(meshes, jointsToRemove, jointToTransferTo)
+        jointsToDelete.extend(jointsToRemove)
+
+        jointsToRemove = ["pinky_02_l", "pinky_03_l"]
+        jointToTransferTo = "pinky_01_l"
+        transfer.LOD_transferWeights(meshes, jointsToRemove, jointToTransferTo)
+        jointsToDelete.extend(jointsToRemove)
+
+        jointsToRemove = ["ring_02_l", "ring_03_l"]
+        jointToTransferTo = "ring_01_l"
+        transfer.LOD_transferWeights(meshes, jointsToRemove, jointToTransferTo)
+        jointsToDelete.extend(jointsToRemove)
+
+        jointsToRemove = ["thumb_02_l", "thumb_03_l"]
+        jointToTransferTo = "thumb_01_l"
+        transfer.LOD_transferWeights(meshes, jointsToRemove, jointToTransferTo)
+        jointsToDelete.extend(jointsToRemove)
+
+        jointsToRemove = ["index_02_r", "index_03_r"]
+        jointToTransferTo = "index_01_r"
+        transfer.LOD_transferWeights(meshes, jointsToRemove, jointToTransferTo)
+        jointsToDelete.extend(jointsToRemove)
+
+        jointsToRemove = ["middle_02_r", "middle_03_r"]
+        jointToTransferTo = "middle_01_r"
+        transfer.LOD_transferWeights(meshes, jointsToRemove, jointToTransferTo)
+        jointsToDelete.extend(jointsToRemove)
+
+        jointsToRemove = ["pinky_02_r", "pinky_03_r"]
+        jointToTransferTo = "pinky_01_r"
+        transfer.LOD_transferWeights(meshes, jointsToRemove, jointToTransferTo)
+        jointsToDelete.extend(jointsToRemove)
+
+        jointsToRemove = ["ring_02_r", "ring_03_r"]
+        jointToTransferTo = "ring_01_r"
+        transfer.LOD_transferWeights(meshes, jointsToRemove, jointToTransferTo)
+        jointsToDelete.extend(jointsToRemove)
+
+        jointsToRemove = ["thumb_02_r", "thumb_03_r"]
+        jointToTransferTo = "thumb_01_r"
+        transfer.LOD_transferWeights(meshes, jointsToRemove, jointToTransferTo)
+        jointsToDelete.extend(jointsToRemove)
+
+        jointsToRemove = ["eyeball_l", "eyeball_r", "eyelid_bott_l", "eyelid_bott_r", "eyelid_top_l", "eyelid_top_r", "jaw"]
+        jointToTransferTo = "head"
+        transfer.LOD_transferWeights(meshes, jointsToRemove, jointToTransferTo)
+        jointsToDelete.extend(jointsToRemove)
+
+        jointsToRemove = ["thigh_twist_01_l", "thigh_twist_02_l"]
+        jointToTransferTo = "thigh_l"
+        transfer.LOD_transferWeights(meshes, jointsToRemove, jointToTransferTo)
+        jointsToDelete.extend(jointsToRemove)
+
+        jointsToRemove = ["thigh_twist_01_r", "thigh_twist_02_r"]
+        jointToTransferTo = "thigh_r"
+        transfer.LOD_transferWeights(meshes, jointsToRemove, jointToTransferTo)
+        jointsToDelete.extend(jointsToRemove)
+
+        jointsToRemove = ["foot_flap_l"]
+        jointToTransferTo = "calf_l"
+        transfer.LOD_transferWeights(meshes, jointsToRemove, jointToTransferTo)
+        jointsToDelete.extend(jointsToRemove)
+
+        jointsToRemove = ["foot_flap_r"]
+        jointToTransferTo = "calf_r"
+        transfer.LOD_transferWeights(meshes, jointsToRemove, jointToTransferTo)
+        jointsToDelete.extend(jointsToRemove)
+
+        # Delete the joints we no longer want.
+        for one in jointsToDelete:
+            if cmds.objExists(one):
+                cmds.delete(one)
+                
+        # Export the fbx file.
+        export.exportLOD("Muriel_LOD3.fbx", "LOD3")
+        
+        # Re-open the file without saving.
+        fullPath = cmds.file(q = True, sceneName = True)
+        cmds.file(fullPath, open=True, f=True)
+
+        print "LOD3 Exported Successfully!!!!!!!!!!!!!!!"
+
+MurielLOD3()
+
+# Muriel LOD4 -------------------------------------------------------------------------------------------------------------
+def MurielLOD4():
+    if saveResult == "Yes":
+
+        cmds.file(save=True)
+    
+        meshes = ["corneas_geo_LOD4", "eye_shadow_geo_LOD4", "eyeballs_geo_LOD4", "Head_geo_LOD4", "Legs_Shell_geo_LOD4", "lowerteeth_geo_LOD4", "Soft_geo_LOD4", "tongue_geo_LOD4", "Torso_Shell_geo_LOD4", "Upperteeth_geo_LOD4", "Wings_Blade_A_geo_LOD4", "Wings_Blade_B_geo_LOD4", "Wings_Blade_C_geo_LOD4", "Wings_Blade_D_geo_LOD4", "Wings_Top_Geo_LOD4"]
+
+        # Export a skin weights file for each mesh we plan to operate on.
+        cmds.select(meshes)
+        menu.exportSelectedSkinWeights(True)
+
+        # Set the pose if needed
+        cmds.setAttr("thumb_03_l.rotateZ", 23.02290762)
+        cmds.setAttr("thumb_02_l.rotateZ", 10.06960964)
+        cmds.setAttr("thumb_01_l.rotateZ", 8.338113921)
+        cmds.setAttr("thumb_01_l.rotateY", 3.002087702)
+        cmds.setAttr("thumb_01_l.rotateX", -3.231699475)
+        cmds.setAttr("ring_03_l.rotateZ", 29.40776524)
+        cmds.setAttr("ring_02_l.rotateZ", 29.40776524)
+        cmds.setAttr("ring_01_l.rotateZ", 29.40776524)
+        cmds.setAttr("pinky_03_l.rotateZ", 39.74899767)
+        cmds.setAttr("pinky_02_l.rotateZ", 39.74899767)
+        cmds.setAttr("pinky_01_l.rotateZ", 39.74899767)
+        cmds.setAttr("middle_03_l.rotateZ", 23.93256607)
+        cmds.setAttr("middle_02_l.rotateZ", 23.93256607)
+        cmds.setAttr("middle_01_l.rotateZ", 23.93256607)
+        cmds.setAttr("index_03_l.rotateZ", 16.68034218)
+        cmds.setAttr("index_02_l.rotateZ", 16.68034218)
+        cmds.setAttr("index_01_l.rotateZ", 16.68034218)
+
+        cmds.setAttr("thumb_03_r.rotateZ", 23.02290762)
+        cmds.setAttr("thumb_02_r.rotateZ", 10.06960964)
+        cmds.setAttr("thumb_01_r.rotateZ", 8.338113921)
+        cmds.setAttr("thumb_01_r.rotateY", 3.002087702)
+        cmds.setAttr("thumb_01_r.rotateX", -3.231699475)
+        cmds.setAttr("ring_03_r.rotateZ", 29.40776524)
+        cmds.setAttr("ring_02_r.rotateZ", 29.40776524)
+        cmds.setAttr("ring_01_r.rotateZ", 29.40776524)
+        cmds.setAttr("pinky_03_r.rotateZ", 39.74899767)
+        cmds.setAttr("pinky_02_r.rotateZ", 39.74899767)
+        cmds.setAttr("pinky_01_r.rotateZ", 39.74899767)
+        cmds.setAttr("middle_03_r.rotateZ", 23.93256607)
+        cmds.setAttr("middle_02_r.rotateZ", 23.93256607)
+        cmds.setAttr("middle_01_r.rotateZ", 23.93256607)
+        cmds.setAttr("index_03_r.rotateZ", 16.68034218)
+        cmds.setAttr("index_02_r.rotateZ", 16.68034218)
+        cmds.setAttr("index_01_r.rotateZ", 16.68034218)
+
+        # Delete the history on the meshes so that they get baked down into their posed positions
+        cmds.delete(meshes, ch=True)
+
+        # Move the skeleton back to its bindpose (This script REQUIRES the bind pose to be the pose that we want to export the skeleton to.  This is the case for all known characters created with A.R.T.)
+        cmds.dagPose("root", bp=True, r=True) 
+
+        # Re-import the skinning we exported before onto the mesh so thats its bound in its posed position.
+        cmds.select(meshes)
+        impw.importMultipleSkinWeights("", True)
+        
+        jointsToDelete = []
+
+        # remove unneeded joints
+        jointsToRemove = ["chest_ring"]
+        jointToTransferTo = "spine_03"
+        transfer.LOD_transferWeights(meshes, jointsToRemove, jointToTransferTo)
+        jointsToDelete.extend(jointsToRemove)
+
+        jointsToRemove = ["wing_blade_D_l", "wing_blade_A_l", "wing_blade_B_l", "wing_blade_C_l"]
+        jointToTransferTo = "wing_blade_track_l"
+        transfer.LOD_transferWeights(meshes, jointsToRemove, jointToTransferTo)
+        jointsToDelete.extend(jointsToRemove)
+
+        jointsToRemove = ["wing_blade_A_r", "wing_blade_B_r", "wing_blade_C_r", "wing_blade_D_r"]
+        jointToTransferTo = "wing_blade_track_r"
+        transfer.LOD_transferWeights(meshes, jointsToRemove, jointToTransferTo)
+        jointsToDelete.extend(jointsToRemove)
+
+        jointsToRemove = ["upperarm_twist_01_l", "upperarm_twist_02_l"]
+        jointToTransferTo = "upperarm_l"
+        transfer.LOD_transferWeights(meshes, jointsToRemove, jointToTransferTo)
+        jointsToDelete.extend(jointsToRemove)
+
+        jointsToRemove = ["upperarm_twist_01_r", "upperarm_twist_02_r"]
+        jointToTransferTo = "upperarm_r"
+        transfer.LOD_transferWeights(meshes, jointsToRemove, jointToTransferTo)
+        jointsToDelete.extend(jointsToRemove)
+
+        jointsToRemove = ["hand_flap_l"]
+        jointToTransferTo = "lowerarm_l"
+        transfer.LOD_transferWeights(meshes, jointsToRemove, jointToTransferTo)
+        jointsToDelete.extend(jointsToRemove)
+
+        jointsToRemove = ["hand_flap_r"]
+        jointToTransferTo = "lowerarm_r"
+        transfer.LOD_transferWeights(meshes, jointsToRemove, jointToTransferTo)
+        jointsToDelete.extend(jointsToRemove)
+
+        jointsToRemove = ["index_01_l", "index_02_l", "index_03_l", "middle_01_l", "middle_02_l", "middle_03_l", "pinky_01_l", "pinky_02_l", "pinky_03_l", "ring_01_l", "ring_02_l", "ring_03_l", "thumb_01_l", "thumb_02_l", "thumb_03_l"]
+        jointToTransferTo = "hand_l"
+        transfer.LOD_transferWeights(meshes, jointsToRemove, jointToTransferTo)
+        jointsToDelete.extend(jointsToRemove)
+
+        jointsToRemove = ["index_01_r", "index_02_r", "index_03_r", "middle_01_r", "middle_02_r", "middle_03_r", "pinky_01_r", "pinky_02_r", "pinky_03_r", "ring_01_r", "ring_02_r", "ring_03_r", "thumb_01_r", "thumb_02_r", "thumb_03_r"]
+        jointToTransferTo = "hand_r"
+        transfer.LOD_transferWeights(meshes, jointsToRemove, jointToTransferTo)
+        jointsToDelete.extend(jointsToRemove)
+
+        jointsToRemove = ["eyeball_l", "eyeball_r", "eyelid_bott_l", "eyelid_bott_r", "eyelid_top_l", "eyelid_top_r", "jaw"]
+        jointToTransferTo = "head"
+        transfer.LOD_transferWeights(meshes, jointsToRemove, jointToTransferTo)
+        jointsToDelete.extend(jointsToRemove)
+
+        jointsToRemove = ["thigh_twist_01_l", "thigh_twist_02_l"]
+        jointToTransferTo = "thigh_l"
+        transfer.LOD_transferWeights(meshes, jointsToRemove, jointToTransferTo)
+        jointsToDelete.extend(jointsToRemove)
+
+        jointsToRemove = ["thigh_twist_01_r", "thigh_twist_02_r"]
+        jointToTransferTo = "thigh_r"
+        transfer.LOD_transferWeights(meshes, jointsToRemove, jointToTransferTo)
+        jointsToDelete.extend(jointsToRemove)
+
+        jointsToRemove = ["foot_flap_l"]
+        jointToTransferTo = "calf_l"
+        transfer.LOD_transferWeights(meshes, jointsToRemove, jointToTransferTo)
+        jointsToDelete.extend(jointsToRemove)
+
+        jointsToRemove = ["foot_flap_r"]
+        jointToTransferTo = "calf_r"
+        transfer.LOD_transferWeights(meshes, jointsToRemove, jointToTransferTo)
+        jointsToDelete.extend(jointsToRemove)
+
+        jointsToRemove = ["ball_l"]
+        jointToTransferTo = "foot_l"
+        transfer.LOD_transferWeights(meshes, jointsToRemove, jointToTransferTo)
+        jointsToDelete.extend(jointsToRemove)
+
+        jointsToRemove = ["ball_r"]
+        jointToTransferTo = "foot_r"
+        transfer.LOD_transferWeights(meshes, jointsToRemove, jointToTransferTo)
+        jointsToDelete.extend(jointsToRemove)
+
+
+        # Delete the joints we no longer want.
+        for one in jointsToDelete:
+            if cmds.objExists(one):
+                cmds.delete(one)
+                
+        # Export the fbx file.
+        export.exportLOD("Muriel_LOD4.fbx", "LOD4")
+        
+        # Re-open the file without saving.
+        fullPath = cmds.file(q = True, sceneName = True)
+        cmds.file(fullPath, open=True, f=True)
+
+        print "LOD4 Exported Successfully!!!!!!!!!!!!!!!"
+
+MurielLOD4()
+
+
diff --git a/Engine/Extras/Maya_AnimationRiggingTools/ARTv1/General/ART/Projects/Orion/Scripts/Muriel_Post_Script.py b/Engine/Extras/Maya_AnimationRiggingTools/ARTv1/General/ART/Projects/Orion/Scripts/Muriel_Post_Script.py
new file mode 100644
--- /dev/null
+++ b/Engine/Extras/Maya_AnimationRiggingTools/ARTv1/General/ART/Projects/Orion/Scripts/Muriel_Post_Script.py
@@ -0,0 +1,265 @@
+import maya.cmds as cmds
+
+# ------------------------------------------------------------------
+# New wing control shapes
+lWingControls = ["wing_blade_A_l_anim", "wing_blade_B_l_anim", "wing_blade_C_l_anim"]
+rWingControls = ["wing_blade_A_r_anim", "wing_blade_B_r_anim", "wing_blade_C_r_anim"]
+for list in [lWingControls, rWingControls]:
+	val = 0
+	for each in list:
+		#create new curve shape
+		control = cmds.curve(d = 1, p = [(-0.124602, 0, 1.096506), (-0.975917, 0, 1.036319), (-0.560906, 0, 0.946236), (-0.798049, 0, 0.798033), (-1.042985, 0, 0.430511), (-1.128122, 0, -0.00276108), (-1.042702, 0, -0.431934), (-0.798049, 0, -0.798033), (-0.559059, 0, -0.944259), (-0.975917, 0, -1.036319), (-0.124602, 0, -1.096506), (-0.537718, 0, -0.349716), (-0.439199, 0, -0.785581), (-0.652776, 0, -0.652998), (-0.853221, 0, -0.353358), (-0.923366, 0, 0), (-0.856561, 0, 0.336531), (-0.658494, 0, 0.644451), (-0.452392, 0, 0.781229), (-0.537718, 0, 0.349716), (-0.124602, 0, 1.096506)], k = [0, 1, 2, 3, 4, 5, 6, 7, 8, 9, 10, 11, 12, 13, 14, 15, 16, 17, 18, 19, 20])
+		cmds.setAttr(control + ".scaleX", 18)
+		cmds.setAttr(control + ".scaleY", 18)
+		cmds.setAttr(control + ".scaleZ", 18)
+		
+		if each in lWingControls:
+			cmds.setAttr(control + ".rotateZ", 90)
+			cmds.setAttr(control + ".translateX", val)
+		if each in rWingControls:
+			cmds.setAttr(control + ".rotateZ", -90)
+			cmds.setAttr(control + ".translateX", val * -1)
+			
+		val = val + 4
+			
+		cmds.makeIdentity(control, apply = True, t =1, r =1, s =1)
+		parentShape = cmds.listRelatives(each, f = True, s = True)[0]
+		cmds.delete(parentShape)
+		shape = cmds.listRelatives(control, f = True, s = True)[0]
+		cmds.parent(shape, each, add = True, shape = True)
+		
+		cmds.delete(control)
+		
+		
+		
+# ------------------------------------------------------------------
+# D Wing
+lWingControls = ["wing_blade_D_l_anim"]
+rWingControls = ["wing_blade_D_r_anim"]
+for list in [lWingControls, rWingControls]:
+	for each in list:
+		#create new curve shape
+		control = cmds.curve(d = 1, p = [(-0.124602, 0, 1.096506), (-0.975917, 0, 1.036319), (-0.560906, 0, 0.946236), (-0.798049, 0, 0.798033), (-1.042985, 0, 0.430511), (-1.128122, 0, -0.00276108), (-1.042702, 0, -0.431934), (-0.798049, 0, -0.798033), (-0.559059, 0, -0.944259), (-0.975917, 0, -1.036319), (-0.124602, 0, -1.096506), (-0.537718, 0, -0.349716), (-0.439199, 0, -0.785581), (-0.652776, 0, -0.652998), (-0.853221, 0, -0.353358), (-0.923366, 0, 0), (-0.856561, 0, 0.336531), (-0.658494, 0, 0.644451), (-0.452392, 0, 0.781229), (-0.537718, 0, 0.349716), (-0.124602, 0, 1.096506)], k = [0, 1, 2, 3, 4, 5, 6, 7, 8, 9, 10, 11, 12, 13, 14, 15, 16, 17, 18, 19, 20])
+		cmds.setAttr(control + ".scaleX", 18)
+		cmds.setAttr(control + ".scaleY", 18)
+		cmds.setAttr(control + ".scaleZ", 18)
+		
+		if each in lWingControls:
+			cmds.setAttr(control + ".rotateZ", -90)
+		if each in rWingControls:
+			cmds.setAttr(control + ".rotateZ", 90)
+			
+		cmds.makeIdentity(control, apply = True, t =1, r =1, s =1)
+		parentShape = cmds.listRelatives(each, f = True, s = True)[0]
+		cmds.delete(parentShape)
+		shape = cmds.listRelatives(control, f = True, s = True)[0]
+		cmds.parent(shape, each, add = True, shape = True)
+		
+		cmds.delete(control)
+		
+		
+# ------------------------------------------------------------------
+# Hand Sizes
+
+cmds.select(["hand_flap_l_anim.cv[0:32]", "hand_flap_r_anim.cv[0:32]"])
+cmds.scale(0.294922, 0.294922, 0.294922, r = True, ocp = True)
+
+    
+cmds.select(["wing_blade_B_l_anim.cv[0:32]", "wing_blade_A_r_anim.cv[0:32]", "wing_ext_l_anim.cv[0:32]", "wing_ext_r_anim.cv[0:32]", "wing_blade_main_l_anim.cv[0:32]", "chest_ring_anim.cv[0:32]", "wing_blade_A_l_anim.cv[0:32]", "wing_blade_B_r_anim.cv[0:32]", "wing_main_r_anim.cv[0:32]", "wing_blade_C_l_anim.cv[0:32]", "wing_blade_C_r_anim.cv[0:32]", "wing_blade_D_l_anim.cv[0:32]", "wing_blade_main_r_anim.cv[0:32]", "wing_blade_D_r_anim.cv[0:32]", "wing_main_l_anim.cv[0:32]"])
+cmds.scale(0.55, 0.55, 0.55, r = True, ocp = True)
+
+cmds.select(["foot_flap_l_anim.cv[0:32]", "foot_flap_r_anim.cv[0:32]"])
+cmds.scale(0.248359, 0.248359, 0.248359, r = True, ocp = True)
+
+# ------------------------------------------------------------------
+def scaleControls():
+    sides = ["_l", "_r"]
+    for side in sides:
+        # feet
+        cmds.select(["ik_foot_anim"+side+".cv[0:17]"])
+        cmds.scale(1.5, 0.8, 1.0, r = True, ocp = False)
+        cmds.move(0, 4.678845, -1.187034, r = True)
+
+        cmds.select(["ik_wrist"+side+"_anim.cv[0:7]"])
+        cmds.scale(1.9, 1.9, 1.9, r = True, ocp = False)
+
+        cmds.select(["eyelid_bott"+side+"_anim.cv[0:32]"])
+        cmds.scale(0.1, 0.1, 0.01, r = True, ocp = False)
+        cmds.move(0, -3.757386, 0.654023, r = True)
+
+        cmds.select(["eyelid_top"+side+"_anim.cv[0:32]"])
+        cmds.scale(0.1, 0.1, 0.01, r = True, ocp = False)
+        cmds.move(0, -3.757386, 0.654023, r = True)
+
+    # Other
+    cmds.select(["eyeball_l_anim.cv[0:32]"])
+    cmds.scale(0.1, 0.1, 0.05, r = True, ocp = False)
+    cmds.move(0, -3.757386, 0, r = True)
+
+    cmds.select(["eyeball_r_anim.cv[0:32]"])
+    cmds.scale(0.1, 0.1, 0.05, r = True, ocp = False)
+    cmds.move(0, 3.757386, 0, r = True)
+
+    cmds.select(["jaw_anim.cv[0:32]"])
+    cmds.scale(0.2, 0.2, 0.2, r = True, ocp = True)
+    cmds.move(0, -6, -8, r = True)
+
+    cmds.select(None)
+
+scaleControls()
+
+# ------------------------------------------------------------------
+def invertEyelids():
+    sides = ["_l", "_r"]
+    for side in sides:
+        cmds.delete("driver_eyelid_bott"+side+"_parentConstraint1")
+        cmds.rotate(0, -180, 0, "eyelid_bott"+side+"_anim_grp", r=True, os=True)
+        cmds.parentConstraint("eyelid_bott"+side+"_anim", "driver_eyelid_bott"+side, mo=True)
+invertEyelids()
+
+def invertRightEyeball():
+    cmds.delete("driver_eyeball_r_parentConstraint1")
+    cmds.rotate(0, -180, -180, "eyeball_r_anim_grp", r=True, os=True)
+    cmds.parentConstraint("eyeball_r_anim", "driver_eyeball_r", mo=True)
+invertRightEyeball()
+
+# ------------------------------------------------------------------
+def createGameCam():
+    gameCam = cmds.camera(centerOfInterest=5, focalLength=18, lensSqueezeRatio=1, cameraScale=1, horizontalFilmAperture=1.41732, horizontalFilmOffset=0, verticalFilmAperture=0.94488, verticalFilmOffset=0, filmFit="Overscan", overscan=1.05, motionBlur=0, shutterAngle=144, nearClipPlane=1.0, farClipPlane=10000, orthographic=0, orthographicWidth=30, panZoomEnabled=0, horizontalPan=0, verticalPan=0, zoom=1, displayFilmGate=True)[0]
+    print gameCam
+    gameCam = cmds.rename(gameCam, "GameCam")
+    cmds.camera(gameCam, e=True, horizontalFilmAperture=1.6724376, focalLength=21.2399857011)
+    cmds.setAttr(gameCam+".displayGateMaskOpacity", 1)
+    cmds.setAttr(gameCam+".displayGateMaskColor", 0, 0, 0, type="double3")
+    cmds.setAttr(gameCam+".r", 90, 0, 180, type="double3")
+    
+    # Create the control
+    control = cmds.circle(name=gameCam+"_Ctrl", ch=True, o=True, r=15)[0]
+    cmds.makeIdentity(control, a=True, t=True, r=True, s=True)
+    # Create groups for the sdk and the constraints
+    const_grp = cmds.group(control, name=control+"_grp")
+    if cmds.objExists("rig_grp"):
+        cmds.parent(const_grp, "rig_grp")
+    # Parent Constrain the cam to the control
+    cmds.parentConstraint(control, gameCam, mo=True)
+    
+    cmds.connectAttr(control+".visibility", gameCam+".visibility", f=True)
+    cmds.setAttr(const_grp+".v", 0)
+    
+    cmds.setAttr(const_grp+".t", 0, 500, 200, type="double3")
+createGameCam()
+# ------------------------------------------------------------------
+
+##################################################################################
+##################################################################################
+
+def addCapsule(radius, halfHeight):
+    trueHalfHeight = (halfHeight*2)-(radius*2)
+    capsule = cmds.polyCylinder(name="collision_capsule", r=radius, h=trueHalfHeight, sx=18, sy=1, sz=5, ax=[0,0,1], rcp=True, cuv=3, ch=True)
+    cmds.setAttr(capsule[0]+".tz", halfHeight+2.2)
+    cmds.setAttr(capsule[0]+".rz", 90)
+
+    cmds.addAttr(capsule[0], ln="Radius", dv=radius, keyable=True)
+    cmds.addAttr(capsule[0], ln="HalfHeight", dv=halfHeight, keyable=True)
+
+    expr = capsule[1]+".radius = "+capsule[0]+".Radius;\n"+capsule[1]+".height = ("+capsule[0]+".HalfHeight*2) - ("+capsule[0]+".Radius*2);"
+    cmds.expression(s=expr, o=capsule[1], ae=1, uc="all")
+
+    cmds.setAttr(capsule[0]+".overrideEnabled", 1)
+    cmds.setAttr(capsule[0]+".overrideColor", 18)
+    cmds.setAttr(capsule[0]+".overrideShading", 0)
+
+    cmds.addAttr(capsule[0], ln="Shaded", dv=0, min=0, max=1, keyable=True)
+    cmds.connectAttr(capsule[0]+".Shaded", capsule[0]+".overrideShading")
+
+    cmds.parentConstraint("root", capsule[0], mo=True)
+    cmds.select(capsule[0], r=True)
+    capsuleDisplayLayer = cmds.createDisplayLayer(nr=True, name="Collision Capsule")
+    cmds.setAttr(capsuleDisplayLayer+".displayType", 2)
+    cmds.setAttr(capsuleDisplayLayer+".v", 0)
+    
+    for attr in ["tx", "ty", "tz", "rx", "ry", "rz", "sx", "sy", "sz"]:
+        cmds.setAttr(capsule[0]+"."+attr, lock=True, keyable=False, channelBox=False)
+
+    cmds.select(clear=True)
+addCapsule(42, 106)
+
+# ------------------------------------------------------------------
+# add ik bones Function that is called later after we open the export file.
+def addIKBones():
+    try:
+		# create the joints
+		cmds.select(clear = True)
+		ikFootRoot = cmds.joint(name = "ik_foot_root")
+		cmds.select(clear = True)
+		
+		cmds.select(clear = True)
+		ikFootLeft = cmds.joint(name = "ik_foot_l")
+		cmds.select(clear = True)
+		
+		cmds.select(clear = True)
+		ikFootRight = cmds.joint(name = "ik_foot_r")
+		cmds.select(clear = True)
+		
+		cmds.select(clear = True)
+		ikHandRoot = cmds.joint(name = "ik_hand_root")
+		cmds.select(clear = True)
+		
+		cmds.select(clear = True)
+		ikHandGun = cmds.joint(name = "ik_hand_gun")
+		cmds.select(clear = True)
+		
+		cmds.select(clear = True)
+		ikHandLeft = cmds.joint(name = "ik_hand_l")
+		cmds.select(clear = True)
+		
+		cmds.select(clear = True)
+		ikHandRight = cmds.joint(name = "ik_hand_r")
+		cmds.select(clear = True)
+		
+		# create hierarhcy
+		cmds.parent(ikFootRoot, "root")
+		cmds.parent(ikHandRoot, "root")
+		cmds.parent(ikFootLeft, ikFootRoot)
+		cmds.parent(ikFootRight, ikFootRoot)
+		cmds.parent(ikHandGun, ikHandRoot)
+		cmds.parent(ikHandLeft, ikHandGun)
+		cmds.parent(ikHandRight, ikHandGun)
+		
+		# constrain the joints
+		leftHandConstraint = cmds.parentConstraint("hand_l", ikHandLeft)[0]
+		rightHandGunConstraint = cmds.parentConstraint("hand_r", ikHandGun)[0]
+		rightHandConstraint = cmds.parentConstraint("hand_r", ikHandRight)[0]
+		leftFootConstraint = cmds.parentConstraint("foot_l", ikFootLeft)[0]
+		rightFootConstraint = cmds.parentConstraint("foot_r", ikFootRight)[0]
+	
+    except:
+		cmds.warning("Something went wrong")
+
+# ------------------------------------------------------------------
+# ADD the IK bones to the Export file.
+#save the AnimRig file before opening the export file.
+cmds.file(save = True, type = "mayaBinary", force = True)
+
+#ADD IK BONES
+#open export file
+sceneName = cmds.file(q = True, sceneName = True)
+exportFile = sceneName.partition("AnimRigs")[0] + "ExportFiles" +  sceneName.partition("AnimRigs")[2].partition(".mb")[0] + "_Export.mb"
+cmds.file(exportFile, open = True, force = True)
+
+#add ik bones
+addIKBones()
+
+# ------------------------------------------------------------------
+def addCustomAttrs():
+    attrs = cmds.listAttr("root")
+    exists = False
+    for one in attrs:
+        if one.find("DistanceCurve") != -1:
+            exists = True
+    if exists == False:
+        cmds.addAttr("root", ln="DistanceCurve", keyable=True)
+addCustomAttrs()
+
+#save the export file.
+cmds.file(save = True, type = "mayaBinary", force = True)
\ No newline at end of file
