Branching

diff --git a/Engine/Extras/Maya_AnimationRiggingTools/ARTv1/General/ART/Projects/Orion/Scripts/Price_CreateLODs.py b/Engine/Extras/Maya_AnimationRiggingTools/ARTv1/General/ART/Projects/Orion/Scripts/Price_CreateLODs.py
new file mode 100644
--- /dev/null
+++ b/Engine/Extras/Maya_AnimationRiggingTools/ARTv1/General/ART/Projects/Orion/Scripts/Price_CreateLODs.py
@@ -0,0 +1,715 @@
+import maya.cmds as cmds
+import customMayaMenu as menu
+import Tools.ART_ImportMultipleWeights as impw
+import Tools.ART_TransferWeights as transfer
+import Tools.ART_ExportLODtoFBX as export
+reload(export)
+
+################################################## PRICE #########################################################################
+# Warn the user that this operation will save their file.
+saveResult = cmds.confirmDialog(title = "Save Warning", message="In order to continue your file must be saved.  Would you like to save it?  If yes it will be saved and after the operation is complete your file will be re-opened.", button = ["Yes", "Cancel"], defaultButton='Yes', cancelButton='Cancel', dismissString='Cancel')
+
+# If they are ok with saving the file, warn them that it will also delete all of the current files in their skin weights directory.
+if saveResult == "Yes":
+    deleteSkinFilesResult = cmds.confirmDialog(title = "SkinWeights Warning", message="In order to continue we must delete all of the files in your ..\General\ART\SkinWeights\ directory.  Would you like to continue?.", button = ["Yes", "Cancel"], defaultButton='Yes', cancelButton='Cancel', dismissString='Cancel')
+    if deleteSkinFilesResult == "Yes":
+        toolsPath = cmds.internalVar(usd = True) + "mayaTools.txt"
+        if os.path.exists(toolsPath):
+            f = open(toolsPath, 'r')
+            mayaToolsDir = f.readline()
+            f.close()
+
+        path = mayaToolsDir + "\General\ART\SkinWeights\\"
+        files = os.listdir(path)
+        #delete all weight cache files in the system folder
+        for file in files:
+            os.remove(path + file)
+            #print path + file
+
+
+
+# Price LOD1 -------------------------------------------------------------------------------------------------------------
+def PriceLOD1():
+    if saveResult == "Yes":
+
+        cmds.file(save=True)
+
+        meshes = ["abdomen_geo_LOD1", "abdomen_plate_geo_LOD1", "arm_hoses_big_geo_LOD1", "arm_hoses_sml_geo_LOD1", "arms_geo_LOD1", "cornea_geo_LOD1", "elbow_pads_geo_LOD1", "eye_shadow_geo_LOD1", "eyeballs_geo_LOD1", "feet_geo_LOD1", "forearm_armor_geo_LOD1", "gun_barrel_geo_LOD1", "gun_barrel_inner_geo_LOD1", "gun_body_panels_geo_LOD1", "gun_geo_LOD1", "gun_handle_geo_LOD1", "gun_scope_geo_LOD1", "gun_scope_holder_geo_LOD1", "gun_stock_base_geo_LOD1", "gun_stock_base_in_geo_LOD1", "gun_stock_base_out_geo_LOD1", "gun_stock_bott_geo_LOD1", "gun_stock_in_geo_LOD1", "gun_stock_panels_geo_LOD1", "hand_cover_r_geo_LOD1", "hands_geo_LOD1", "head_geo_LOD1", "knuckles_geo_LOD1", "loarm_pistonB_bott_geo_LOD1", "loarm_pistonB_top_geo_LOD1", "loarm_pistonC_bott_geo_LOD1", "loarm_pistonC_top_geo_LOD1", "loarm_pistonD_bot_geo_LOD1", "loarm_pistonD_top_geo_LOD1", "lower_teeth_geo_LOD1", "mask_geo_LOD1", "packs_geo_LOD1", "shins_geo_LOD1", "shotgun_arm_geo_LOD1", "shotgun_geo_LOD1", "shoulderpad_r_geo_LOD1", "sniper_canister_arm_01_geo_LOD1", "sniper_canister_arm_02_geo_LOD1", "sniper_canister_arm_03_geo_LOD1", "sniper_canister_geo_LOD1", "tongue_geo_LOD1", "torso_geo_LOD1", "trap_geo_LOD1", "uparm_pistonC_bicep_bott_geo_LOD1", "uparm_pistonC_bicep_top_geo_LOD1", "upper_legs_geo_LOD1", "upper_teeth_geo_LOD1"]
+
+        jointsToDelete = []
+
+        # remove unneeded joints
+        jointsToRemove = ["shotgun_clip"]
+        jointToTransferTo = "shotgun"
+        transfer.LOD_transferWeights(meshes, jointsToRemove, jointToTransferTo)
+        jointsToDelete.extend(jointsToRemove)
+
+        jointsToRemove = ["arm_pouch_l"]
+        jointToTransferTo = "upperarm_twist_02_l"
+        transfer.LOD_transferWeights(meshes, jointsToRemove, jointToTransferTo)
+        jointsToDelete.extend(jointsToRemove)
+
+        jointsToRemove = ["wrist_cover"]
+        jointToTransferTo = "lowerarm_twist_01_r"
+        transfer.LOD_transferWeights(meshes, jointsToRemove, jointToTransferTo)
+        jointsToDelete.extend(jointsToRemove)
+
+        jointsToRemove = ["eyeball_l", "eyeball_r", "eyebrow_in_l", "eyebrow_in_r", "eyebrow_out_l", "eyebrow_out_r", "eyelid_bott_l", "eyelid_bott_r", "eyelid_top_l", "eyelid_top_r", "helmet_sight", "mask"]
+        jointToTransferTo = "head"
+        transfer.LOD_transferWeights(meshes, jointsToRemove, jointToTransferTo)
+        jointsToDelete.extend(jointsToRemove)
+
+        # Delete the joints we no longer want.
+        for one in jointsToDelete:
+            if cmds.objExists(one):
+                cmds.delete(one)
+                
+        # Export the fbx file.
+        export.exportLOD("Price_LOD1.fbx", "LOD1")
+        
+        # Re-open the file without saving.
+        fullPath = cmds.file(q = True, sceneName = True)
+        cmds.file(fullPath, open=True, f=True)
+
+        print "LOD1 Exported Successfully!!!!!!!!!!!!!!!"
+
+PriceLOD1()
+
+# Price LOD2 -------------------------------------------------------------------------------------------------------------
+def PriceLOD2():
+    if saveResult == "Yes":
+
+        cmds.file(save=True)
+        
+        meshes = ["abdomen_geo_LOD2", "abdomen_plate_geo_LOD2", "arm_hoses_big_geo_LOD2", "arm_hoses_sml_geo_LOD2", "arms_geo_LOD2", "cornea_geo_LOD2", "elbow_pads_geo_LOD2", "eye_shadow_geo_LOD2", "eyeballs_geo_LOD2", "feet_geo_LOD2", "forearm_armor_geo_LOD2", "gun_barrel_geo_LOD2", "gun_barrel_inner_geo_LOD2", "gun_body_panels_geo_LOD2", "gun_geo_LOD2", "gun_handle_geo_LOD2", "gun_scope_geo_LOD2", "gun_scope_holder_geo_LOD2", "gun_stock_base_geo_LOD2", "gun_stock_base_in_geo_LOD2", "gun_stock_base_out_geo_LOD2", "gun_stock_bott_geo_LOD2", "gun_stock_in_geo_LOD2", "gun_stock_panels_geo_LOD2", "hand_cover_r_geo_LOD2", "hands_geo_LOD2", "head_geo_LOD2", "knuckles_geo_LOD2", "loarm_pistonB_bott_geo_LOD2", "loarm_pistonB_top_geo_LOD2", "loarm_pistonC_bott_geo_LOD2", "loarm_pistonC_top_geo_LOD2", "loarm_pistonD_bot_geo_LOD2", "loarm_pistonD_top_geo_LOD2", "lower_teeth_geo_LOD2", "mask_geo_LOD2", "packs_geo_LOD2", "shins_geo_LOD2", "shotgun_arm_geo_LOD2", "shotgun_geo_LOD2", "shoulderpad_r_geo_LOD2", "sniper_canister_arm_01_geo_LOD2", "sniper_canister_arm_02_geo_LOD2", "sniper_canister_arm_03_geo_LOD2", "sniper_canister_geo_LOD2", "tongue_geo_LOD2", "torso_geo_LOD2", "trap_geo_LOD2", "uparm_pistonC_bicep_bott_geo_LOD2", "uparm_pistonC_bicep_top_geo_LOD2", "upper_legs_geo_LOD2", "upper_teeth_geo_LOD2"]
+
+        # Export a skin weights file for each mesh we plan to operate on.
+        cmds.select(meshes)
+        menu.exportSelectedSkinWeights(True)
+
+        # Set the pose if needed
+        cmds.setAttr("thumb_03_l.rotateZ", 30.05735749)
+        cmds.setAttr("thumb_03_l.rotateY", 4.575153365e-015)
+        cmds.setAttr("thumb_03_l.rotateX", 2.807367528e-014)
+        cmds.setAttr("ring_03_l.rotateZ", 50.02688615)
+        cmds.setAttr("ring_03_l.rotateY", -4.767840667e-014)
+        cmds.setAttr("ring_03_l.rotateX", 2.313919533e-014)
+        cmds.setAttr("pinky_03_l.rotateZ", 50.02688615)
+        cmds.setAttr("pinky_03_l.rotateY", -1.296376071e-015)
+        cmds.setAttr("pinky_03_l.rotateX", 4.230075606e-014)
+        cmds.setAttr("middle_03_l.rotateZ", 55.02709033)
+        cmds.setAttr("middle_03_l.rotateY", -1.652397549e-014)
+        cmds.setAttr("middle_03_l.rotateX", 2.297826517e-014)
+        cmds.setAttr("index_03_l.rotateZ", 64.64121137)
+        cmds.setAttr("index_03_l.rotateY", -2.480634643e-014)
+        cmds.setAttr("index_03_l.rotateX", 1.077164418e-014)
+
+        cmds.setAttr("thumb_03_r.rotateZ", 30.05735749)
+        cmds.setAttr("thumb_03_r.rotateY", 4.575153365e-015)
+        cmds.setAttr("thumb_03_r.rotateX", 2.807367528e-014)
+        cmds.setAttr("ring_03_r.rotateZ", 50.02688615)
+        cmds.setAttr("ring_03_r.rotateY", -4.767840667e-014)
+        cmds.setAttr("ring_03_r.rotateX", 2.313919533e-014)
+        cmds.setAttr("pinky_03_r.rotateZ", 50.02688615)
+        cmds.setAttr("pinky_03_r.rotateY", -1.296376071e-015)
+        cmds.setAttr("pinky_03_r.rotateX", 4.230075606e-014)
+        cmds.setAttr("middle_03_r.rotateZ", 55.02709033)
+        cmds.setAttr("middle_03_r.rotateY", -1.652397549e-014)
+        cmds.setAttr("middle_03_r.rotateX", 2.297826517e-014)
+        cmds.setAttr("index_03_r.rotateZ", 64.64121137)
+        cmds.setAttr("index_03_r.rotateY", -2.480634643e-014)
+        cmds.setAttr("index_03_r.rotateX", 1.077164418e-014)
+
+        # Delete the history on the meshes so that they get baked down into their posed positions
+        cmds.delete(meshes, ch=True)
+
+        # Move the skeleton back to its bindpose (This script REQUIRES the bind pose to be the pose that we want to export the skeleton to.  This is the case for all known characters created with A.R.T.)
+        cmds.dagPose("root", bp=True, r=True) 
+
+        # Re-import the skinning we exported before onto the mesh so thats its bound in its posed position.
+        cmds.select(meshes)
+        impw.importMultipleSkinWeights("", True)
+
+        jointsToDelete = []
+
+         # remove unneeded joints
+        jointsToRemove = ["hip_pack_a", "hip_pack_b", "hip_pack_c", "hip_pack_d"]
+        jointToTransferTo = "pelvis"
+        transfer.LOD_transferWeights(meshes, jointsToRemove, jointToTransferTo)
+        jointsToDelete.extend(jointsToRemove)
+
+        jointsToRemove = ["sniper_can_arm_01", "sniper_can_arm_02", "sniper_can_arm_03", "pack_hose_out", "pack_hose_in"]
+        jointToTransferTo = "spine_03"
+        transfer.LOD_transferWeights(meshes, jointsToRemove, jointToTransferTo)
+        jointsToDelete.extend(jointsToRemove)
+
+        jointsToRemove = ["shotgun_clip"]
+        jointToTransferTo = "shotgun"
+        transfer.LOD_transferWeights(meshes, jointsToRemove, jointToTransferTo)
+        jointsToDelete.extend(jointsToRemove)
+
+        jointsToRemove = ["arm_pouch_l"]
+        jointToTransferTo = "upperarm_twist_02_l"
+        transfer.LOD_transferWeights(meshes, jointsToRemove, jointToTransferTo)
+        jointsToDelete.extend(jointsToRemove)
+
+        jointsToRemove = ["bot_uparm", "bot_uparm_pistonC_top"]
+        jointToTransferTo = "upperarm_twist_01_r"
+        transfer.LOD_transferWeights(meshes, jointsToRemove, jointToTransferTo)
+        jointsToDelete.extend(jointsToRemove)
+
+        jointsToRemove = ["arm_hose_top", "bot_loarm", "bot_loarm_pistonB_top", "bot_loarm_pistonC_top", "bot_loarm_pistonD_top", "bot_loarm_twist", "bot_loarm_pistonC_bot", "bot_loarm_pistonD_bot", "bot_uparm_pistonC_bot", "arm_hose_bott", "wrist_cover"]
+        jointToTransferTo = "lowerarm_r"
+        transfer.LOD_transferWeights(meshes, jointsToRemove, jointToTransferTo)
+        jointsToDelete.extend(jointsToRemove)
+
+        jointsToRemove = ["arm_hose_bott", "wrist_cover"]
+        jointToTransferTo = "lowerarm_twist_01_r"
+        transfer.LOD_transferWeights(meshes, jointsToRemove, jointToTransferTo)
+        jointsToDelete.extend(jointsToRemove)
+
+        jointsToRemove = ["index_03_l"]
+        jointToTransferTo = "index_02_l"
+        transfer.LOD_transferWeights(meshes, jointsToRemove, jointToTransferTo)
+        jointsToDelete.extend(jointsToRemove)
+
+        jointsToRemove = ["middle_03_l"]
+        jointToTransferTo = "middle_02_l"
+        transfer.LOD_transferWeights(meshes, jointsToRemove, jointToTransferTo)
+        jointsToDelete.extend(jointsToRemove)
+
+        jointsToRemove = ["pinky_03_l"]
+        jointToTransferTo = "pinky_02_l"
+        transfer.LOD_transferWeights(meshes, jointsToRemove, jointToTransferTo)
+        jointsToDelete.extend(jointsToRemove)
+
+        jointsToRemove = ["ring_03_l"]
+        jointToTransferTo = "ring_02_l"
+        transfer.LOD_transferWeights(meshes, jointsToRemove, jointToTransferTo)
+        jointsToDelete.extend(jointsToRemove)
+
+        jointsToRemove = ["thumb_03_l"]
+        jointToTransferTo = "thumb_02_l"
+        transfer.LOD_transferWeights(meshes, jointsToRemove, jointToTransferTo)
+        jointsToDelete.extend(jointsToRemove)
+
+        jointsToRemove = ["index_03_r"]
+        jointToTransferTo = "index_02_r"
+        transfer.LOD_transferWeights(meshes, jointsToRemove, jointToTransferTo)
+        jointsToDelete.extend(jointsToRemove)
+
+        jointsToRemove = ["middle_03_r"]
+        jointToTransferTo = "middle_02_r"
+        transfer.LOD_transferWeights(meshes, jointsToRemove, jointToTransferTo)
+        jointsToDelete.extend(jointsToRemove)
+
+        jointsToRemove = ["pinky_03_r"]
+        jointToTransferTo = "pinky_02_r"
+        transfer.LOD_transferWeights(meshes, jointsToRemove, jointToTransferTo)
+        jointsToDelete.extend(jointsToRemove)
+
+        jointsToRemove = ["ring_03_r"]
+        jointToTransferTo = "ring_02_r"
+        transfer.LOD_transferWeights(meshes, jointsToRemove, jointToTransferTo)
+        jointsToDelete.extend(jointsToRemove)
+
+        jointsToRemove = ["thumb_03_r"]
+        jointToTransferTo = "thumb_02_r"
+        transfer.LOD_transferWeights(meshes, jointsToRemove, jointToTransferTo)
+        jointsToDelete.extend(jointsToRemove)
+
+        jointsToRemove = ["gun_barrel", "gun_body_panel_l", "gun_body_panel_r", "gun_scope", "gun_stock_base", "gun_barrel_inner", "gun_stock_base_in_bott", "gun_stock_base_out_bott", "gun_stock_base_in_top", "gun_stock_base_out_top", "gun_stock_in_bott", "gun_stock_in_top", "gun_stock_bott", "gun_stock_panel_l", "gun_stock_panel_r"]
+        jointToTransferTo = "gun"
+        transfer.LOD_transferWeights(meshes, jointsToRemove, jointToTransferTo)
+        jointsToDelete.extend(jointsToRemove)
+
+        jointsToRemove = ["bot_loarm_pistonB_bot"]
+        jointToTransferTo = "hand_r"
+        transfer.LOD_transferWeights(meshes, jointsToRemove, jointToTransferTo)
+        jointsToDelete.extend(jointsToRemove)
+
+        jointsToRemove = ["eyeball_l", "eyeball_r", "eyebrow_in_l", "eyebrow_in_r", "eyebrow_out_l", "eyebrow_out_r", "eyelid_bott_l", "eyelid_bott_r", "eyelid_top_l", "eyelid_top_r", "helmet_sight", "jaw", "mask"]
+        jointToTransferTo = "head"
+        transfer.LOD_transferWeights(meshes, jointsToRemove, jointToTransferTo)
+        jointsToDelete.extend(jointsToRemove)
+
+        jointsToRemove = ["ankle_plate_in_l", "ankle_plate_out_l"]
+        jointToTransferTo = "calf_l"
+        transfer.LOD_transferWeights(meshes, jointsToRemove, jointToTransferTo)
+        jointsToDelete.extend(jointsToRemove)
+
+        jointsToRemove = ["ankle_plate_in_r", "ankle_plate_out_r"]
+        jointToTransferTo = "calf_r"
+        transfer.LOD_transferWeights(meshes, jointsToRemove, jointToTransferTo)
+        jointsToDelete.extend(jointsToRemove)
+
+        jointsToRemove = ["heel_plate_l"]
+        jointToTransferTo = "foot_l"
+        transfer.LOD_transferWeights(meshes, jointsToRemove, jointToTransferTo)
+        jointsToDelete.extend(jointsToRemove)
+
+        jointsToRemove = ["heel_plate_r"]
+        jointToTransferTo = "foot_r"
+        transfer.LOD_transferWeights(meshes, jointsToRemove, jointToTransferTo)
+        jointsToDelete.extend(jointsToRemove)
+
+        # Delete the joints we no longer want.
+        for one in jointsToDelete:
+            if cmds.objExists(one):
+                cmds.delete(one)
+                
+        # Export the fbx file.
+        export.exportLOD("Price_LOD2.fbx", "LOD2")
+
+        # Re-open the file without saving.
+        fullPath = cmds.file(q = True, sceneName = True)
+        cmds.file(fullPath, open=True, f=True)
+
+        print "LOD2 Exported Successfully!!!!!!!!!!!!!!!"
+
+PriceLOD2()
+
+# Price LOD3 -------------------------------------------------------------------------------------------------------------
+def PriceLOD3():
+    if saveResult == "Yes":
+
+        cmds.file(save=True)
+    
+        meshes = ["abdomen_geo_LOD3", "abdomen_plate_geo_LOD3", "arm_hoses_big_geo_LOD3", "arm_hoses_sml_geo_LOD3", "arms_geo_LOD3", "cornea_geo_LOD3", "elbow_pads_geo_LOD3", "eye_shadow_geo_LOD3", "eyeballs_geo_LOD3", "feet_geo_LOD3", "forearm_armor_geo_LOD3", "gun_barrel_geo_LOD3", "gun_barrel_inner_geo_LOD3", "gun_body_panels_geo_LOD3", "gun_geo_LOD3", "gun_handle_geo_LOD3", "gun_scope_geo_LOD3", "gun_scope_holder_geo_LOD3", "gun_stock_base_geo_LOD3", "gun_stock_base_in_geo_LOD3", "gun_stock_base_out_geo_LOD3", "gun_stock_bott_geo_LOD3", "gun_stock_in_geo_LOD3", "gun_stock_panels_geo_LOD3", "hand_cover_r_geo_LOD3", "hands_geo_LOD3", "head_geo_LOD3", "knuckles_geo_LOD3", "loarm_pistonB_bott_geo_LOD3", "loarm_pistonB_top_geo_LOD3", "loarm_pistonC_bott_geo_LOD3", "loarm_pistonC_top_geo_LOD3", "loarm_pistonD_bot_geo_LOD3", "loarm_pistonD_top_geo_LOD3", "lower_teeth_geo_LOD3", "mask_geo_LOD3", "packs_geo_LOD3", "shins_geo_LOD3", "shotgun_arm_geo_LOD3", "shotgun_geo_LOD3", "shoulderpad_r_geo_LOD3", "sniper_canister_arm_01_geo_LOD3", "sniper_canister_arm_02_geo_LOD3", "sniper_canister_arm_03_geo_LOD3", "sniper_canister_geo_LOD3", "tongue_geo_LOD3", "torso_geo_LOD3", "trap_geo_LOD3", "uparm_pistonC_bicep_bott_geo_LOD3", "uparm_pistonC_bicep_top_geo_LOD3", "upper_legs_geo_LOD3", "upper_teeth_geo_LOD3"]
+    
+        # Export a skin weights file for each mesh we plan to operate on.
+        cmds.select(meshes)
+        menu.exportSelectedSkinWeights(True)
+
+        # Set the pose if needed
+        cmds.setAttr("thumb_03_l.rotateZ", 30.05735749)
+        cmds.setAttr("thumb_03_l.rotateY", 4.575153365e-015)
+        cmds.setAttr("thumb_03_l.rotateX", 2.807367528e-014)
+        cmds.setAttr("thumb_02_l.rotateZ", 22.03618142)
+        cmds.setAttr("thumb_02_l.rotateY", 5.675048056)
+        cmds.setAttr("thumb_02_l.rotateX", 22.94177697)
+        cmds.setAttr("thumb_01_l.rotateZ", -9.145300774)
+        cmds.setAttr("thumb_01_l.rotateY", 15.64425403)
+        cmds.setAttr("thumb_01_l.rotateX", 14.45073388)
+        cmds.setAttr("ring_03_l.rotateZ", 50.02688615)
+        cmds.setAttr("ring_03_l.rotateY", -4.767840667e-014)
+        cmds.setAttr("ring_03_l.rotateX", 2.313919533e-014)
+        cmds.setAttr("ring_02_l.rotateZ", 54.79433274)
+        cmds.setAttr("ring_02_l.rotateY", -7.710913552e-015)
+        cmds.setAttr("ring_02_l.rotateX", 2.923377018e-014)
+        cmds.setAttr("ring_01_l.rotateZ", 54.76549897)
+        cmds.setAttr("ring_01_l.rotateY", -1.525905275e-014)
+        cmds.setAttr("ring_01_l.rotateX", 4.56459293e-014)
+        cmds.setAttr("pinky_03_l.rotateZ", 50.02688615)
+        cmds.setAttr("pinky_03_l.rotateY", -1.296376071e-015)
+        cmds.setAttr("pinky_03_l.rotateX", 4.230075606e-014)
+        cmds.setAttr("pinky_02_l.rotateZ", 54.79433274)
+        cmds.setAttr("pinky_02_l.rotateY", -3.788711524e-014)
+        cmds.setAttr("pinky_02_l.rotateX", 1.292100339e-014)
+        cmds.setAttr("pinky_01_l.rotateZ", 53.17737673)
+        cmds.setAttr("pinky_01_l.rotateY", -10.46408531)
+        cmds.setAttr("pinky_01_l.rotateX", 10.56158491)
+        cmds.setAttr("middle_03_l.rotateZ", 55.02709033)
+        cmds.setAttr("middle_03_l.rotateY", -1.652397549e-014)
+        cmds.setAttr("middle_03_l.rotateX", 2.297826517e-014)
+        cmds.setAttr("middle_02_l.rotateZ", 59.79453692)
+        cmds.setAttr("middle_02_l.rotateY", -6.951251345e-015)
+        cmds.setAttr("middle_02_l.rotateX", 2.166752877e-014)
+        cmds.setAttr("middle_01_l.rotateZ", 51.05024479)
+        cmds.setAttr("middle_01_l.rotateY", -1.133072605e-014)
+        cmds.setAttr("middle_01_l.rotateX", 5.819421144e-014)
+        cmds.setAttr("index_03_l.rotateZ", 64.64121137)
+        cmds.setAttr("index_03_l.rotateY", -2.480634643e-014)
+        cmds.setAttr("index_03_l.rotateX", 1.077164418e-014)
+        cmds.setAttr("index_02_l.rotateZ", 64.20477863)
+        cmds.setAttr("index_02_l.rotateY", -2.481003482e-014)
+        cmds.setAttr("index_02_l.rotateX", 5.153492507e-014)
+        cmds.setAttr("index_01_l.rotateZ", 46.94060635)
+        cmds.setAttr("index_01_l.rotateY", 3.692762633)
+        cmds.setAttr("index_01_l.rotateX", -8.969261741)
+
+        cmds.setAttr("thumb_03_r.rotateZ", 30.05735749)
+        cmds.setAttr("thumb_03_r.rotateY", 4.575153365e-015)
+        cmds.setAttr("thumb_03_r.rotateX", 2.807367528e-014)
+        cmds.setAttr("thumb_02_r.rotateZ", 22.03618142)
+        cmds.setAttr("thumb_02_r.rotateY", 5.675048056)
+        cmds.setAttr("thumb_02_r.rotateX", 22.94177697)
+        cmds.setAttr("thumb_01_r.rotateZ", -9.145300774)
+        cmds.setAttr("thumb_01_r.rotateY", 15.64425403)
+        cmds.setAttr("thumb_01_r.rotateX", 14.45073388)
+        cmds.setAttr("ring_03_r.rotateZ", 50.02688615)
+        cmds.setAttr("ring_03_r.rotateY", -4.767840667e-014)
+        cmds.setAttr("ring_03_r.rotateX", 2.313919533e-014)
+        cmds.setAttr("ring_02_r.rotateZ", 54.79433274)
+        cmds.setAttr("ring_02_r.rotateY", -7.710913552e-015)
+        cmds.setAttr("ring_02_r.rotateX", 2.923377018e-014)
+        cmds.setAttr("ring_01_r.rotateZ", 54.76549897)
+        cmds.setAttr("ring_01_r.rotateY", -1.525905275e-014)
+        cmds.setAttr("ring_01_r.rotateX", 4.56459293e-014)
+        cmds.setAttr("pinky_03_r.rotateZ", 50.02688615)
+        cmds.setAttr("pinky_03_r.rotateY", -1.296376071e-015)
+        cmds.setAttr("pinky_03_r.rotateX", 4.230075606e-014)
+        cmds.setAttr("pinky_02_r.rotateZ", 54.79433274)
+        cmds.setAttr("pinky_02_r.rotateY", -3.788711524e-014)
+        cmds.setAttr("pinky_02_r.rotateX", 1.292100339e-014)
+        cmds.setAttr("pinky_01_r.rotateZ", 53.17737673)
+        cmds.setAttr("pinky_01_r.rotateY", -10.46408531)
+        cmds.setAttr("pinky_01_r.rotateX", 10.56158491)
+        cmds.setAttr("middle_03_r.rotateZ", 55.02709033)
+        cmds.setAttr("middle_03_r.rotateY", -1.652397549e-014)
+        cmds.setAttr("middle_03_r.rotateX", 2.297826517e-014)
+        cmds.setAttr("middle_02_r.rotateZ", 59.79453692)
+        cmds.setAttr("middle_02_r.rotateY", -6.951251345e-015)
+        cmds.setAttr("middle_02_r.rotateX", 2.166752877e-014)
+        cmds.setAttr("middle_01_r.rotateZ", 51.05024479)
+        cmds.setAttr("middle_01_r.rotateY", -1.133072605e-014)
+        cmds.setAttr("middle_01_r.rotateX", 5.819421144e-014)
+        cmds.setAttr("index_03_r.rotateZ", 64.64121137)
+        cmds.setAttr("index_03_r.rotateY", -2.480634643e-014)
+        cmds.setAttr("index_03_r.rotateX", 1.077164418e-014)
+        cmds.setAttr("index_02_r.rotateZ", 64.20477863)
+        cmds.setAttr("index_02_r.rotateY", -2.481003482e-014)
+        cmds.setAttr("index_02_r.rotateX", 5.153492507e-014)
+        cmds.setAttr("index_01_r.rotateZ", 46.94060635)
+        cmds.setAttr("index_01_r.rotateY", 3.692762633)
+        cmds.setAttr("index_01_r.rotateX", -8.969261741)
+
+        # Delete the history on the meshes so that they get baked down into their posed positions
+        cmds.delete(meshes, ch=True)
+
+        # Move the skeleton back to its bindpose (This script REQUIRES the bind pose to be the pose that we want to export the skeleton to.  This is the case for all known characters created with A.R.T.)
+        cmds.dagPose("root", bp=True, r=True) 
+
+        # Re-import the skinning we exported before onto the mesh so thats its bound in its posed position.
+        cmds.select(meshes)
+        impw.importMultipleSkinWeights("", True)
+        
+        jointsToDelete = []
+
+        # remove unneeded joints
+        jointsToRemove = ["hip_pack_a", "hip_pack_b", "hip_pack_c", "hip_pack_d"]
+        jointToTransferTo = "pelvis"
+        transfer.LOD_transferWeights(meshes, jointsToRemove, jointToTransferTo)
+        jointsToDelete.extend(jointsToRemove)
+
+        jointsToRemove = ["trap_handle"]
+        jointToTransferTo = "spine_01"
+        transfer.LOD_transferWeights(meshes, jointsToRemove, jointToTransferTo)
+        jointsToDelete.extend(jointsToRemove)
+
+        jointsToRemove = ["sniper_can_arm_01", "sniper_can_arm_02", "sniper_can_arm_03", "pack_hose_out", "pack_hose_in"]
+        jointToTransferTo = "spine_03"
+        transfer.LOD_transferWeights(meshes, jointsToRemove, jointToTransferTo)
+        jointsToDelete.extend(jointsToRemove)
+
+        jointsToRemove = ["shotgun_clip"]
+        jointToTransferTo = "shotgun"
+        transfer.LOD_transferWeights(meshes, jointsToRemove, jointToTransferTo)
+        jointsToDelete.extend(jointsToRemove)
+
+        jointsToRemove = ["upperarm_twist_01_l", "upperarm_twist_02_l", "arm_pouch_l"]
+        jointToTransferTo = "upperarm_l"
+        transfer.LOD_transferWeights(meshes, jointsToRemove, jointToTransferTo)
+        jointsToDelete.extend(jointsToRemove)
+
+        jointsToRemove = ["upperarm_twist_01_r", "bot_uparm", "bot_uparm_pistonC_top", "upperarm_twist_02_r"]
+        jointToTransferTo = "upperarm_r"
+        transfer.LOD_transferWeights(meshes, jointsToRemove, jointToTransferTo)
+        jointsToDelete.extend(jointsToRemove)
+
+        jointsToRemove = ["lowerarm_twist_01_l", "lowerarm_twist_02_l"]
+        jointToTransferTo = "lowerarm_l"
+        transfer.LOD_transferWeights(meshes, jointsToRemove, jointToTransferTo)
+        jointsToDelete.extend(jointsToRemove)
+
+        jointsToRemove = ["arm_hose_top", "bot_loarm", "bot_loarm_pistonB_top", "bot_loarm_pistonC_top", "bot_loarm_pistonD_top", "bot_loarm_twist", "bot_loarm_pistonC_bot", "bot_loarm_pistonD_bot", "bot_uparm_pistonC_bot", "lowerarm_twist_01_r", "arm_hose_bott", "wrist_cover"]
+        jointToTransferTo = "lowerarm_r"
+        transfer.LOD_transferWeights(meshes, jointsToRemove, jointToTransferTo)
+        jointsToDelete.extend(jointsToRemove)
+
+        jointsToRemove = ["index_01_l", "index_02_l", "index_03_l", "middle_01_l", "middle_02_l", "middle_03_l", "pinky_01_l", "pinky_02_l", "pinky_03_l", "ring_01_l", "ring_02_l", "ring_03_l", "thumb_01_l", "thumb_02_l", "thumb_03_l"]
+        jointToTransferTo = "hand_l"
+        transfer.LOD_transferWeights(meshes, jointsToRemove, jointToTransferTo)
+        jointsToDelete.extend(jointsToRemove)
+
+        jointsToRemove = ["gun_barrel", "gun_body_panel_l", "gun_body_panel_r", "gun_scope", "gun_stock_base", "gun_barrel_inner", "gun_stock_base_in_bott", "gun_stock_base_out_bott", "gun_stock_base_in_top", "gun_stock_base_out_top", "gun_stock_in_bott", "gun_stock_in_top", "gun_stock_bott", "gun_stock_panel_l", "gun_stock_panel_r"]
+        jointToTransferTo = "gun"
+        transfer.LOD_transferWeights(meshes, jointsToRemove, jointToTransferTo)
+        jointsToDelete.extend(jointsToRemove)
+
+        jointsToRemove = ["bot_loarm_pistonB_bot", "index_01_r", "index_02_r", "index_03_r", "middle_01_r", "middle_02_r", "middle_03_r", "pinky_01_r", "pinky_02_r", "pinky_03_r", "ring_01_r", "ring_02_r", "ring_03_r", "thumb_01_r", "thumb_02_r", "thumb_03_r"]
+        jointToTransferTo = "hand_r"
+        transfer.LOD_transferWeights(meshes, jointsToRemove, jointToTransferTo)
+        jointsToDelete.extend(jointsToRemove)
+
+        jointsToRemove = ["eyeball_l", "eyeball_r", "eyebrow_in_l", "eyebrow_in_r", "eyebrow_out_l", "eyebrow_out_r", "eyelid_bott_l", "eyelid_bott_r", "eyelid_top_l", "eyelid_top_r", "helmet_sight", "jaw", "mask"]
+        jointToTransferTo = "head"
+        transfer.LOD_transferWeights(meshes, jointsToRemove, jointToTransferTo)
+        jointsToDelete.extend(jointsToRemove)
+
+        jointsToRemove = ["thigh_twist_01_l", "thigh_twist_02_l"]
+        jointToTransferTo = "thigh_l"
+        transfer.LOD_transferWeights(meshes, jointsToRemove, jointToTransferTo)
+        jointsToDelete.extend(jointsToRemove)
+
+        jointsToRemove = ["thigh_twist_01_r", "thigh_twist_02_r"]
+        jointToTransferTo = "thigh_r"
+        transfer.LOD_transferWeights(meshes, jointsToRemove, jointToTransferTo)
+        jointsToDelete.extend(jointsToRemove)
+
+        jointsToRemove = ["ankle_plate_in_l", "ankle_plate_out_l", "calf_twist_01_l"]
+        jointToTransferTo = "calf_l"
+        transfer.LOD_transferWeights(meshes, jointsToRemove, jointToTransferTo)
+        jointsToDelete.extend(jointsToRemove)
+
+        jointsToRemove = ["ankle_plate_in_r", "ankle_plate_out_r", "calf_twist_01_r"]
+        jointToTransferTo = "calf_r"
+        transfer.LOD_transferWeights(meshes, jointsToRemove, jointToTransferTo)
+        jointsToDelete.extend(jointsToRemove)
+
+        jointsToRemove = ["heel_plate_l"]
+        jointToTransferTo = "foot_l"
+        transfer.LOD_transferWeights(meshes, jointsToRemove, jointToTransferTo)
+        jointsToDelete.extend(jointsToRemove)
+
+        jointsToRemove = ["heel_plate_r"]
+        jointToTransferTo = "foot_r"
+        transfer.LOD_transferWeights(meshes, jointsToRemove, jointToTransferTo)
+        jointsToDelete.extend(jointsToRemove)
+        
+        # Delete the joints we no longer want.
+        for one in jointsToDelete:
+            if cmds.objExists(one):
+                cmds.delete(one)
+                
+        # Export the fbx file.
+        export.exportLOD("Price_LOD3.fbx", "LOD3")
+        
+        # Re-open the file without saving.
+        fullPath = cmds.file(q = True, sceneName = True)
+        cmds.file(fullPath, open=True, f=True)
+
+        print "LOD3 Exported Successfully!!!!!!!!!!!!!!!"
+
+PriceLOD3()
+
+# Price LOD4 -------------------------------------------------------------------------------------------------------------
+def PriceLOD4():
+    if saveResult == "Yes":
+
+        cmds.file(save=True)
+    
+        meshes = ["abdomen_geo_LOD4", "abdomen_plate_geo_LOD4", "arm_hoses_big_geo_LOD4", "arm_hoses_sml_geo_LOD4", "arms_geo_LOD4", "cornea_geo_LOD4", "elbow_pads_geo_LOD4", "eye_shadow_geo_LOD4", "eyeballs_geo_LOD4", "feet_geo_LOD4", "forearm_armor_geo_LOD4", "gun_barrel_geo_LOD4", "gun_barrel_inner_geo_LOD4", "gun_body_panels_geo_LOD4", "gun_geo_LOD4", "gun_handle_geo_LOD4", "gun_scope_geo_LOD4", "gun_scope_holder_geo_LOD4", "gun_stock_base_geo_LOD4", "gun_stock_base_in_geo_LOD4", "gun_stock_base_out_geo_LOD4", "gun_stock_bott_geo_LOD4", "gun_stock_in_geo_LOD4", "gun_stock_panels_geo_LOD4", "hand_cover_r_geo_LOD4", "hands_geo_LOD4", "head_geo_LOD4", "knuckles_geo_LOD4", "loarm_pistonB_bott_geo_LOD4", "loarm_pistonB_top_geo_LOD4", "loarm_pistonC_bott_geo_LOD4", "loarm_pistonC_top_geo_LOD4", "loarm_pistonD_bot_geo_LOD4", "loarm_pistonD_top_geo_LOD4", "lower_teeth_geo_LOD4", "mask_geo_LOD4", "packs_geo_LOD4", "shins_geo_LOD4", "shotgun_arm_geo_LOD4", "shotgun_geo_LOD4", "shoulderpad_r_geo_LOD4", "sniper_canister_arm_01_geo_LOD4", "sniper_canister_arm_02_geo_LOD4", "sniper_canister_arm_03_geo_LOD4", "sniper_canister_geo_LOD4", "tongue_geo_LOD4", "torso_geo_LOD4", "trap_geo_LOD4", "uparm_pistonC_bicep_bott_geo_LOD4", "uparm_pistonC_bicep_top_geo_LOD4", "upper_legs_geo_LOD4", "upper_teeth_geo_LOD4"]
+
+        # Export a skin weights file for each mesh we plan to operate on.
+        cmds.select(meshes)
+        menu.exportSelectedSkinWeights(True)
+
+        # Set the pose if needed
+        cmds.setAttr("thumb_03_l.rotateZ", 30.05735749)
+        cmds.setAttr("thumb_03_l.rotateY", 4.575153365e-015)
+        cmds.setAttr("thumb_03_l.rotateX", 2.807367528e-014)
+        cmds.setAttr("thumb_02_l.rotateZ", 22.03618142)
+        cmds.setAttr("thumb_02_l.rotateY", 5.675048056)
+        cmds.setAttr("thumb_02_l.rotateX", 22.94177697)
+        cmds.setAttr("thumb_01_l.rotateZ", -9.145300774)
+        cmds.setAttr("thumb_01_l.rotateY", 15.64425403)
+        cmds.setAttr("thumb_01_l.rotateX", 14.45073388)
+        cmds.setAttr("ring_03_l.rotateZ", 50.02688615)
+        cmds.setAttr("ring_03_l.rotateY", -4.767840667e-014)
+        cmds.setAttr("ring_03_l.rotateX", 2.313919533e-014)
+        cmds.setAttr("ring_02_l.rotateZ", 54.79433274)
+        cmds.setAttr("ring_02_l.rotateY", -7.710913552e-015)
+        cmds.setAttr("ring_02_l.rotateX", 2.923377018e-014)
+        cmds.setAttr("ring_01_l.rotateZ", 54.76549897)
+        cmds.setAttr("ring_01_l.rotateY", -1.525905275e-014)
+        cmds.setAttr("ring_01_l.rotateX", 4.56459293e-014)
+        cmds.setAttr("pinky_03_l.rotateZ", 50.02688615)
+        cmds.setAttr("pinky_03_l.rotateY", -1.296376071e-015)
+        cmds.setAttr("pinky_03_l.rotateX", 4.230075606e-014)
+        cmds.setAttr("pinky_02_l.rotateZ", 54.79433274)
+        cmds.setAttr("pinky_02_l.rotateY", -3.788711524e-014)
+        cmds.setAttr("pinky_02_l.rotateX", 1.292100339e-014)
+        cmds.setAttr("pinky_01_l.rotateZ", 53.17737673)
+        cmds.setAttr("pinky_01_l.rotateY", -10.46408531)
+        cmds.setAttr("pinky_01_l.rotateX", 10.56158491)
+        cmds.setAttr("middle_03_l.rotateZ", 55.02709033)
+        cmds.setAttr("middle_03_l.rotateY", -1.652397549e-014)
+        cmds.setAttr("middle_03_l.rotateX", 2.297826517e-014)
+        cmds.setAttr("middle_02_l.rotateZ", 59.79453692)
+        cmds.setAttr("middle_02_l.rotateY", -6.951251345e-015)
+        cmds.setAttr("middle_02_l.rotateX", 2.166752877e-014)
+        cmds.setAttr("middle_01_l.rotateZ", 51.05024479)
+        cmds.setAttr("middle_01_l.rotateY", -1.133072605e-014)
+        cmds.setAttr("middle_01_l.rotateX", 5.819421144e-014)
+        cmds.setAttr("index_03_l.rotateZ", 64.64121137)
+        cmds.setAttr("index_03_l.rotateY", -2.480634643e-014)
+        cmds.setAttr("index_03_l.rotateX", 1.077164418e-014)
+        cmds.setAttr("index_02_l.rotateZ", 64.20477863)
+        cmds.setAttr("index_02_l.rotateY", -2.481003482e-014)
+        cmds.setAttr("index_02_l.rotateX", 5.153492507e-014)
+        cmds.setAttr("index_01_l.rotateZ", 46.94060635)
+        cmds.setAttr("index_01_l.rotateY", 3.692762633)
+        cmds.setAttr("index_01_l.rotateX", -8.969261741)
+
+        cmds.setAttr("thumb_03_r.rotateZ", 30.05735749)
+        cmds.setAttr("thumb_03_r.rotateY", 4.575153365e-015)
+        cmds.setAttr("thumb_03_r.rotateX", 2.807367528e-014)
+        cmds.setAttr("thumb_02_r.rotateZ", 22.03618142)
+        cmds.setAttr("thumb_02_r.rotateY", 5.675048056)
+        cmds.setAttr("thumb_02_r.rotateX", 22.94177697)
+        cmds.setAttr("thumb_01_r.rotateZ", -9.145300774)
+        cmds.setAttr("thumb_01_r.rotateY", 15.64425403)
+        cmds.setAttr("thumb_01_r.rotateX", 14.45073388)
+        cmds.setAttr("ring_03_r.rotateZ", 50.02688615)
+        cmds.setAttr("ring_03_r.rotateY", -4.767840667e-014)
+        cmds.setAttr("ring_03_r.rotateX", 2.313919533e-014)
+        cmds.setAttr("ring_02_r.rotateZ", 54.79433274)
+        cmds.setAttr("ring_02_r.rotateY", -7.710913552e-015)
+        cmds.setAttr("ring_02_r.rotateX", 2.923377018e-014)
+        cmds.setAttr("ring_01_r.rotateZ", 54.76549897)
+        cmds.setAttr("ring_01_r.rotateY", -1.525905275e-014)
+        cmds.setAttr("ring_01_r.rotateX", 4.56459293e-014)
+        cmds.setAttr("pinky_03_r.rotateZ", 50.02688615)
+        cmds.setAttr("pinky_03_r.rotateY", -1.296376071e-015)
+        cmds.setAttr("pinky_03_r.rotateX", 4.230075606e-014)
+        cmds.setAttr("pinky_02_r.rotateZ", 54.79433274)
+        cmds.setAttr("pinky_02_r.rotateY", -3.788711524e-014)
+        cmds.setAttr("pinky_02_r.rotateX", 1.292100339e-014)
+        cmds.setAttr("pinky_01_r.rotateZ", 53.17737673)
+        cmds.setAttr("pinky_01_r.rotateY", -10.46408531)
+        cmds.setAttr("pinky_01_r.rotateX", 10.56158491)
+        cmds.setAttr("middle_03_r.rotateZ", 55.02709033)
+        cmds.setAttr("middle_03_r.rotateY", -1.652397549e-014)
+        cmds.setAttr("middle_03_r.rotateX", 2.297826517e-014)
+        cmds.setAttr("middle_02_r.rotateZ", 59.79453692)
+        cmds.setAttr("middle_02_r.rotateY", -6.951251345e-015)
+        cmds.setAttr("middle_02_r.rotateX", 2.166752877e-014)
+        cmds.setAttr("middle_01_r.rotateZ", 51.05024479)
+        cmds.setAttr("middle_01_r.rotateY", -1.133072605e-014)
+        cmds.setAttr("middle_01_r.rotateX", 5.819421144e-014)
+        cmds.setAttr("index_03_r.rotateZ", 64.64121137)
+        cmds.setAttr("index_03_r.rotateY", -2.480634643e-014)
+        cmds.setAttr("index_03_r.rotateX", 1.077164418e-014)
+        cmds.setAttr("index_02_r.rotateZ", 64.20477863)
+        cmds.setAttr("index_02_r.rotateY", -2.481003482e-014)
+        cmds.setAttr("index_02_r.rotateX", 5.153492507e-014)
+        cmds.setAttr("index_01_r.rotateZ", 46.94060635)
+        cmds.setAttr("index_01_r.rotateY", 3.692762633)
+        cmds.setAttr("index_01_r.rotateX", -8.969261741)
+
+        # Delete the history on the meshes so that they get baked down into their posed positions
+        cmds.delete(meshes, ch=True)
+
+        # Move the skeleton back to its bindpose (This script REQUIRES the bind pose to be the pose that we want to export the skeleton to.  This is the case for all known characters created with A.R.T.)
+        cmds.dagPose("root", bp=True, r=True) 
+
+        # Re-import the skinning we exported before onto the mesh so thats its bound in its posed position.
+        cmds.select(meshes)
+        impw.importMultipleSkinWeights("", True)
+        
+        jointsToDelete = []
+
+        # remove unneeded joints
+        jointsToRemove = ["hip_pack_a", "hip_pack_b", "hip_pack_c", "hip_pack_d"]
+        jointToTransferTo = "pelvis"
+        transfer.LOD_transferWeights(meshes, jointsToRemove, jointToTransferTo)
+        jointsToDelete.extend(jointsToRemove)
+
+        jointsToRemove = ["trap_handle"]
+        jointToTransferTo = "spine_01"
+        transfer.LOD_transferWeights(meshes, jointsToRemove, jointToTransferTo)
+        jointsToDelete.extend(jointsToRemove)
+
+        jointsToRemove = ["sniper_can_arm_01", "sniper_can_arm_02", "sniper_can_arm_03", "pack_hose_out", "pack_hose_in"]
+        jointToTransferTo = "spine_03"
+        transfer.LOD_transferWeights(meshes, jointsToRemove, jointToTransferTo)
+        jointsToDelete.extend(jointsToRemove)
+
+        jointsToRemove = ["shotgun", "shotgun_clip"]
+        jointToTransferTo = "shotgun_base"
+        transfer.LOD_transferWeights(meshes, jointsToRemove, jointToTransferTo)
+        jointsToDelete.extend(jointsToRemove)
+
+        jointsToRemove = ["upperarm_twist_01_l", "upperarm_twist_02_l", "arm_pouch_l"]
+        jointToTransferTo = "upperarm_l"
+        transfer.LOD_transferWeights(meshes, jointsToRemove, jointToTransferTo)
+        jointsToDelete.extend(jointsToRemove)
+
+        jointsToRemove = ["upperarm_twist_01_r", "bot_uparm", "bot_uparm_pistonC_top", "upperarm_twist_02_r"]
+        jointToTransferTo = "upperarm_r"
+        transfer.LOD_transferWeights(meshes, jointsToRemove, jointToTransferTo)
+        jointsToDelete.extend(jointsToRemove)
+
+        jointsToRemove = ["lowerarm_twist_01_l", "lowerarm_twist_02_l"]
+        jointToTransferTo = "lowerarm_l"
+        transfer.LOD_transferWeights(meshes, jointsToRemove, jointToTransferTo)
+        jointsToDelete.extend(jointsToRemove)
+
+        jointsToRemove = ["arm_hose_top", "bot_loarm", "bot_loarm_pistonB_top", "bot_loarm_pistonC_top", "bot_loarm_pistonD_top", "bot_loarm_twist", "bot_loarm_pistonC_bot", "bot_loarm_pistonD_bot", "bot_uparm_pistonC_bot", "lowerarm_twist_01_r", "arm_hose_bott", "wrist_cover"]
+        jointToTransferTo = "lowerarm_r"
+        transfer.LOD_transferWeights(meshes, jointsToRemove, jointToTransferTo)
+        jointsToDelete.extend(jointsToRemove)
+
+        jointsToRemove = ["index_01_l", "index_02_l", "index_03_l", "middle_01_l", "middle_02_l", "middle_03_l", "pinky_01_l", "pinky_02_l", "pinky_03_l", "ring_01_l", "ring_02_l", "ring_03_l", "thumb_01_l", "thumb_02_l", "thumb_03_l"]
+        jointToTransferTo = "hand_l"
+        transfer.LOD_transferWeights(meshes, jointsToRemove, jointToTransferTo)
+        jointsToDelete.extend(jointsToRemove)
+
+        jointsToRemove = ["bot_hand", "gun", "gun_barrel", "gun_body_panel_l", "gun_body_panel_r", "gun_scope", "gun_stock_base", "gun_barrel_inner", "gun_stock_base_in_bott", "gun_stock_base_out_bott", "gun_stock_base_in_top", "gun_stock_base_out_top", "gun_stock_in_bott", "gun_stock_in_top", "gun_stock_bott", "gun_stock_panel_l", "gun_stock_panel_r", "bot_loarm_pistonB_bot", "index_01_r", "index_02_r", "index_03_r", "middle_01_r", "middle_02_r", "middle_03_r", "pinky_01_r", "pinky_02_r", "pinky_03_r", "ring_01_r", "ring_02_r", "ring_03_r", "thumb_01_r", "thumb_02_r", "thumb_03_r"]
+        jointToTransferTo = "hand_r"
+        transfer.LOD_transferWeights(meshes, jointsToRemove, jointToTransferTo)
+        jointsToDelete.extend(jointsToRemove)
+
+        jointsToRemove = ["eyeball_l", "eyeball_r", "eyebrow_in_l", "eyebrow_in_r", "eyebrow_out_l", "eyebrow_out_r", "eyelid_bott_l", "eyelid_bott_r", "eyelid_top_l", "eyelid_top_r", "helmet_sight", "jaw", "mask"]
+        jointToTransferTo = "head"
+        transfer.LOD_transferWeights(meshes, jointsToRemove, jointToTransferTo)
+        jointsToDelete.extend(jointsToRemove)
+
+        jointsToRemove = ["thigh_twist_01_l", "thigh_twist_02_l"]
+        jointToTransferTo = "thigh_l"
+        transfer.LOD_transferWeights(meshes, jointsToRemove, jointToTransferTo)
+        jointsToDelete.extend(jointsToRemove)
+
+        jointsToRemove = ["thigh_twist_01_r", "thigh_twist_02_r"]
+        jointToTransferTo = "thigh_r"
+        transfer.LOD_transferWeights(meshes, jointsToRemove, jointToTransferTo)
+        jointsToDelete.extend(jointsToRemove)
+
+        jointsToRemove = ["ankle_plate_in_l", "ankle_plate_out_l", "calf_twist_01_l"]
+        jointToTransferTo = "calf_l"
+        transfer.LOD_transferWeights(meshes, jointsToRemove, jointToTransferTo)
+        jointsToDelete.extend(jointsToRemove)
+
+        jointsToRemove = ["ankle_plate_in_r", "ankle_plate_out_r", "calf_twist_01_r"]
+        jointToTransferTo = "calf_r"
+        transfer.LOD_transferWeights(meshes, jointsToRemove, jointToTransferTo)
+        jointsToDelete.extend(jointsToRemove)
+
+        jointsToRemove = ["ball_l", "heel_plate_l"]
+        jointToTransferTo = "foot_l"
+        transfer.LOD_transferWeights(meshes, jointsToRemove, jointToTransferTo)
+        jointsToDelete.extend(jointsToRemove)
+
+        jointsToRemove = ["ball_r", "heel_plate_r"]
+        jointToTransferTo = "foot_r"
+        transfer.LOD_transferWeights(meshes, jointsToRemove, jointToTransferTo)
+        jointsToDelete.extend(jointsToRemove)
+        
+        # Delete the joints we no longer want.
+        for one in jointsToDelete:
+            if cmds.objExists(one):
+                cmds.delete(one)
+                
+        # Export the fbx file.
+        export.exportLOD("Price_LOD4.fbx", "LOD4")
+        
+        # Re-open the file without saving.
+        fullPath = cmds.file(q = True, sceneName = True)
+        cmds.file(fullPath, open=True, f=True)
+
+        print "LOD4 Exported Successfully!!!!!!!!!!!!!!!"
+
+PriceLOD4()
+
+
diff --git a/Engine/Extras/Maya_AnimationRiggingTools/ARTv1/General/ART/Projects/Orion/Scripts/Price_Post_Script.py b/Engine/Extras/Maya_AnimationRiggingTools/ARTv1/General/ART/Projects/Orion/Scripts/Price_Post_Script.py
new file mode 100644
--- /dev/null
+++ b/Engine/Extras/Maya_AnimationRiggingTools/ARTv1/General/ART/Projects/Orion/Scripts/Price_Post_Script.py
@@ -0,0 +1,562 @@
+import maya.cmds as cmds
+import Tools.ART_addSpaceSwitching as space
+reload(space)
+
+#NOTE This rig has custom changes to be made during the RigPose placement phase.  See the commented out script at the bottom of this file.
+
+
+# ------------------------------------------------------------------
+# scale all of the controls to be more visible with the mesh.
+def scaleControls():
+    objects = ["bot_uparm_pistonC_top_anim", "bot_uparm_pistonC_bot_anim", "bot_loarm_pistonB_top_anim", "bot_loarm_pistonB_bot_anim", "bot_loarm_pistonC_top_anim", "bot_loarm_pistonC_bot_anim", "bot_loarm_pistonD_top_anim", "bot_loarm_pistonD_bot_anim", "bot_loarm_anim", "bot_loarm_twist_anim", "bot_hand_anim", "gun_scope_anim", "gun_stock_base_anim", "gun_stock_bott_anim", "gun_barrel_inner_anim", "gun_body_panel_l_anim", "gun_body_panel_r_anim", "gun_stock_base_in_bott_anim", "gun_stock_base_out_bott_anim", "gun_stock_base_in_top_anim", "gun_stock_base_out_top_anim", "gun_stock_in_bott_anim", "gun_stock_in_top_anim", "trap_handle_anim", "jaw_anim", "sniper_can_arm_01_anim"]
+    for one in objects:
+        cmds.select([one+".cv[0:32]"])
+        cmds.scale(0.227057, 0.227057, 0.227057, r = True, ocp = True)
+        cmds.move(0, -5, 0, r = True, os=True, wd=True)
+        
+    objects = ["bot_uparm_anim", "gun_anim", "gun_stock_panel_r_anim", "gun_stock_panel_l_anim", "gun_barrel_anim", "shotgun_base_anim", "shotgun_clip_anim", "shotgun_anim", "ankle_plate_in_l_anim", "ankle_plate_in_r_anim",  "ankle_plate_out_l_anim", "ankle_plate_out_r_anim", "sniper_can_arm_02_anim", "sniper_can_arm_03_anim", "arm_hose_bott_anim", "arm_hose_top_anim", "mask_anim", "helmet_sight_anim", "arm_pouch_l_anim", "pack_hose_out_anim", "pack_hose_in_anim", "heel_plate_l_anim", "heel_plate_r_anim", "wrist_cover_anim"]
+    for one in objects:
+        cmds.select([one+".cv[0:32]"])
+        cmds.scale(0.2, 0.2, 0.2, r = True, ocp = True)
+        
+    for side in ["_l", "_r"]:
+        #ik feet
+        cmds.select("ik_foot_anim"+side + ".cv[0:32]")
+        cmds.scale(1, 1, 1, r=True)
+        cmds.move(0, 0, -4.695816, r = True, os=True, wd=True)
+    
+        cmds.select(["eyebrow_in"+side+"_anim.cv[0:32]"])
+        cmds.scale(0.117744, 0.117744, 0.117744, r=True, ocp = False)
+        cmds.move(0, 0, 0, r=True)
+
+        cmds.select(["eyebrow_out"+side+"_anim.cv[0:32]"])
+        cmds.scale(0.117744, 0.117744, 0.117744, r=True, ocp = False)
+        cmds.move(0, 0, 0, r=True)
+
+        cmds.select(["eyeball"+side+"_anim.cv[0:32]"])
+        cmds.scale(0.1, 0.1, 0.05, r=True, ocp = False)
+        if side == "_l":
+            cmds.move(0, -3.757386, 0, r=True, os=True)
+        else:
+            cmds.move(0, -3.757386, 0, r=True, os=True)
+
+        cmds.select(["eyelid_bott"+side+"_anim.cv[0:32]"])
+        cmds.scale(0.1, 0.1, 0.01, r=True, ocp = False)
+        if side == "_l":
+            cmds.move(0, -3.757386, 0.5, r=True, os=True)
+        else:
+            cmds.move(0, 3.757386, 0.5, r=True, os=True)
+
+        cmds.select(["eyelid_top"+side+"_anim.cv[0:32]"])
+        cmds.scale(0.1, 0.1, 0.01, r=True, ocp = False)
+        if side == "_l":
+            cmds.move(0, -3.757386, 0.5, r=True, os=True)
+        else:
+            cmds.move(0, 3.757386, 0.5, r=True, os=True)
+            
+    cmds.select(["gun_anim.cv[0:32]"])
+    cmds.move(8, 0, 0, r = True, os=True, wd=True)
+
+    cmds.select(["gun_scope_anim.cv[0:32]"])
+    cmds.move(0, 10, 0, r = True, os=True, wd=True)
+
+    cmds.select(["gun_stock_panel_r_anim.cv[0:32]"])
+    cmds.move(2.5, 0, 0, r = True, os=True, wd=True)
+
+    cmds.select(["gun_stock_panel_l_anim.cv[0:32]"])
+    cmds.move(-2.5, 0, 0, r = True, os=True, wd=True)
+
+    cmds.select(["gun_stock_bott_anim.cv[0:32]"])
+    cmds.move(0, -1, 0, r = True, os=True, wd=True)
+
+    cmds.select(["gun_barrel_anim.cv[0:32]"])
+    cmds.move(0, 0, -5, r = True, os=True, wd=True)
+
+    cmds.select(["gun_stock_base_out_top_anim.cv[0:32]"])
+    cmds.move(0, 10, 0, r = True, os=True, wd=True)
+
+    cmds.select(["gun_stock_base_in_top_anim.cv[0:32]"])
+    cmds.move(0, 10, 0, r = True, os=True, wd=True)
+
+    cmds.select(["gun_stock_in_top_anim.cv[0:32]"])
+    cmds.move(0, 15, 0, r = True, os=True, wd=True)
+
+    cmds.select(["bot_hand_anim.cv[0:32]"])
+    cmds.move(0, 0, 6, r = True, os=True, wd=True)
+
+    cmds.select(["bot_loarm_anim.cv[0:32]"])
+    cmds.move(0, 10, 0, r = True, os=True, wd=True)
+
+    cmds.select(["bot_uparm_anim.cv[0:32]"])
+    cmds.move(0, 5, -8, r = True, os=True, wd=True)
+
+    cmds.select(["shotgun_base_anim.cv[0:32]"])
+    cmds.move(7, 0, 0, r = True, os=True, wd=True)
+
+    cmds.select(["shotgun_anim.cv[0:32]"])
+    cmds.move(0, 0, 3, r = True, os=True, wd=True)
+
+    cmds.select(["shotgun_clip_anim.cv[0:32]"])
+    cmds.move(-5, 0, 0, r = True, os=True, wd=True)
+
+    cmds.select(["jaw_anim.cv[0:32]"])
+    cmds.move(0, -5, -8, r = True, os=True, wd=True)
+
+    cmds.select(["sniper_can_arm_01_anim.cv[0:32]"])
+    cmds.move(0, 7, 9, r = True, os=True, wd=True)
+
+    cmds.select(["sniper_can_arm_02_anim.cv[0:32]"])
+    cmds.move(0, 5, 5, r = True, os=True, wd=True)
+
+    cmds.select(["sniper_can_arm_03_anim.cv[0:32]"])
+    cmds.scale(1.0, 6.3, 6.3, r = True, ocp = True)
+    cmds.move(0, 5, 0, r = True, os=True, wd=True)
+
+    cmds.select(None)
+scaleControls()
+
+# ------------------------------------------------------------------
+def createPistons():
+    import Tools.ART_CreateHydraulicRigs as hyd
+    reload(hyd)
+    
+    # pistons
+    cmds.select("bot_uparm_pistonC_top_anim", "bot_uparm_pistonC_bot_anim", "driver_upperarm_r")
+    hyd.CreateHydraulicRig(-1)
+    cmds.select("bot_loarm_pistonB_top_anim", "bot_loarm_pistonB_bot_anim", "driver_lowerarm_r")
+    hyd.CreateHydraulicRig(-1)
+    cmds.select("bot_loarm_pistonC_top_anim", "bot_loarm_pistonC_bot_anim", "driver_lowerarm_r")
+    hyd.CreateHydraulicRig(-1)
+    cmds.select("bot_loarm_pistonD_top_anim", "bot_loarm_pistonD_bot_anim", "driver_lowerarm_r")
+    hyd.CreateHydraulicRig(-1)
+
+    # manually changing the up object for a few of the constraints that were created above so that they rotate at teh correct angles.
+    cmds.aimConstraint("bot_loarm_pistonB_bot_anim_sub_aimConstraint1", e=True, mo=True, wuo="driver_bot_hand")
+    cmds.aimConstraint("bot_loarm_pistonC_bot_anim_sub_aimConstraint1", e=True, mo=True, u=(0,-1,0), wuo="bot_loarm_pistonD_bot_anim")
+    cmds.aimConstraint("bot_loarm_pistonD_bot_anim_sub_aimConstraint1", e=True, mo=True, u=(0,-1,0), wuo="bot_loarm_pistonC_bot_anim")
+    cmds.aimConstraint("bot_uparm_pistonC_bot_anim_sub_aimConstraint1", e=True, mo=True, wuo="driver_lowerarm_r")
+
+    # bot arm twist
+    cmds.orientConstraint("driver_lowerarm_twist_01_r", "bot_loarm_twist_anim_space_switcher", mo = True)
+createPistons()
+
+# ------------------------------------------------------------------
+def anklePlatesRig():
+    for side in ["_l", "_r"]:
+        sidegrp = cmds.group(em=True, n="ankle_plates_grp"+side)
+        const = cmds.parentConstraint("driver_calf_twist_01"+side, sidegrp)[0]
+        cmds.delete(const)
+        cmds.makeIdentity(sidegrp, a=True, t=True, r=True, s=True)
+        cmds.parent(sidegrp, "ctrl_rig")
+        cmds.parentConstraint("driver_calf_twist_01"+side, sidegrp, mo=True)
+        
+        for dir in ["_in", "_out"]:
+            target = cmds.spaceLocator(n="ankle_plate"+dir+"_target"+side)[0]
+            const = cmds.parentConstraint("ankle_plate"+dir+side+"_anim", target)[0]
+            cmds.delete(const)
+            if side == "_l":
+                cmds.move(-10, 0, 0, target, r=True, os=True, wd=True)
+            else:
+                cmds.move(10, 0, 0, target, r=True, os=True, wd=True)
+            
+            up = cmds.spaceLocator(n="ankle_plate"+dir+"_up"+side)[0]
+            const = cmds.parentConstraint("ankle_plate"+dir+side+"_anim", up)[0]
+            cmds.delete(const)
+            if side == "_l":
+                cmds.move(0, 0, -5, up, r=True, os=True, wd=True)
+            else:
+                cmds.move(0, 0, 5, up, r=True, os=True, wd=True)
+
+            
+            grp = cmds.group(em=True, n="ankle_plate"+dir+"_grp"+side)
+            const = cmds.parentConstraint("driver_foot"+side, grp)[0]
+            cmds.delete(const)
+            cmds.makeIdentity(grp, a=True, t=True, r=True, s=True)
+            
+            cmds.parent(target, up, grp)
+            cmds.parent(grp, sidegrp)
+            
+            driver_foot_loc = cmds.xform("driver_foot"+side, q=True, ws=True, rp=True)
+            cmds.xform("ankle_plate"+dir+side+"_anim_grp", ws=True, rp=driver_foot_loc)
+
+            cmds.parentConstraint("driver_foot"+side, grp, mo=True, sr=["x"])
+          
+            cmds.aimConstraint(target, "ankle_plate"+dir+side+"_anim_grp", mo=True, aimVector=(-1, 0, 0), upVector=(0, 0, -1), worldUpType="object", worldUpVector=(0, 0, -1), worldUpObject=str(up))
+            
+  
+            cmds.parentConstraint("driver_foot"+side, "ankle_plate"+dir+side+"_anim_grp", mo=True, sr=["x", "y", "z"])
+            
+            cmds.setAttr("ankle_plate"+dir+"_grp"+side+".rotateOrder", 4)
+            
+            if side == "_r":
+                cmds.setAttr("ankle_plate"+dir+"_grp"+side+"_parentConstraint1.target[0].targetOffsetRotateX", -180)
+                cmds.setAttr("ankle_plate"+dir+"_grp"+side+"_parentConstraint1.target[0].targetOffsetRotateY", -90)
+            
+            cmds.setAttr(grp+".v", 0)
+            
+anklePlatesRig()
+
+# ------------------------------------------------------------------
+def heelRig():
+    for side in ["_l", "_r"]:
+        cmds.setDrivenKeyframe(["heel_plate"+side+"_anim_grp.rx"], cd = "driver_foot"+side+".rz", itt = "spline", ott = "spline")
+        cmds.setAttr("ik_foot_anim"+side+".rx", 40)
+        cmds.setAttr("heel_plate"+side+"_anim_grp.rx", -17.471)
+        cmds.setDrivenKeyframe(["heel_plate"+side+"_anim_grp.rx"], cd = "driver_foot"+side+".rz", itt = "spline", ott = "spline")
+        cmds.setAttr("ik_foot_anim"+side+".rx", 80)
+        cmds.setAttr("heel_plate"+side+"_anim_grp.rx", -59.01)
+        cmds.setDrivenKeyframe(["heel_plate"+side+"_anim_grp.rx"], cd = "driver_foot"+side+".rz", itt = "spline", ott = "spline")
+        cmds.setAttr("ik_foot_anim"+side+".rx", 0)
+heelRig()
+
+# ------------------------------------------------------------------
+def trapHandleRig():
+    cmds.parentConstraint("driver_spine_02", "driver_spine_03", "trap_handle_parent_grp", mo=True)
+    cmds.setAttr("trap_handle_parent_grp_parentConstraint1.driver_spine_01W0", 0.12)
+    cmds.setAttr("trap_handle_parent_grp_parentConstraint1.driver_spine_02W1", 0.18)
+    cmds.setAttr("trap_handle_parent_grp_parentConstraint1.driver_spine_03W2", 0.7)
+
+trapHandleRig()
+
+# ------------------------------------------------------------------
+def wristPlatesRig():
+    cmds.rotate(-10.523881, 0, 0, "wrist_cover_anim_grp")
+    
+    target = cmds.spaceLocator(n="wrist_cover_target")[0]
+    const = cmds.parentConstraint("driver_bot_loarm_pistonB_bot", target)[0]
+    cmds.delete(const)
+    
+    up = cmds.spaceLocator(n="wrist_cover_up")[0]
+    const = cmds.parentConstraint("driver_wrist_cover", up)[0]
+    cmds.delete(const)
+    cmds.move(0, 5, 0, up, r=True, os=True, wd=True)
+    
+    grp = cmds.group(em=True, n="wrist_cover_grp")
+    const = cmds.parentConstraint("driver_hand_r", grp)[0]
+    cmds.delete(const)
+    cmds.makeIdentity(grp, a=True, t=True, r=True, s=True)
+    
+    cmds.parent(target, up, grp)
+    cmds.parent(grp, "ctrl_rig")
+    
+    prntConst = cmds.parentConstraint("driver_bot_loarm_pistonB_bot", "driver_hand_r", grp, mo=True)[0]
+    #cmds.pointConstraint("driver_bot_loarm_pistonB_bot", "driver_hand_r", grp, mo=True)
+    #ornConst = cmds.orientConstraint("driver_bot_loarm_pistonB_bot", "driver_hand_r", grp, mo=True)[0]
+    cmds.setAttr(prntConst+".interpType", 2)
+
+    cmds.aimConstraint(target, "wrist_cover_anim_grp", mo=True, aimVector=(-1, 0, 0), upVector=(0, 1, 0), worldUpType="object", worldUpVector=(0, 1, 0), worldUpObject=str(up))
+    cmds.parentConstraint("driver_bot_loarm_pistonB_bot", "driver_hand_r", "wrist_cover_anim_grp", mo=True, sr=["x", "y", "z"])
+    #cmds.pointConstraint("driver_bot_loarm_pistonB_bot", "driver_hand_r", "wrist_cover_anim_grp", mo=True)
+    
+    cmds.setAttr(grp+".v", 0)
+wristPlatesRig()
+
+# ------------------------------------------------------------------
+def shotgunRig():
+    import Tools.ART_addSpaceSwitching as space
+    reload(space)
+
+    cmds.addAttr("shotgun_anim", ln = "shotgun_rotate", dv = 0, min = 0, max = 10, keyable = True)
+    cmds.addAttr("shotgun_anim", ln = "shotgun_rail", dv = 0, min = 0, max = 10, keyable = True)
+
+    shotgun_rail_grp = cmds.group("shotgun_anim", n="shotgun_anim_rail_grp")
+    cmds.xform(shotgun_rail_grp, ws=True, piv=[16.745853637042327, 26.165296554565593, 153.62670823479255])
+
+    cmds.setDrivenKeyframe(["shotgun_anim_rail_grp.rz"], cd = "shotgun_anim.shotgun_rotate", itt = "auto", ott = "auto")
+    cmds.setAttr("shotgun_anim.shotgun_rotate", 10)
+    cmds.setAttr("shotgun_anim_rail_grp.rz", -90)
+    cmds.setDrivenKeyframe(["shotgun_anim_rail_grp.rz"], cd = "shotgun_anim.shotgun_rotate", itt = "auto", ott = "auto")
+    cmds.setAttr("shotgun_anim.shotgun_rotate", 0)
+
+    cmds.setDrivenKeyframe(["shotgun_base_anim_grp.tx", "shotgun_base_anim_grp.ty", "shotgun_base_anim_grp.tz", "shotgun_base_anim_grp.rx"], cd = "shotgun_anim.shotgun_rail", itt = "auto", ott = "auto")
+    cmds.setAttr("shotgun_anim.shotgun_rail", 5)
+    cmds.setAttr("shotgun_base_anim_grp.tx", 3.14)
+    cmds.setAttr("shotgun_base_anim_grp.ty", -3.774)
+    cmds.setAttr("shotgun_base_anim_grp.tz", 9.402)
+    cmds.setAttr("shotgun_base_anim_grp.rx", 13.9)
+    cmds.setDrivenKeyframe(["shotgun_base_anim_grp.tx", "shotgun_base_anim_grp.ty", "shotgun_base_anim_grp.tz", "shotgun_base_anim_grp.rx"], cd = "shotgun_anim.shotgun_rail", itt = "auto", ott = "auto")
+    cmds.setAttr("shotgun_anim.shotgun_rail", 10)
+    cmds.setAttr("shotgun_base_anim_grp.tx", 4.999)
+    cmds.setAttr("shotgun_base_anim_grp.ty", -8.562)
+    cmds.setAttr("shotgun_base_anim_grp.tz", 12.609)
+    cmds.setAttr("shotgun_base_anim_grp.rx", 72.809)
+    cmds.setDrivenKeyframe(["shotgun_base_anim_grp.tx", "shotgun_base_anim_grp.ty", "shotgun_base_anim_grp.tz", "shotgun_base_anim_grp.rx"], cd = "shotgun_anim.shotgun_rail", itt = "auto", ott = "auto")
+    cmds.setAttr("shotgun_anim.shotgun_rail", 0)
+
+    space.addSpaceSwitching("shotgun_anim", "driver_hand_l", False)
+shotgunRig()
+
+# ------------------------------------------------------------------
+def shoulderPadRig():
+    cmds.setAttr("Rig_Settings.rArmMode", 0)
+
+    shoulderPad_grp = cmds.group("bot_uparm_anim_grp", n="bot_uparm_anim_pivot")
+    cmds.xform(shoulderPad_grp, ws=True, piv=[16.745853637042327, 26.165296554565593, 153.62670823479255])
+
+    driver_upperarm_r_loc = cmds.xform("driver_upperarm_r", q=True, ws=True, rp=True)
+    print driver_upperarm_r_loc
+    cmds.xform(shoulderPad_grp, ws=True, rp=driver_upperarm_r_loc)
+
+    cmds.setAttr("fk_arm_r_anim.ry", 0)
+    cmds.setAttr("bot_uparm_anim_pivot.ry", 19.153)
+    cmds.setAttr("bot_uparm_anim_grp.tx", 2.209)
+    cmds.setAttr("bot_uparm_anim_grp.tz", -2.583)
+    cmds.setAttr("bot_uparm_anim_grp.ry", -38.873)
+    cmds.setDrivenKeyframe(["bot_uparm_anim_pivot.ry", "bot_uparm_anim_grp.tx", "bot_uparm_anim_grp.tz", "bot_uparm_anim_grp.ry"], cd = "driver_upperarm_r.ry", itt = "spline", ott = "spline")
+    cmds.setAttr("fk_arm_r_anim.ry", 45)
+    cmds.setAttr("bot_uparm_anim_pivot.ry", 0)
+    cmds.setAttr("bot_uparm_anim_grp.tx", 0)
+    cmds.setAttr("bot_uparm_anim_grp.tz", 0)
+    cmds.setAttr("bot_uparm_anim_grp.ry", 0)
+    cmds.setDrivenKeyframe(["bot_uparm_anim_pivot.ry", "bot_uparm_anim_grp.tx", "bot_uparm_anim_grp.tz", "bot_uparm_anim_grp.ry"], cd = "driver_upperarm_r.ry", itt = "spline", ott = "spline")
+    cmds.setAttr("fk_arm_r_anim.ry", -20)
+    cmds.setAttr("bot_uparm_anim_pivot.ry", 27.813)
+    cmds.setAttr("bot_uparm_anim_grp.tx", 6.98)
+    cmds.setAttr("bot_uparm_anim_grp.tz", -1.586)
+    cmds.setAttr("bot_uparm_anim_grp.ry", -58.887)
+    cmds.setDrivenKeyframe(["bot_uparm_anim_pivot.ry", "bot_uparm_anim_grp.tx", "bot_uparm_anim_grp.tz", "bot_uparm_anim_grp.ry"], cd = "driver_upperarm_r.ry", itt = "spline", ott = "spline")
+    cmds.setAttr("fk_arm_r_anim.ry", 0)
+    
+    cmds.setAttr("Rig_Settings.rArmMode", 1)
+
+shoulderPadRig()
+
+# ------------------------------------------------------------------
+def sniperCanisterRig():
+    import Tools.ART_addSpaceSwitching as space
+    reload(space)
+
+    space.addSpaceSwitching("sniper_can_arm_03_anim", "driver_hand_l", False)
+    space.addSpaceSwitching("sniper_can_arm_03_anim", "driver_gun", False)
+sniperCanisterRig()
+
+# ------------------------------------------------------------------
+def trapRig():
+    import Tools.ART_addSpaceSwitching as space
+    reload(space)
+
+    space.addSpaceSwitching("trap_handle_anim", "driver_hand_l", False)
+trapRig()
+
+# ------------------------------------------------------------------
+def leftHandRig():
+    import Tools.ART_addSpaceSwitching as space
+    reload(space)
+
+    space.addSpaceSwitching("ik_wrist_l_anim", "driver_gun_stock_bott", False)
+    
+    cmds.setAttr("Rig_Settings.lForearmTwistAmount", 0.9)
+    cmds.setAttr("Rig_Settings.lForearmTwist2Amount", 0.6)
+    
+leftHandRig()
+
+# ------------------------------------------------------------------
+def maskRig():
+    import Tools.ART_addSpaceSwitching as space
+    reload(space)
+
+    space.addSpaceSwitching("mask_anim", "driver_hand_l", False)
+    space.addSpaceSwitching("mask_anim", "driver_hand_r", False)
+maskRig()
+
+# ------------------------------------------------------------------
+def invertEyelids():
+    sides = ["_l", "_r"]
+    for side in sides:
+        cmds.delete("driver_eyelid_bott"+side+"_parentConstraint1")
+        cmds.rotate(0, -180, 0, "eyelid_bott"+side+"_anim_grp", r=True, os=True)
+        cmds.parentConstraint("eyelid_bott"+side+"_anim", "driver_eyelid_bott"+side, mo=True)
+invertEyelids()
+
+def invertRightEyeball():
+    cmds.delete("driver_eyeball_r_parentConstraint1")
+    cmds.rotate(0, -180, -180, "eyeball_r_anim_grp", r=True, os=True)
+    cmds.parentConstraint("eyeball_r_anim", "driver_eyeball_r", mo=True)
+invertRightEyeball()
+
+# ------------------------------------------------------------------
+def dynamicsOff():
+    for one in ["hip_pack_a_anim", "hip_pack_b_anim", "hip_pack_c_anim", "hip_pack_d_anim"]:
+        cmds.setAttr(one+".chainStartEnvelope", 0)
+        cmds.setAttr(one+".chainStartEnvelope", l=True)
+dynamicsOff()
+
+# ------------------------------------------------------------------
+def createGameCam():
+    gameCam = cmds.camera(centerOfInterest=5, focalLength=18, lensSqueezeRatio=1, cameraScale=1, horizontalFilmAperture=1.41732, horizontalFilmOffset=0, verticalFilmAperture=0.94488, verticalFilmOffset=0, filmFit="Overscan", overscan=1.05, motionBlur=0, shutterAngle=144, nearClipPlane=1.0, farClipPlane=10000, orthographic=0, orthographicWidth=30, panZoomEnabled=0, horizontalPan=0, verticalPan=0, zoom=1, displayFilmGate=True)[0]
+    print gameCam
+    gameCam = cmds.rename(gameCam, "GameCam")
+    cmds.camera(gameCam, e=True, horizontalFilmAperture=1.6724376, focalLength=21.2399857011)
+    cmds.setAttr(gameCam+".displayGateMaskOpacity", 1)
+    cmds.setAttr(gameCam+".displayGateMaskColor", 0, 0, 0, type="double3")
+    cmds.setAttr(gameCam+".r", 90, 0, 180, type="double3")
+    
+    # Create the control
+    control = cmds.circle(name=gameCam+"_Ctrl", ch=True, o=True, r=15)[0]
+    cmds.makeIdentity(control, a=True, t=True, r=True, s=True)
+    # Create groups for the sdk and the constraints
+    const_grp = cmds.group(control, name=control+"_grp")
+    if cmds.objExists("rig_grp"):
+        cmds.parent(const_grp, "rig_grp")
+    # Parent Constrain the cam to the control
+    cmds.parentConstraint(control, gameCam, mo=True)
+    
+    cmds.connectAttr(control+".visibility", gameCam+".visibility", f=True)
+    cmds.setAttr(const_grp+".v", 0)
+    
+    cmds.setAttr(const_grp+".t", 0, 500, 200, type="double3")
+createGameCam()
+
+# ------------------------------------------------------------------
+def addCapsule(radius, halfHeight):
+    trueHalfHeight = (halfHeight*2)-(radius*2)
+    capsule = cmds.polyCylinder(name="collision_capsule", r=radius, h=trueHalfHeight, sx=18, sy=1, sz=5, ax=[0,0,1], rcp=True, cuv=3, ch=True)
+    cmds.setAttr(capsule[0]+".tz", halfHeight+2.2)
+    cmds.setAttr(capsule[0]+".rz", 90)
+
+    cmds.addAttr(capsule[0], ln="Radius", dv=radius, keyable=True)
+    cmds.addAttr(capsule[0], ln="HalfHeight", dv=halfHeight, keyable=True)
+
+    expr = capsule[1]+".radius = "+capsule[0]+".Radius;\n"+capsule[1]+".height = ("+capsule[0]+".HalfHeight*2) - ("+capsule[0]+".Radius*2);"
+    cmds.expression(s=expr, o=capsule[1], ae=1, uc="all")
+
+    cmds.setAttr(capsule[0]+".overrideEnabled", 1)
+    cmds.setAttr(capsule[0]+".overrideColor", 18)
+    cmds.setAttr(capsule[0]+".overrideShading", 0)
+
+    cmds.addAttr(capsule[0], ln="Shaded", dv=0, min=0, max=1, keyable=True)
+    cmds.connectAttr(capsule[0]+".Shaded", capsule[0]+".overrideShading")
+
+    cmds.parentConstraint("root", capsule[0], mo=True)
+    cmds.select(capsule[0], r=True)
+    capsuleDisplayLayer = cmds.createDisplayLayer(nr=True, name="Collision Capsule")
+    cmds.setAttr(capsuleDisplayLayer+".displayType", 2)
+    cmds.setAttr(capsuleDisplayLayer+".v", 0)
+    
+    for attr in ["tx", "ty", "tz", "rx", "ry", "rz", "sx", "sy", "sz"]:
+        cmds.setAttr(capsule[0]+"."+attr, lock=True, keyable=False, channelBox=False)
+
+    cmds.select(clear=True)
+addCapsule(42, 86)
+
+# ------------------------------------------------------------------
+# ------------------------------------------------------------------
+# add ik bones Function that is called later after we open the export file.
+def addIKBones():
+    try:
+		# create the joints
+		cmds.select(clear = True)
+		ikFootRoot = cmds.joint(name = "ik_foot_root")
+		cmds.select(clear = True)
+		
+		cmds.select(clear = True)
+		ikFootLeft = cmds.joint(name = "ik_foot_l")
+		cmds.select(clear = True)
+		
+		cmds.select(clear = True)
+		ikFootRight = cmds.joint(name = "ik_foot_r")
+		cmds.select(clear = True)
+		
+		cmds.select(clear = True)
+		ikHandRoot = cmds.joint(name = "ik_hand_root")
+		cmds.select(clear = True)
+		
+		cmds.select(clear = True)
+		ikHandGun = cmds.joint(name = "ik_hand_gun")
+		cmds.select(clear = True)
+		
+		cmds.select(clear = True)
+		ikHandLeft = cmds.joint(name = "ik_hand_l")
+		cmds.select(clear = True)
+		
+		cmds.select(clear = True)
+		ikHandRight = cmds.joint(name = "ik_hand_r")
+		cmds.select(clear = True)
+		
+		# create hierarhcy
+		cmds.parent(ikFootRoot, "root")
+		cmds.parent(ikHandRoot, "root")
+		cmds.parent(ikFootLeft, ikFootRoot)
+		cmds.parent(ikFootRight, ikFootRoot)
+		cmds.parent(ikHandGun, ikHandRoot)
+		cmds.parent(ikHandLeft, ikHandGun)
+		cmds.parent(ikHandRight, ikHandGun)
+		
+		# constrain the joints
+		leftHandConstraint = cmds.parentConstraint("hand_l", ikHandLeft)[0]
+		rightHandGunConstraint = cmds.parentConstraint("hand_r", ikHandGun)[0]
+		rightHandConstraint = cmds.parentConstraint("hand_r", ikHandRight)[0]
+		leftFootConstraint = cmds.parentConstraint("foot_l", ikFootLeft)[0]
+		rightFootConstraint = cmds.parentConstraint("foot_r", ikFootRight)[0]
+	
+    except:
+		cmds.warning("Something went wrong")
+
+# ------------------------------------------------------------------
+# ADD the IK bones to the Export file.
+#save the AnimRig file before opening the export file.
+cmds.file(save = True, type = "mayaBinary", force = True)
+
+#ADD IK BONES
+#open export file
+sceneName = cmds.file(q = True, sceneName = True)
+exportFile = sceneName.partition("AnimRigs")[0] + "ExportFiles" +  sceneName.partition("AnimRigs")[2].partition(".mb")[0] + "_Export.mb"
+cmds.file(exportFile, open = True, force = True)
+
+#add ik bones
+addIKBones()
+
+# ------------------------------------------------------------------
+def addCustomAttrs():
+    attrs = cmds.listAttr("root")
+    exists = False
+    for one in attrs:
+        if one.find("DistanceCurve") != -1:
+            exists = True
+    if exists == False:
+        cmds.addAttr("root", ln="DistanceCurve", keyable=True)
+addCustomAttrs()
+
+#save the export file.
+cmds.file(save = True, type = "mayaBinary", force = True)
+
+######################################################################################
+######################################################################################
+'''
+# This script is only used for aligning the pistons to each other when placing the joint movers.   It is also to be used while creating the Rig Pose.  Don't use it in the post script running.
+# for each piston end, place the mover at the center of the piston disk (visually), and place the UP locator along the +Z axis of the top part of the piston. (+X should be the down axis, -Y should point up and out of the piston.)
+# This script will aim both piston parts at the locator and then rotate (at the end) the bott one to align with the angle it should use if its different.
+def alignPistons():
+    const = cmds.aimConstraint("bot_uparm_pistonC_top_mover", "bot_uparm_pistonC_bot_mover", weight=1, mo=False, aimVector=(-1, 0, 0), upVector=(0, 0, 1), worldUpType="object", worldUpVector=(0, 0, 1), worldUpObject="bot_uparm_pistonC_top_mover_UP")
+    cmds.delete(const)
+    const = cmds.aimConstraint("bot_uparm_pistonC_bot_mover", "bot_uparm_pistonC_top_mover", weight=1, mo=False, aimVector=(1, 0, 0), upVector=(0, 0, 1), worldUpType="object", worldUpVector=(0, 0, 1), worldUpObject="bot_uparm_pistonC_top_mover_UP")
+    cmds.delete(const)        
+    cmds.rotate(77.41, 0, 0, "bot_uparm_pistonC_bot_mover", r=True, os=True)
+
+    const = cmds.aimConstraint("bot_loarm_pistonB_bot_mover", "bot_loarm_pistonB_top_mover", weight=1, mo=False, aimVector=(1, 0, 0), upVector=(0, 0, 1), worldUpType="object", worldUpVector=(0, 0, 1), worldUpObject="bot_loarm_pistonB_top_mover_UP")
+    cmds.delete(const)
+    const = cmds.aimConstraint("bot_loarm_pistonB_top_mover", "bot_loarm_pistonB_bot_mover", weight=1, mo=False, aimVector=(-1, 0, 0), upVector=(0, 0, 1), worldUpType="object", worldUpVector=(0, 0, 1), worldUpObject="bot_loarm_pistonB_top_mover_UP")
+    cmds.delete(const)
+    cmds.rotate(15.8, 0, 0, "bot_loarm_pistonB_bot_mover", r=True, os=True)
+
+    const = cmds.aimConstraint("bot_loarm_pistonC_bot_mover", "bot_loarm_pistonC_top_mover", weight=1, mo=False, aimVector=(-1, 0, 0), upVector=(0, 0, 1), worldUpType="object", worldUpVector=(0, 0, 1), worldUpObject="bot_loarm_pistonC_bot_mover_UP")
+    cmds.delete(const)
+    const = cmds.aimConstraint("bot_loarm_pistonC_top_mover", "bot_loarm_pistonC_bot_mover", weight=1, mo=False, aimVector=(1, 0, 0), upVector=(0, 0, 1), worldUpType="object", worldUpVector=(0, 0, 1), worldUpObject="bot_loarm_pistonC_bot_mover_UP")
+    cmds.delete(const)
+    cmds.rotate(0, 0, 0, "bot_loarm_pistonC_top_mover", r=True, os=True)
+            
+    const = cmds.aimConstraint("bot_loarm_pistonD_bot_mover", "bot_loarm_pistonD_top_mover", weight=1, mo=False, aimVector=(-1, 0, 0), upVector=(0, 0, 1), worldUpType="object", worldUpVector=(0, 0, 1), worldUpObject="bot_loarm_pistonD_bot_mover_UP")
+    cmds.delete(const)
+    const = cmds.aimConstraint("bot_loarm_pistonD_top_mover", "bot_loarm_pistonD_bot_mover", weight=1, mo=False, aimVector=(1, 0, 0), upVector=(0, 0, 1), worldUpType="object", worldUpVector=(0, 0, 1), worldUpObject="bot_loarm_pistonD_bot_mover_UP")
+    cmds.delete(const)
+    cmds.rotate(-84.5, 0, 0, "bot_loarm_pistonD_top_mover", r=True, os=True)
+alignPistons()
+
+
+# Then make sure that the UP locators are parent constrained to the uparm and loarm or loarm twist movers. (the ones that are leafs, not the actual skeleton ones.) Use the following script to do this.
+# REMINDER: This must be done BEFORE you move on to Deformation Setup phase.
+cmds.parentConstraint("bot_uparm_mover", "bot_uparm_pistonC_top_mover_UP", weight=1, mo=True)
+cmds.parentConstraint("bot_loarm_mover", "bot_loarm_pistonB_top_mover_UP", weight=1, mo=True)
+cmds.parentConstraint("bot_loarm_twist_mover", "bot_loarm_pistonC_bot_mover_UP", weight=1, mo=True)
+cmds.parentConstraint("bot_loarm_twist_mover", "bot_loarm_pistonD_bot_mover_UP", weight=1, mo=True)
+'''
\ No newline at end of file
diff --git a/Engine/Extras/Maya_AnimationRiggingTools/ARTv1/General/ART/Projects/Orion/Scripts/Pyro_Goblin_Post_Script.py b/Engine/Extras/Maya_AnimationRiggingTools/ARTv1/General/ART/Projects/Orion/Scripts/Pyro_Goblin_Post_Script.py
new file mode 100644
--- /dev/null
+++ b/Engine/Extras/Maya_AnimationRiggingTools/ARTv1/General/ART/Projects/Orion/Scripts/Pyro_Goblin_Post_Script.py
@@ -0,0 +1,148 @@
+import maya.cmds as cmds
+import maya.mel as mel
+import math
+import Tools.ART_CreateHydraulicRigs as hyd
+reload(hyd)
+import Tools.ART_addSpaceSwitching as space
+reload(space)
+mel.eval("source generateChannelMenu.mel;")
+
+# ------------------------------------------------------------------
+def scaleControls():
+    objects = ["ik_wrist_l_anim", "ik_wrist_r_anim", "index_finger_fk_ctrl_1_l", "thumb_finger_fk_ctrl_1_l", "body_anim", "mid_ik_anim", "chest_ik_anim", "hip_anim",  "neck_01_fk_anim", "head_fk_anim", "fk_wrist_l_anim" , "fk_wrist_r_anim" , "fk_elbow_l_anim", "fk_elbow_r_anim", "fk_arm_l_anim", "fk_arm_r_anim", "index_finger_fk_ctrl_1_r", "index_finger_fk_ctrl_2_r", "index_finger_fk_ctrl_3_r", "middle_finger_fk_ctrl_1_r", "middle_finger_fk_ctrl_2_r", "middle_finger_fk_ctrl_3_r", "ring_finger_fk_ctrl_1_r", "ring_finger_fk_ctrl_2_r", "ring_finger_fk_ctrl_3_r", "pinky_finger_fk_ctrl_1_r", "thumb_finger_fk_ctrl_1_r", "thumb_finger_fk_ctrl_2_r", "thumb_finger_fk_ctrl_3_r", "spine_01_anim", "spine_02_anim", "spine_03_anim", "neck_01_fk_anim"]
+    for one in objects:
+        cmds.select([one+".cv[0:*]"])
+        cmds.scale(2.3, 2.3, 2.3, r=True, os=True)
+
+    objects = ["arm_piston_A_top_anim", "arm_piston_A_bott_anim", "arm_piston_B_top_anim", "arm_piston_B_bott_anim", "muscle_belly_anim", "toe_big_a_l_anim", "toe_big_a_r_anim", "toe_mid_a_l_anim",  "toe_mid_a_r_anim", "toe_pink_a_l_anim", "toe_pink_a_r_anim"]
+    for one in objects:
+        cmds.select([one+".cv[0:*]"])
+        cmds.scale(0.36, 0.36, 0.36, r=True, os=True)
+
+    # eyes
+    for side in ["_l", "_r"]:
+        cmds.select(["eyeball"+side+"_anim.cv[0:*]"])
+        cmds.scale(0.1, 0.1, 0.05, r=True, ocp = False)
+        if side == "_l":
+            cmds.move(0, -3.757386, 0, r=True, os=True)
+        else:
+            cmds.move(0, -3.757386, 0, r=True, os=True)
+
+        cmds.select(["eyelid_bott"+side+"_anim.cv[0:*]"])
+        cmds.scale(0.1, 0.1, 0.01, r=True, ocp = False)
+        if side == "_l":
+            cmds.move(0, -3.757386, 0.5, r=True, os=True)
+        else:
+            cmds.move(0, 3.757386, 0.5, r=True, os=True)
+
+        cmds.select(["eyelid_top"+side+"_anim.cv[0:*]"])
+        cmds.scale(0.1, 0.1, 0.01, r=True, ocp = False)
+        if side == "_l":
+            cmds.move(0, -3.757386, 0.5, r=True, os=True)
+        else:
+            cmds.move(0, 3.757386, 0.5, r=True, os=True)
+
+    # jaw
+    cmds.select(["jaw_anim.cv[0:*]"])
+    cmds.scale(0.25, 0.25, 0.25, r = True, ocp = True)
+    cmds.move(0, -24, -16, r = True)
+
+    # ears
+    cmds.select(["ear_l_anim.cv[0:*]"])
+    cmds.scale(0.25, 0.25, 0.25, r = True, ocp = False)
+    cmds.move(6, 8, 12, r = True)
+
+    cmds.select(["ear_r_anim.cv[0:*]"])
+    cmds.scale(0.25, 0.25, 0.25, r = True, ocp = False)
+    cmds.move(-6, 8, 12, r = True)
+
+    # foot
+    cmds.select(["ik_foot_anim_l.cv[0:*]"])
+    cmds.scale(1.4, 1.4, 1.4, r = True, ocp = False)
+    cmds.move(0, 0, -1.129323, r = True)
+
+    cmds.select(["ik_foot_anim_r.cv[0:*]"])
+    cmds.scale(1.4, 1.4, 1.4, r = True, ocp = False)
+    cmds.move(0, 0, -1.129323, r = True)
+
+    cmds.select(["head_fk_anim.cv[0:*]"])
+    cmds.scale(1.9, 1.9, 1.9, r = True, ocp = False)
+
+    cmds.select(None)
+
+scaleControls()
+
+# ------------------------------------------------------------------
+def dynamicsOff():
+    for one in ["tongue_dyn_anim", "ponytail_dyn_anim"]:
+        cmds.setAttr(one+".chainStartEnvelope", 0)
+        cmds.setAttr(one+".chainStartEnvelope", l=True)
+dynamicsOff()
+
+# ------------------------------------------------------------------
+def invertEyelids():
+    sides = ["_l", "_r"]
+    for side in sides:
+        cmds.delete("driver_eyelid_bott"+side+"_parentConstraint1")
+        cmds.rotate(0, -180, 0, "eyelid_bott"+side+"_anim_grp", r=True, os=True)
+        cmds.parentConstraint("eyelid_bott"+side+"_anim", "driver_eyelid_bott"+side, mo=True)
+invertEyelids()
+
+def invertRightEyeball():
+    cmds.delete("driver_eyeball_r_parentConstraint1")
+    cmds.rotate(0, -180, -180, "eyeball_r_anim_grp", r=True, os=True)
+    cmds.parentConstraint("eyeball_r_anim", "driver_eyeball_r", mo=True)
+invertRightEyeball()
+
+# ------------------------------------------------------------------
+def createPistons():
+    import Tools.ART_CreateHydraulicRigs as hyd
+    reload(hyd)
+
+    # pistons
+    cmds.select("arm_piston_A_top_anim", "arm_piston_A_bott_anim", "driver_index_01_l")
+    hyd.CreateHydraulicRig(-1)
+    cmds.select("arm_piston_B_top_anim", "arm_piston_B_bott_anim", "driver_index_01_l")
+    hyd.CreateHydraulicRig(-1)
+
+    cmds.setAttr("Rig_Settings.lForearmTwistAmount", 1)
+
+createPistons()
+
+# ------------------------------------------------------------------
+def fixToes():
+    for one in ["driver_toe_mid_a_l_parentConstraint1", "driver_toe_mid_a_r_parentConstraint1", "driver_toe_pink_a_l_parentConstraint1", "driver_toe_pink_a_r_parentConstraint1", "driver_toe_big_a_l_parentConstraint1", "driver_toe_big_a_r_parentConstraint1"]:
+        cmds.setAttr(one+".interpType", 2)
+fixToes()
+
+# ------------------------------------------------------------------
+def stretchArmRig():
+
+    # Turn on the stretch attr.
+    cmds.setAttr("ik_wrist_l_anim.stretch", 1)
+
+    # Break the connections for the upperam so it doesn stretch.  Also for the scaleblend on the lower arm.
+    mel.eval("source generateChannelMenu.mel;CBdeleteConnection \"ik_upperarm_l.sx\"")
+    mel.eval("source generateChannelMenu.mel;CBdeleteConnection \"ik_upperarm_l.sy\"")
+    mel.eval("source generateChannelMenu.mel;CBdeleteConnection \"ik_upperarm_l.sz\"")
+
+    mel.eval("source generateChannelMenu.mel;CBdeleteConnection \"lo_arm_scale_blend_l.b\"")
+    cmds.setAttr("lo_arm_scale_blend_l.blender", 0)
+
+    # Create a plus minus node, connect the current arm distance into it so I can find the difference between the current arm and the original length.
+    plusMinusCurrent = cmds.shadingNode("plusMinusAverage", asUtility=True, name="currentLength_plusMinusNode")
+    cmds.connectAttr("ik_arm_distBetween_l.distance", plusMinusCurrent+".input1D[0]", f=True)
+    cmds.connectAttr("ik_arm_distBetween_l.distance", plusMinusCurrent+".input1D[1]", f=True)
+    mel.eval("CBdeleteConnection \""+plusMinusCurrent+".input1D[1]\"")
+    cmds.setAttr(plusMinusCurrent+".operation", 2)
+
+    # Hook the previous node into another node that adds the difference to the original distance.
+    plusMinusAdd = cmds.shadingNode("plusMinusAverage", asUtility=True, name="addCurrent_plusMinusNode")
+    cmds.connectAttr("ik_arm_distBetween_l.distance", plusMinusAdd+".input1D[0]", f=True)
+    cmds.connectAttr(plusMinusCurrent+".output1D", plusMinusAdd+".input1D[1]", f=True)
+
+    # Re-connect these values into the arm to replace the old connections.
+    cmds.connectAttr(plusMinusAdd+".output1D", "ik_arm_stretch_condition_l.colorIfFalseR", f=True)
+    cmds.connectAttr(plusMinusAdd+".output1D", "ik_arm_stretch_condition_l.secondTerm", f=True)
+
+stretchArmRig()
\ No newline at end of file
