Branching

diff --git a/Engine/Extras/Maya_AnimationRiggingTools/ARTv1/General/ART/Projects/Orion/Scripts/Pyro_Post_Script.py b/Engine/Extras/Maya_AnimationRiggingTools/ARTv1/General/ART/Projects/Orion/Scripts/Pyro_Post_Script.py
new file mode 100644
--- /dev/null
+++ b/Engine/Extras/Maya_AnimationRiggingTools/ARTv1/General/ART/Projects/Orion/Scripts/Pyro_Post_Script.py
@@ -0,0 +1,512 @@
+import maya.mel as mel
+import maya.cmds as cmds
+import Tools.ART_CreateHydraulicRigs as hyd
+reload(hyd)
+import Tools.ART_addSpaceSwitching as space
+reload(space)
+
+# ------------------------------------------------------------------
+# add ik bones
+def addIKBones():
+    try:
+        # create the joints
+        cmds.select(clear = True)
+        ikFootRoot = cmds.joint(name = "ik_foot_root")
+        cmds.select(clear = True)
+        
+        cmds.select(clear = True)
+        ikFootLeft = cmds.joint(name = "ik_foot_l")
+        cmds.select(clear = True)
+        
+        cmds.select(clear = True)
+        ikFootRight = cmds.joint(name = "ik_foot_r")
+        cmds.select(clear = True)
+        
+        cmds.select(clear = True)
+        ikHandRoot = cmds.joint(name = "ik_hand_root")
+        cmds.select(clear = True)
+        
+        cmds.select(clear = True)
+        ikHandGun = cmds.joint(name = "ik_hand_gun")
+        cmds.select(clear = True)
+        
+        cmds.select(clear = True)
+        ikHandLeft = cmds.joint(name = "ik_hand_l")
+        cmds.select(clear = True)
+        
+        cmds.select(clear = True)
+        ikHandRight = cmds.joint(name = "ik_hand_r")
+        cmds.select(clear = True)
+        
+        # create hierarhcy
+        cmds.parent(ikFootRoot, "root")
+        cmds.parent(ikHandRoot, "root")
+        cmds.parent(ikFootLeft, ikFootRoot)
+        cmds.parent(ikFootRight, ikFootRoot)
+        cmds.parent(ikHandGun, ikHandRoot)
+        cmds.parent(ikHandLeft, ikHandGun)
+        cmds.parent(ikHandRight, ikHandGun)
+        
+        # constrain the joints
+        leftHandConstraint = cmds.parentConstraint("hand_l", ikHandLeft)[0]
+        rightHandGunConstraint = cmds.parentConstraint("hand_r", ikHandGun)[0]
+        rightHandConstraint = cmds.parentConstraint("hand_r", ikHandRight)[0]
+        leftFootConstraint = cmds.parentConstraint("foot_l", ikFootLeft)[0]
+        rightFootConstraint = cmds.parentConstraint("foot_r", ikFootRight)[0]
+
+    except:
+        cmds.warning("Something went wrong")
+		
+# ------------------------------------------------------------------
+def scaleControls():
+    objects = ["dispencer_base_anim", "dispencer_clip_bl_anim", "dispencer_clip_br_anim", "dispencer_clip_fl_anim", "dispencer_clip_fr_anim", "dispencer_door_bott_bk_anim", "dispencer_door_bott_fr_anim", "dispencer_door_mid_bk_anim", "dispencer_door_mid_fr_anim", "dispencer_door_top_bk_anim",  "dispencer_door_top_fr_anim", "dispencer_top_anim", "dispencer_handle_anim" , "dispencer_plug_anim" , "toe_big_a_l_anim", "toe_mid_a_l_anim", "toe_pink_a_l_anim", "toe_big_a_r_anim", "toe_mid_a_r_anim", "toe_pink_a_r_anim", "heel_l_anim", "heel_r_anim", "back_pack_anim", "bottle_01_anim", "bottle_02_anim", "bottle_sack_anim", "bottle_throw_anim", "throttle_pivot_anim", "muscle_thigh_l_anim", "muscle_thigh_r_anim", "muscle_belly_anim", "muscle_tail_bott_anim", "muscle_jowles_anim", "gums_top_anim", "gums_bottom_anim"]
+    for one in objects:
+        cmds.select([one+".cv[0:*]"])
+        cmds.scale(0.2, 0.2, 0.2, r=True, os=True)
+
+    for side in ["_l", "_r"]:
+        cmds.select(["eyeball"+side+"_anim.cv[0:*]"])
+        cmds.scale(0.1, 0.1, 0.05, r=True, ocp = False)
+        if side == "_l":
+            cmds.move(0, 0, 3.757386, r=True, os=True)
+        else:
+            cmds.move(0, 0, 3.757386, r=True, os=True)
+
+        '''cmds.select(["eyelid_bott"+side+"_anim.cv[0:*]"])
+        cmds.scale(0.1, 0.1, 0.01, r=True, ocp = False)
+        if side == "_l":
+            cmds.move(0, -3.757386, 0.5, r=True, os=True)
+        else:
+            cmds.move(0, 3.757386, 0.5, r=True, os=True)
+
+        cmds.select(["eyelid_top"+side+"_anim.cv[0:*]"])
+        cmds.scale(0.1, 0.1, 0.01, r=True, ocp = False)
+        if side == "_l":
+            cmds.move(0, -3.757386, 0.5, r=True, os=True)
+        else:
+            cmds.move(0, 3.757386, 0.5, r=True, os=True)'''
+
+    #ik feet
+    cmds.select("ik_foot_anim_l.cv[0:*]")
+    cmds.scale(2.8, 1.88, 1, p=[17.904671, -8.630238, 2.758035], r=True)
+    cmds.move(0, 0, -2.758035, r = True, os=True, wd=True)
+    
+    cmds.select("ik_foot_anim_r.cv[0:*]")
+    cmds.scale(2.8, 1.88, 1, p=[-17.904671, -8.630238, 2.758035], r=True)
+    cmds.move(0, 0, -2.758035, r = True, os=True, wd=True)
+
+    # Other
+    cmds.select(["jaw_anim.cv[0:*]"])
+    cmds.scale(0.2, 0.2, 0.2, r=True, ocp = True)
+    cmds.move(0, -24, 80, r=True, os=True)
+
+    cmds.select(["body_anim.cv[0:*]"])
+    cmds.scale(3, 3, 3, r=True, os=True)
+    cmds.rotate(0, 0, -90, r=True, os=True)
+    
+    cmds.select(["master_anim.cv[0:*]"])
+    cmds.scale(2, 2, 2, r=True, os=True)
+
+    cmds.select(["offset_anim.cv[0:*]"])
+    cmds.scale(2, 2, 2, r=True, os=True)
+
+    bodyObjs90 = ["hip_anim", "mid_ik_anim", "chest_ik_anim", "head_fk_anim"]
+    for one in bodyObjs90:
+        cmds.select([one+".cv[0:*]"])
+        cmds.scale(3, 3, 3, r = True)
+        cmds.rotate(0, 0, -90, r=True, os=True)
+
+    bodyObjs = ["fk_tail_01_anim", "fk_tail_02_anim", "fk_tail_03_anim", "fk_tail_04_anim",  "spine_01_anim", "spine_02_anim", "spine_03_anim", "fk_thigh_l_anim" , "fk_thigh_r_anim", "fk_foot_l_anim", "fk_foot_r_anim", "neck_01_fk_anim"]
+    for one in bodyObjs:
+        cmds.select([one+".cv[0:*]"])
+        cmds.scale(4, 4, 4, r = True)
+
+    for anim in ["fk_tail_01_anim", "fk_tail_02_anim", "fk_tail_03_anim", "fk_tail_04_anim"]:
+        cmds.select(anim + ".cv[*]")
+        cmds.scale(1.7,1.7,1.7, r = True)
+
+    cmds.select("bottle_throw_anim.cv[*]")
+    cmds.move(0, 0, 0, r = True, os = True, ws = True)
+    cmds.scale(1.8, 1.8, 1.8, r = True)
+
+    cmds.select("spine_03_anim.cv[*]")
+    cmds.scale(1.5, 1.5, 1.5, r = True)
+
+    cmds.select("head_fk_anim.cv[*]")
+    cmds.scale(2.2, 2.2, 2.2, r = True)
+    cmds.rotate(0,0,33, r = True)
+
+    cmds.select("throttle_pivot_anim.cv[*]")
+    cmds.scale(9,9,0, r = True)
+    cmds.rotate(-20,0,0, r = True)
+
+    cmds.setAttr("thigh_twist_l_Tracker.v", 0)
+    cmds.setAttr("thigh_twist_r_Tracker.v", 0)
+
+scaleControls()
+
+# ------------------------------------------------------------------
+# Create the dog legs for the Dino
+def dogLegRig():
+    mel.eval("ikSpringSolver")
+    for side in ["_l", "_r"]:
+        ###########################
+        # DELETE old IK handle
+        cmds.delete("foot_ikHandle"+side)
+        
+        # Disconnect the attrs driving the foot joint so that we can re-drive it with a parent constraint.
+        attrToDisconnect = "foot"+side+".translate"
+        connection = cmds.connectionInfo(attrToDisconnect, sfd=True)
+        cmds.disconnectAttr(connection, attrToDisconnect)
+        cmds.delete("foot"+side+"_orientConstraint1")
+        
+        # PARENT driver_foot_l to driver_heel_l
+        cmds.parent("driver_foot"+side, "driver_heel"+side)
+        
+        # Parent constrain the deform foot to follow the driver foot
+        cmds.parentConstraint("driver_foot"+side, "foot"+side, mo=True)
+    
+        # DUPLICATE driver_heel_l and rename to ik_leg_heel_l
+        ik_leg_heel = cmds.duplicate("driver_heel"+side, n="ik_leg_heel"+side, po=True)[0]
+        # PARENT ik_leg_heel_l to ik_leg_calf_l
+        cmds.parent(ik_leg_heel, "ik_leg_calf"+side)
+        # PARENT ik_leg_foot_l to ik_leg_heel_l(dupe)
+        cmds.parent("ik_leg_foot"+side, ik_leg_heel)
+        cmds.parent("ik_leg_thigh"+side, w=True)
+        cmds.makeIdentity("ik_leg_thigh"+side, a=True, r=True)
+        cmds.parentConstraint("leg_joints_grp"+side, "ik_leg_thigh"+side, mo=True)
+        cmds.parent("ik_leg_thigh"+side, "rig_grp")
+        
+        # DUPLICATE driver_heel_l and rename to fk_leg_heel_l
+        fk_leg_heel = cmds.duplicate("driver_heel"+side, n="fk_leg_heel"+side, po=True)[0]
+        # PARENT fk_leg_heel_l to fk_leg_calf_l
+        cmds.parent(fk_leg_heel, "fk_leg_calf"+side)
+        # PARENT fk_leg_foot_l to fk_leg_heel_l(dupe)
+        cmds.parent("fk_leg_foot"+side, fk_leg_heel)
+        # ORIENT CONSTRAIN fk_leg_foot_l to heel_l_anim
+        cmds.orientConstraint("heel"+side+"_anim", "fk_leg_heel"+side, mo=True)
+        cmds.parentConstraint("heel"+side+"_anim", "fk_foot"+side+"_anim_grp", mo=True)
+
+        #DUPLICATE driver_heel_l and rename to result_leg_heel_l
+        result_leg_heel = cmds.duplicate("driver_heel"+side, n="result_leg_heel"+side, po=True)[0]
+        # PARENT result_leg_heel_l to result_leg_calf_l
+        cmds.parent(result_leg_heel, "result_leg_calf"+side)
+        # PARENT result_leg_foot_l to result_leg_heel_l(dupe)
+        cmds.parent("result_leg_foot"+side, result_leg_heel)
+        
+        #LOAD ikSpringSolver;
+        # CREATE spring IK from ik_leg_thigh_l to ik_leg_foot_l
+        foot_ikhandle = cmds.ikHandle(n="foot_ikHandle"+side, sol="ikSpringSolver", sj="ik_leg_thigh"+side, ee="ik_leg_foot"+side)[0]
+        cmds.setAttr(foot_ikhandle+".v", 0)
+        
+        #PARENT ik handle to ik_foot_ball_pivot_l
+        cmds.parent(foot_ikhandle, "ik_foot_ball_pivot"+side)
+
+        # CREATE PV consraint on ik handle to noflip_pv_loc_l
+        PVConst = cmds.poleVectorConstraint("noflip_pv_loc"+side, foot_ikhandle)[0]
+        if side == "_l":
+            cmds.setAttr(PVConst+".offsetX", 200)
+        else:
+            cmds.setAttr(PVConst+".offsetX", -200)
+
+        # CONNECT the knee twist attr on ik_foot_anim_l to the twist attr on new ik handle
+        cmds.connectAttr("ik_foot_anim"+side+".knee_twist", foot_ikhandle+".twist", f=True)
+        # DELETE driver_heel_l_parentConstraint1
+        cmds.delete("driver_heel"+side+"_parentConstraint1")
+        cmds.parentConstraint(result_leg_heel, "driver_heel"+side, mo=True)
+        
+        # CREATE a parent constraint using ik_leg_heel_l, fk_leg_heel_l and drive result_leg_heel_l
+        resultConstraint = cmds.parentConstraint(fk_leg_heel, ik_leg_heel, result_leg_heel, mo=True)[0]
+        # CREATE a reverse node and hook up a IK/FK switch between the Rig_Settings.lLegMode and the 2 weights on this new constraint
+        shortside = side.replace("_", "")
+        reverseNode = cmds.shadingNode("reverse", asUtility=True, name=result_leg_heel+"_reverse")
+        cmds.connectAttr("Rig_Settings."+shortside+"LegMode", resultConstraint+"."+ik_leg_heel+"W1", f=True)
+        cmds.connectAttr("Rig_Settings."+shortside+"LegMode", reverseNode+".inputX", f=True)
+        cmds.connectAttr(reverseNode+".outputX", resultConstraint+"."+fk_leg_heel+"W0", f=True)
+
+        # Create a KneeFlex Attribute
+        cmds.addAttr("ik_foot_anim"+side, ln="KneeFlex", at="double", min=0, max=1, dv=0.5, k=True)
+        # Connect KneeFlex to :
+        reverseNodeBias = cmds.shadingNode("reverse", asUtility=True, name=result_leg_heel+"_reverse")
+        cmds.connectAttr("ik_foot_anim"+side+".KneeFlex", foot_ikhandle+".springAngleBias[0].springAngleBias_FloatValue", f=True)
+        cmds.connectAttr("ik_foot_anim"+side+".KneeFlex", reverseNodeBias+".inputX", f=True)
+        cmds.connectAttr(reverseNodeBias+".outputX", foot_ikhandle+".springAngleBias[1].springAngleBias_FloatValue", f=True)
+        
+        visReverseNode = cmds.shadingNode("reverse", asUtility=True, name="heel"+side+"_anim_reverse")
+        cmds.connectAttr("Rig_Settings."+shortside+"LegMode", visReverseNode+".inputX", f=True)
+        cmds.connectAttr(visReverseNode+".outputX", "heel"+side+"_parent_grp.v", f=True)
+
+dogLegRig()
+
+# ------------------------------------------------------------------
+def dispencerRig():
+    # Next 2 lines are just for finding the location to use in the 3rd line.
+    #cmds.xform("dispencer_base_anim", q=True, t=True)
+    #cmds.xform("dispencer_base_anim", q=True, ro=True)
+    cmds.xform("dispencer_base_anim_grp", t=[27.842794846757382, 83.38044772340979, -30.319686119574822], ro=[19.709222794231103, -0.026517130554039203, 0.18267468754167737])
+dispencerRig()
+
+# ------------------------------------------------------------------
+def invertEyelids():
+    sides = ["_l", "_r"]
+    for side in sides:
+        cmds.delete("driver_eyelid_bott"+side+"_parentConstraint1")
+        cmds.rotate(0, -180, 0, "eyelid_bott"+side+"_anim_grp", r=True, os=True)
+        cmds.parentConstraint("eyelid_bott"+side+"_anim", "driver_eyelid_bott"+side, mo=True)
+#invertEyelids()
+
+def invertRightEyeball():
+    cmds.delete("driver_eyeball_r_parentConstraint1")
+    cmds.rotate(0, -180, -180, "eyeball_r_anim_grp", r=True, os=True)
+    cmds.parentConstraint("eyeball_r_anim", "driver_eyeball_r", mo=True)
+invertRightEyeball()
+
+# ------------------------------------------------------------------
+def fixToes():
+    for one in ["driver_toe_mid_a_l_parentConstraint1", "driver_toe_mid_a_r_parentConstraint1", "driver_toe_pink_a_l_parentConstraint1", "driver_toe_pink_a_r_parentConstraint1", "driver_toe_big_a_l_parentConstraint1", "driver_toe_big_a_r_parentConstraint1"]:
+        cmds.setAttr(one+".interpType", 2)
+fixToes()
+
+# ------------------------------------------------------------------
+def dynamicsOff():
+    for one in ["tail_dyn_anim", "tail_armor_dyn_anim", "tongue_dyn_anim"]:
+        cmds.setAttr(one+".chainStartEnvelope", 0)
+        cmds.setAttr(one+".chainStartEnvelope", l=True)
+dynamicsOff()
+
+# ------------------------------------------------------------------
+def createGameCam():
+    gameCam = cmds.camera(centerOfInterest=5, focalLength=18, lensSqueezeRatio=1, cameraScale=1, horizontalFilmAperture=1.41732, horizontalFilmOffset=0, verticalFilmAperture=0.94488, verticalFilmOffset=0, filmFit="Overscan", overscan=1.05, motionBlur=0, shutterAngle=144, nearClipPlane=1.0, farClipPlane=10000, orthographic=0, orthographicWidth=30, panZoomEnabled=0, horizontalPan=0, verticalPan=0, zoom=1, displayFilmGate=True)[0]
+    print gameCam
+    gameCam = cmds.rename(gameCam, "GameCam")
+    cmds.camera(gameCam, e=True, horizontalFilmAperture=1.6724376, focalLength=21.2399857011)
+    cmds.setAttr(gameCam+".displayGateMaskOpacity", 1)
+    cmds.setAttr(gameCam+".displayGateMaskColor", 0, 0, 0, type="double3")
+    cmds.setAttr(gameCam+".r", 90, 0, 180, type="double3")
+    
+    # Create the control
+    control = cmds.circle(name=gameCam+"_Ctrl", ch=True, o=True, r=15)[0]
+    cmds.makeIdentity(control, a=True, t=True, r=True, s=True)
+    # Create groups for the sdk and the constraints
+    const_grp = cmds.group(control, name=control+"_grp")
+    if cmds.objExists("rig_grp"):
+        cmds.parent(const_grp, "rig_grp")
+    # Parent Constrain the cam to the control
+    cmds.parentConstraint(control, gameCam, mo=True)
+    
+    cmds.connectAttr(control+".visibility", gameCam+".visibility", f=True)
+    cmds.setAttr(const_grp+".v", 0)
+    
+    cmds.setAttr(const_grp+".t", 0, 500, 200, type="double3")
+createGameCam()
+
+# ------------------------------------------------------------------
+def addCapsule(radius, halfHeight):
+    trueHalfHeight = (halfHeight*2)-(radius*2)
+    capsule = cmds.polyCylinder(name="collision_capsule", r=radius, h=trueHalfHeight, sx=18, sy=1, sz=5, ax=[0,0,1], rcp=True, cuv=3, ch=True)
+    cmds.setAttr(capsule[0]+".tz", halfHeight+2.2)
+    cmds.setAttr(capsule[0]+".rz", 90)
+
+    cmds.addAttr(capsule[0], ln="Radius", dv=radius, keyable=True)
+    cmds.addAttr(capsule[0], ln="HalfHeight", dv=halfHeight, keyable=True)
+
+    expr = capsule[1]+".radius = "+capsule[0]+".Radius;\n"+capsule[1]+".height = ("+capsule[0]+".HalfHeight*2) - ("+capsule[0]+".Radius*2);"
+    cmds.expression(s=expr, o=capsule[1], ae=1, uc="all")
+
+    cmds.setAttr(capsule[0]+".overrideEnabled", 1)
+    cmds.setAttr(capsule[0]+".overrideColor", 18)
+    cmds.setAttr(capsule[0]+".overrideShading", 0)
+
+    cmds.addAttr(capsule[0], ln="Shaded", dv=0, min=0, max=1, keyable=True)
+    cmds.connectAttr(capsule[0]+".Shaded", capsule[0]+".overrideShading")
+
+    cmds.parentConstraint("root", capsule[0], mo=True)
+    cmds.select(capsule[0], r=True)
+    capsuleDisplayLayer = cmds.createDisplayLayer(nr=True, name="Collision Capsule")
+    cmds.setAttr(capsuleDisplayLayer+".displayType", 2)
+    cmds.setAttr(capsuleDisplayLayer+".v", 0)
+    
+    for attr in ["tx", "ty", "tz", "rx", "ry", "rz", "sx", "sy", "sz"]:
+        cmds.setAttr(capsule[0]+"."+attr, lock=True, keyable=False, channelBox=False)
+
+    cmds.select(clear=True)
+addCapsule(72, 108)
+    
+# ------------------------------------------------------------------
+#The destination attribute 'Goblin:ik_foot_anim_l_space_switcher_follow_parentConstraint1.W2' cannot be found.
+# Import in and hook up the Goblin for Pyro
+def hookUpGoblin():
+    import Tools.ART_addSpaceSwitching as space
+    reload(space)
+    
+    # import in Goblin
+    cmds.file("D:/Build/usr/jeremy_ernst/MayaTools/General/ART/Projects/Orion/AnimRigs/Pyro_Goblin.mb", i=True, options='v=0', type="mayaBinary", namespace="Goblin", mergeNamespacesOnClash=False)
+
+    cmds.setAttr("Goblin:Rig_Settings.spine_fk", 1)
+    cmds.setAttr("Goblin:Rig_Settings.spine_ik", 0)
+
+    # set elbows to body space
+    cmds.setAttr("Goblin:ik_elbow_l_anim_space_switcher.space_body_anim", 1)
+    cmds.setAttr("Goblin:ik_elbow_r_anim_space_switcher.space_body_anim", 1)
+
+    cmds.setAttr("Goblin:ik_wrist_l_anim_space_switcher.space_body_anim", 1)
+    cmds.setAttr("Goblin:ik_wrist_r_anim_space_switcher.space_body_anim", 1)
+
+    cmds.setAttr("Goblin:ik_foot_anim_l_space_switcher.space_body_anim", 1)
+    cmds.setAttr("Goblin:ik_foot_anim_r_space_switcher.space_body_anim", 1)
+
+    # turn on auto clavicles
+    #cmds.setAttr("Goblin:clavicle_l_anim.autoShoulders", 1)
+    #cmds.setAttr("Goblin:clavicle_r_anim.autoShoulders", 1)
+
+    # set head to world space
+    cmds.setAttr("Goblin:head_fk_anim.fkOrientation", 2)
+
+    # position Goblin
+    cmds.setAttr("Goblin:body_anim.tx", 75.168)
+
+    cmds.setAttr("Goblin:ik_foot_anim_l.tx", 17.359)
+    cmds.setAttr("Goblin:ik_foot_anim_l.ty", -16.157)
+    cmds.setAttr("Goblin:ik_foot_anim_l.tz", 10.935)
+    cmds.setAttr("Goblin:ik_foot_anim_l.rx", -52.374)
+    cmds.setAttr("Goblin:ik_foot_anim_l.ry", 6.81)
+    cmds.setAttr("Goblin:ik_foot_anim_l.rz", 45.2873)
+    cmds.setAttr("Goblin:ik_foot_anim_l.knee_twist", 47.2)
+
+    cmds.setAttr("Goblin:ik_foot_anim_r.tx", -17.359)
+    cmds.setAttr("Goblin:ik_foot_anim_r.ty", -16.157)
+    cmds.setAttr("Goblin:ik_foot_anim_r.tz", 10.935)
+    cmds.setAttr("Goblin:ik_foot_anim_r.rx", -52.374)
+    cmds.setAttr("Goblin:ik_foot_anim_r.ry", -6.81)
+    cmds.setAttr("Goblin:ik_foot_anim_r.rz", -45.2873)
+    cmds.setAttr("Goblin:ik_foot_anim_r.knee_twist", 47.2)
+
+    # connect Goblin to the seat
+    space.addSpaceSwitching("Goblin:body_anim", "pelvis", True)
+    space.addSpaceSwitching("Goblin:body_anim", "spine_01", False)
+
+    #cmds.setAttr("hip_anim.autoHips", 1)
+
+    # expose knee locs for Goblin and make knees follow body
+    cmds.select(["Goblin:noflip_aim_grp_l", "Goblin:noflip_aim_grp_r"])
+    cmds.delete(constraints = True)
+
+    cmds.parentConstraint("body_anim", "Goblin:noflip_aim_grp_l", mo = True)
+    cmds.parentConstraint("body_anim", "Goblin:noflip_aim_grp_r", mo = True)
+
+    # Space switching for the hands to the throttle
+    space.addSpaceSwitching("Goblin:ik_wrist_l_anim", "throttle_pivot_anim", False)
+    space.addSpaceSwitching("Goblin:ik_wrist_r_anim", "throttle_pivot_anim", False)
+
+    cmds.setAttr("ik_wrist_l_anim_space_switcher.space_body_anim", 1)
+    cmds.setAttr("ik_wrist_r_anim_space_switcher.space_body_anim", 1)
+
+    # join skeletons
+    cmds.select("Goblin:root", hi = True)
+    joints = cmds.ls(sl = True)
+    for joint in joints:
+        newName = joint.partition(":")[2]
+        cmds.rename(joint, "Goblin_" + newName)
+    cmds.parent("Goblin_root", "root")    
+
+    # Disconnect the attrs driving the Goblin pelvis joint so that we can re-drive it with a parent constraint.
+    attrToDisconnect = "Goblin_pelvis.translate"
+    connection = cmds.connectionInfo(attrToDisconnect, sfd=True)
+    cmds.disconnectAttr(connection, attrToDisconnect)
+    cmds.delete("Goblin_pelvis_orientConstraint1")
+
+    # Parent constrain the deform pelvis to follow the driver pelvis
+    cmds.parentConstraint("Goblin:driver_pelvis", "Goblin_pelvis", mo=True)
+
+    cmds.select(None)
+hookUpGoblin()
+
+# ------------------------------------------------------------------
+# This needs to be added after the Goblin has been brought in and hooked up.
+def bottlesRig():
+    import Tools.ART_addSpaceSwitching as space
+    reload(space)
+
+    space.addSpaceSwitching("bottle_throw_anim", "Goblin_hand_r", False)
+    space.addSpaceSwitching("bottle_01_anim", "Goblin_hand_r", False)
+    space.addSpaceSwitching("bottle_02_anim", "Goblin_hand_r", False)
+bottlesRig()
+
+
+# ------------------------------------------------------------------
+# save the AnimRig file
+cmds.file(save = True, type = "mayaBinary", force = True)
+
+# add IK Bones here so that they are included in the export file.  But dont save becasue we dont want them in the rig file.  Thats why we saved in the line before.
+addIKBones()
+
+# ------------------------------------------------------------------
+# Creeate the export file that goes into source so it can be exported to the game
+def createExportFile():
+    # Import the export file from the reference so that it can be edited.
+    cmds.file("D:/Build/usr/jeremy_ernst/MayaTools/General/ART/Projects/Orion/ExportFiles/Pyro_Export.mb", importReference=True)
+
+    # select the entire skeleton and make a list of it.
+    cmds.select("root", hi=True)
+    selection = cmds.ls(sl=True)
+    cmds.showHidden(a=True)
+
+    # for each object in the selection, disconnect translation and scale.
+    for each in selection:
+        attrs = ["translate", "scale"]
+        cmds.lockNode(each, lock=False)
+
+        for attr in attrs:
+            try:
+                conn = cmds.listConnections(each + "." + attr, s=True, p=True)[0]
+                cmds.disconnectAttr(conn, each + "." + attr)
+            except:
+                pass
+
+    # find all of the constraints in the skeleton and delete them
+    cmds.select("root", hi=True)
+    constraints = cmds.ls(sl=True, type="constraint")
+    cmds.select(constraints)
+    cmds.delete(constraints)
+
+    # delete the control rig so that it severs any connections between it and the geo or skeleton.
+    cmds.select("Goblin:rig_grp", "rig_grp", hi=True)
+    rigs = cmds.ls(sl=True)
+    for one in rigs:
+        cmds.lockNode(one, lock=False)
+    cmds.delete("Goblin:rig_grp", "rig_grp")
+
+    # Remove the Goblin namespace and set the bindposes so that Creating LODs will work later on.
+    cmds.namespace(rm="Goblin", mnp=True)
+    cmds.dagPose("root", bp=True, r=True) 
+    cmds.dagPose("Goblin_root", bp=True, r=True)
+    
+    # select the skeleton and all of the geo and export it to the export file.  Remember to go export an fbx from this.
+    cmds.select("Goblin:Pyro_Goblin_Geo", "Pyro_Dino_Geo", "root", hi=True)
+    cmds.file("D:/Build/ArtSource/Orion/Characters/Heroes/Pyro/Export/Pyro_Export.mb", force=True, typ="mayaBinary", pr=True, es=True)
+createExportFile()
+
+# ------------------------------------------------------------------
+# ADD IK BONES to the export file
+sceneName = cmds.file(q = True, sceneName = True)
+exportFile = sceneName.partition("AnimRigs")[0] + "ExportFiles" +  sceneName.partition("AnimRigs")[2].partition(".mb")[0] + "_Export.mb"
+cmds.file(exportFile, open = True, force = True)
+
+
+addIKBones()
+
+# ------------------------------------------------------------------
+def addCustomAttrs():
+    attrs = cmds.listAttr("root")
+    exists = False
+    for one in attrs:
+        if one.find("DistanceCurve") != -1:
+            exists = True
+    if exists == False:
+        cmds.addAttr("root", ln="DistanceCurve", keyable=True)
+addCustomAttrs()
+
+# save the export file
+cmds.file(save = True, type = "mayaBinary", force = True)
+
+
diff --git a/Engine/Extras/Maya_AnimationRiggingTools/ARTv1/General/ART/Projects/Orion/Scripts/Rampage_CreateLODs.py b/Engine/Extras/Maya_AnimationRiggingTools/ARTv1/General/ART/Projects/Orion/Scripts/Rampage_CreateLODs.py
new file mode 100644
--- /dev/null
+++ b/Engine/Extras/Maya_AnimationRiggingTools/ARTv1/General/ART/Projects/Orion/Scripts/Rampage_CreateLODs.py
@@ -0,0 +1,682 @@
+import maya.cmds as cmds
+import customMayaMenu as menu
+import Tools.ART_ImportMultipleWeights as impw
+import Tools.ART_TransferWeights as transfer
+import Tools.ART_ExportLODtoFBX as export
+reload(export)
+
+################################################## RAMPAGE #########################################################################
+# Warn the user that this operation will save their file.
+saveResult = cmds.confirmDialog(title = "Save Warning", message="In order to continue your file must be saved.  Would you like to save it?  If yes it will be saved and after the operation is complete your file will be re-opened.", button = ["Yes", "Cancel"], defaultButton='Yes', cancelButton='Cancel', dismissString='Cancel')
+
+# If they are ok with saving the file, warn them that it will also delete all of the current files in their skin weights directory.
+if saveResult == "Yes":
+    deleteSkinFilesResult = cmds.confirmDialog(title = "SkinWeights Warning", message="In order to continue we must delete all of the files in your ..\General\ART\SkinWeights\ directory.  Would you like to continue?.", button = ["Yes", "Cancel"], defaultButton='Yes', cancelButton='Cancel', dismissString='Cancel')
+    if deleteSkinFilesResult == "Yes":
+        toolsPath = cmds.internalVar(usd = True) + "mayaTools.txt"
+        if os.path.exists(toolsPath):
+            f = open(toolsPath, 'r')
+            mayaToolsDir = f.readline()
+            f.close()
+
+        path = mayaToolsDir + "\General\ART\SkinWeights\\"
+        files = os.listdir(path)
+        #delete all weight cache files in the system folder
+        for file in files:
+            os.remove(path + file)
+            #print path + file
+
+
+
+# Rampage LOD1 -------------------------------------------------------------------------------------------------------------
+def RampageLOD1():
+    if saveResult == "Yes":
+
+        cmds.file(save=True)
+        
+        meshes = ["body_geo_LOD1", "chain_link_geo_01_LOD1", "chain_link_geo_02_LOD1", "chain_link_geo_03_LOD1", "chain_link_geo_04_LOD1", "chain_link_geo_05_LOD1", "chain_link_geo_06_LOD1", "chain_link_geo_07_LOD1", "cloth_geo_LOD1", "eyeball_l_geo_LOD1", "eyeball_r_geo_LOD1", "groin_armor_geo_LOD1", "legs_geo_LOD1", "lower_teeth_geo_LOD1", "shackles_geo_LOD1", "shin_armor_geo_LOD1", "thigh_armor_geo_LOD1", "upper_teeth_geo_LOD1"]
+
+        jointsToDelete = []
+
+        # remove unneeded joints
+
+        
+        # Delete the joints we no longer want.
+        for one in jointsToDelete:
+            if cmds.objExists(one):
+                cmds.delete(one)
+                
+        # Export the fbx file.
+        export.exportLOD("Rampage_LOD1.fbx", "LOD1")
+        
+        # Re-open the file without saving.
+        fullPath = cmds.file(q = True, sceneName = True)
+        cmds.file(fullPath, open=True, f=True)
+
+        print "LOD1 Exported Successfully!!!!!!!!!!!!!!!"
+
+RampageLOD1()
+
+# Rampage LOD2 -------------------------------------------------------------------------------------------------------------
+def RampageLOD2():
+    if saveResult == "Yes":
+        
+        cmds.file(save=True)
+        
+        meshes = ["body_geo_LOD2", "chain_link_geo_01_LOD2", "chain_link_geo_02_LOD2", "chain_link_geo_03_LOD2", "chain_link_geo_04_LOD2", "chain_link_geo_05_LOD2", "chain_link_geo_06_LOD2", "chain_link_geo_07_LOD2", "cloth_geo_LOD2", "eyeball_l_geo_LOD2", "eyeball_r_geo_LOD2", "groin_armor_geo_LOD2", "legs_geo_LOD2", "lower_teeth_geo_LOD2", "shackles_geo_LOD2", "shin_armor_geo_LOD2", "thigh_armor_geo_LOD2", "upper_teeth_geo_LOD2"]
+
+        jointsToDelete = []
+
+         # remove unneeded joints
+        jointsToRemove = ["lat_muscle_l", "lat_muscle_r"]
+        jointToTransferTo = "spine_02"
+        transfer.LOD_transferWeights(meshes, jointsToRemove, jointToTransferTo)
+        jointsToDelete.extend(jointsToRemove)
+
+        jointsToRemove = ["fin_scap_a_l_02"]
+        jointToTransferTo = "fin_scap_a_l_01"
+        transfer.LOD_transferWeights(meshes, jointsToRemove, jointToTransferTo)
+        jointsToDelete.extend(jointsToRemove)
+
+        jointsToRemove = ["fin_scap_a_r_02"]
+        jointToTransferTo = "fin_scap_a_r_01"
+        transfer.LOD_transferWeights(meshes, jointsToRemove, jointToTransferTo)
+        jointsToDelete.extend(jointsToRemove)
+
+        jointsToRemove = ["fin_scap_b_l_02"]
+        jointToTransferTo = "fin_scap_b_l_01"
+        transfer.LOD_transferWeights(meshes, jointsToRemove, jointToTransferTo)
+        jointsToDelete.extend(jointsToRemove)
+
+        jointsToRemove = ["fin_scap_b_r_02"]
+        jointToTransferTo = "fin_scap_b_r_01"
+        transfer.LOD_transferWeights(meshes, jointsToRemove, jointToTransferTo)
+        jointsToDelete.extend(jointsToRemove)
+
+        jointsToRemove = ["fin_scap_c_l_02"]
+        jointToTransferTo = "fin_scap_c_l_01"
+        transfer.LOD_transferWeights(meshes, jointsToRemove, jointToTransferTo)
+        jointsToDelete.extend(jointsToRemove)
+
+        jointsToRemove = ["fin_scap_c_r_02"]
+        jointToTransferTo = "fin_scap_c_r_01"
+        transfer.LOD_transferWeights(meshes, jointsToRemove, jointToTransferTo)
+        jointsToDelete.extend(jointsToRemove)
+
+        jointsToRemove = ["fin_scap_d_l", "fin_scap_d_r", "pec_muscle_l", "pec_muscle_r"]
+        jointToTransferTo = "spine_03"
+        transfer.LOD_transferWeights(meshes, jointsToRemove, jointToTransferTo)
+        jointsToDelete.extend(jointsToRemove)
+
+        jointsToRemove = ["fin_spine_a_02"]
+        jointToTransferTo = "fin_spine_a_01"
+        transfer.LOD_transferWeights(meshes, jointsToRemove, jointToTransferTo)
+        jointsToDelete.extend(jointsToRemove)
+
+        jointsToRemove = ["fin_spine_b_02"]
+        jointToTransferTo = "fin_spine_b_01"
+        transfer.LOD_transferWeights(meshes, jointsToRemove, jointToTransferTo)
+        jointsToDelete.extend(jointsToRemove)
+
+        jointsToRemove = ["fin_spine_c_02"]
+        jointToTransferTo = "fin_spine_c_01"
+        transfer.LOD_transferWeights(meshes, jointsToRemove, jointToTransferTo)
+        jointsToDelete.extend(jointsToRemove)
+
+        jointsToRemove = ["fin_spine_d_02"]
+        jointToTransferTo = "fin_spine_d_01"
+        transfer.LOD_transferWeights(meshes, jointsToRemove, jointToTransferTo)
+        jointsToDelete.extend(jointsToRemove)
+
+        jointsToRemove = ["upperarm_twist_01_l", "upperarm_twist_02_l", "bicep_muscle_l", "tricep_muscle_l"]
+        jointToTransferTo = "upperarm_l"
+        transfer.LOD_transferWeights(meshes, jointsToRemove, jointToTransferTo)
+        jointsToDelete.extend(jointsToRemove)
+
+        jointsToRemove = ["upperarm_twist_01_r", "upperarm_twist_02_r", "bicep_muscle_r", "tricep_muscle_r"]
+        jointToTransferTo = "upperarm_r"
+        transfer.LOD_transferWeights(meshes, jointsToRemove, jointToTransferTo)
+        jointsToDelete.extend(jointsToRemove)
+
+        jointsToRemove = ["fin_arm_c_l_02"]
+        jointToTransferTo = "fin_arm_c_l_01"
+        transfer.LOD_transferWeights(meshes, jointsToRemove, jointToTransferTo)
+        jointsToDelete.extend(jointsToRemove)
+
+        jointsToRemove = ["fin_arm_d_l_02"]
+        jointToTransferTo = "fin_arm_d_l_01"
+        transfer.LOD_transferWeights(meshes, jointsToRemove, jointToTransferTo)
+        jointsToDelete.extend(jointsToRemove)
+
+        jointsToRemove = ["fin_arm_b_l_02"]
+        jointToTransferTo = "fin_arm_b_l_01"
+        transfer.LOD_transferWeights(meshes, jointsToRemove, jointToTransferTo)
+        jointsToDelete.extend(jointsToRemove)
+
+        jointsToRemove = ["chain_06", "chain_07"]
+        jointToTransferTo = "chain_05"
+        transfer.LOD_transferWeights(meshes, jointsToRemove, jointToTransferTo)
+        jointsToDelete.extend(jointsToRemove)
+
+        jointsToRemove = ["fin_arm_a_l"]
+        jointToTransferTo = "lowerarm_twist_01_l"
+        transfer.LOD_transferWeights(meshes, jointsToRemove, jointToTransferTo)
+        jointsToDelete.extend(jointsToRemove)
+
+        jointsToRemove = ["fin_arm_c_r_02"]
+        jointToTransferTo = "fin_arm_c_r_01"
+        transfer.LOD_transferWeights(meshes, jointsToRemove, jointToTransferTo)
+        jointsToDelete.extend(jointsToRemove)
+
+        jointsToRemove = ["fin_arm_d_r_02"]
+        jointToTransferTo = "fin_arm_d_r_01"
+        transfer.LOD_transferWeights(meshes, jointsToRemove, jointToTransferTo)
+        jointsToDelete.extend(jointsToRemove)
+
+        jointsToRemove = ["fin_arm_b_r_02"]
+        jointToTransferTo = "fin_arm_b_r_01"
+        transfer.LOD_transferWeights(meshes, jointsToRemove, jointToTransferTo)
+        jointsToDelete.extend(jointsToRemove)
+
+        jointsToRemove = ["fin_arm_a_r"]
+        jointToTransferTo = "lowerarm_twist_01_r"
+        transfer.LOD_transferWeights(meshes, jointsToRemove, jointToTransferTo)
+        jointsToDelete.extend(jointsToRemove)
+
+        jointsToRemove = ["eyeball_l", "eyeball_r", "eyebrow_in_l", "eyebrow_in_r", "eyebrow_mid", "eyelid_bott_l", "eyelid_bott_r", "eyelid_top_l", "eyelid_top_r"]
+        jointToTransferTo = "head"
+        transfer.LOD_transferWeights(meshes, jointsToRemove, jointToTransferTo)
+        jointsToDelete.extend(jointsToRemove)
+
+        jointsToRemove = ["chin_tent_k_l_02"]
+        jointToTransferTo = "chin_tent_k_l_01"
+        transfer.LOD_transferWeights(meshes, jointsToRemove, jointToTransferTo)
+        jointsToDelete.extend(jointsToRemove)
+
+        jointsToRemove = ["chin_tent_k_r_02"]
+        jointToTransferTo = "chin_tent_k_r_01"
+        transfer.LOD_transferWeights(meshes, jointsToRemove, jointToTransferTo)
+        jointsToDelete.extend(jointsToRemove)
+
+        jointsToRemove = ["chin_tent_a_l_02"]
+        jointToTransferTo = "chin_tent_a_l_01"
+        transfer.LOD_transferWeights(meshes, jointsToRemove, jointToTransferTo)
+        jointsToDelete.extend(jointsToRemove)
+
+        jointsToRemove = ["chin_tent_c_l_02"]
+        jointToTransferTo = "chin_tent_c_l_01"
+        transfer.LOD_transferWeights(meshes, jointsToRemove, jointToTransferTo)
+        jointsToDelete.extend(jointsToRemove)
+
+        jointsToRemove = ["chin_tent_d_l_02"]
+        jointToTransferTo = "chin_tent_d_l_01"
+        transfer.LOD_transferWeights(meshes, jointsToRemove, jointToTransferTo)
+        jointsToDelete.extend(jointsToRemove)
+
+        jointsToRemove = ["chin_tent_e_l_02"]
+        jointToTransferTo = "chin_tent_e_l_01"
+        transfer.LOD_transferWeights(meshes, jointsToRemove, jointToTransferTo)
+        jointsToDelete.extend(jointsToRemove)
+
+        jointsToRemove = ["chin_tent_f_l_02"]
+        jointToTransferTo = "chin_tent_f_l_01"
+        transfer.LOD_transferWeights(meshes, jointsToRemove, jointToTransferTo)
+        jointsToDelete.extend(jointsToRemove)
+
+        jointsToRemove = ["chin_tent_g_l_02"]
+        jointToTransferTo = "chin_tent_g_l_01"
+        transfer.LOD_transferWeights(meshes, jointsToRemove, jointToTransferTo)
+        jointsToDelete.extend(jointsToRemove)
+
+        jointsToRemove = ["chin_tent_h_l_02"]
+        jointToTransferTo = "chin_tent_h_l_01"
+        transfer.LOD_transferWeights(meshes, jointsToRemove, jointToTransferTo)
+        jointsToDelete.extend(jointsToRemove)
+
+        jointsToRemove = ["chin_tent_i_l_02"]
+        jointToTransferTo = "chin_tent_i_l_01"
+        transfer.LOD_transferWeights(meshes, jointsToRemove, jointToTransferTo)
+        jointsToDelete.extend(jointsToRemove)
+
+        jointsToRemove = ["fin_head_a_l_02"]
+        jointToTransferTo = "fin_head_a_l_01"
+        transfer.LOD_transferWeights(meshes, jointsToRemove, jointToTransferTo)
+        jointsToDelete.extend(jointsToRemove)
+
+        jointsToRemove = ["fin_head_b_l_02"]
+        jointToTransferTo = "fin_head_b_l_01"
+        transfer.LOD_transferWeights(meshes, jointsToRemove, jointToTransferTo)
+        jointsToDelete.extend(jointsToRemove)
+
+        jointsToRemove = ["fin_head_c_l_02"]
+        jointToTransferTo = "fin_head_c_l_01"
+        transfer.LOD_transferWeights(meshes, jointsToRemove, jointToTransferTo)
+        jointsToDelete.extend(jointsToRemove)
+
+        jointsToRemove = ["fin_head_a_r_02"]
+        jointToTransferTo = "fin_head_a_r_01"
+        transfer.LOD_transferWeights(meshes, jointsToRemove, jointToTransferTo)
+        jointsToDelete.extend(jointsToRemove)
+
+        jointsToRemove = ["fin_head_b_r_02"]
+        jointToTransferTo = "fin_head_b_r_01"
+        transfer.LOD_transferWeights(meshes, jointsToRemove, jointToTransferTo)
+        jointsToDelete.extend(jointsToRemove)
+
+        jointsToRemove = ["fin_head_c_r_02"]
+        jointToTransferTo = "fin_head_c_r_01"
+        transfer.LOD_transferWeights(meshes, jointsToRemove, jointToTransferTo)
+        jointsToDelete.extend(jointsToRemove)
+
+        jointsToRemove = ["fin_head_mid_02"]
+        jointToTransferTo = "fin_head_mid_01"
+        transfer.LOD_transferWeights(meshes, jointsToRemove, jointToTransferTo)
+        jointsToDelete.extend(jointsToRemove)
+
+        jointsToRemove = ["toe_index_02_l", "toe_index_03_l"]
+        jointToTransferTo = "toe_index_01_l"
+        transfer.LOD_transferWeights(meshes, jointsToRemove, jointToTransferTo)
+        jointsToDelete.extend(jointsToRemove)
+
+        jointsToRemove = ["toe_mid_02_l", "toe_mid_03_l"]
+        jointToTransferTo = "toe_mid_01_l"
+        transfer.LOD_transferWeights(meshes, jointsToRemove, jointToTransferTo)
+        jointsToDelete.extend(jointsToRemove)
+
+        jointsToRemove = ["toe_pinky_02_l", "toe_pinky_03_l"]
+        jointToTransferTo = "toe_pinky_01_l"
+        transfer.LOD_transferWeights(meshes, jointsToRemove, jointToTransferTo)
+        jointsToDelete.extend(jointsToRemove)
+
+        jointsToRemove = ["toe_index_02_r", "toe_index_03_r"]
+        jointToTransferTo = "toe_index_01_r"
+        transfer.LOD_transferWeights(meshes, jointsToRemove, jointToTransferTo)
+        jointsToDelete.extend(jointsToRemove)
+
+        jointsToRemove = ["toe_mid_02_r", "toe_mid_03_r"]
+        jointToTransferTo = "toe_mid_01_r"
+        transfer.LOD_transferWeights(meshes, jointsToRemove, jointToTransferTo)
+        jointsToDelete.extend(jointsToRemove)
+
+        jointsToRemove = ["toe_pinky_02_r", "toe_pinky_03_r"]
+        jointToTransferTo = "toe_pinky_01_r"
+        transfer.LOD_transferWeights(meshes, jointsToRemove, jointToTransferTo)
+        jointsToDelete.extend(jointsToRemove)
+
+        # Delete the joints we no longer want.
+        for one in jointsToDelete:
+            if cmds.objExists(one):
+                cmds.delete(one)
+                
+        # Export the fbx file.
+        export.exportLOD("Rampage_LOD2.fbx", "LOD2")
+
+        # Re-open the file without saving.
+        fullPath = cmds.file(q = True, sceneName = True)
+        cmds.file(fullPath, open=True, f=True)
+
+        print "LOD2 Exported Successfully!!!!!!!!!!!!!!!"
+
+RampageLOD2()
+
+# Rampage LOD3 -------------------------------------------------------------------------------------------------------------
+def RampageLOD3():
+    if saveResult == "Yes":
+
+        cmds.file(save=True)
+    
+        meshes = ["body_geo_LOD3", "chain_link_geo_01_LOD3", "chain_link_geo_02_LOD3", "chain_link_geo_03_LOD3", "chain_link_geo_04_LOD3", "chain_link_geo_05_LOD3", "chain_link_geo_06_LOD3", "chain_link_geo_07_LOD3", "cloth_geo_LOD3", "eyeball_l_geo_LOD3", "eyeball_r_geo_LOD3", "groin_armor_geo_LOD3", "legs_geo_LOD3", "lower_teeth_geo_LOD3", "shackles_geo_LOD3", "shin_armor_geo_LOD3", "thigh_armor_geo_LOD3", "upper_teeth_geo_LOD3"]
+    
+        jointsToDelete = []
+
+        # remove unneeded joints
+        jointsToRemove = ["lat_muscle_l", "lat_muscle_r"]
+        jointToTransferTo = "spine_02"
+        transfer.LOD_transferWeights(meshes, jointsToRemove, jointToTransferTo)
+        jointsToDelete.extend(jointsToRemove)
+
+        jointsToRemove = ["fin_scap_a_l_02"]
+        jointToTransferTo = "fin_scap_a_l_01"
+        transfer.LOD_transferWeights(meshes, jointsToRemove, jointToTransferTo)
+        jointsToDelete.extend(jointsToRemove)
+
+        jointsToRemove = ["fin_scap_a_r_02"]
+        jointToTransferTo = "fin_scap_a_r_01"
+        transfer.LOD_transferWeights(meshes, jointsToRemove, jointToTransferTo)
+        jointsToDelete.extend(jointsToRemove)
+
+        jointsToRemove = ["fin_scap_b_l_02"]
+        jointToTransferTo = "fin_scap_b_l_01"
+        transfer.LOD_transferWeights(meshes, jointsToRemove, jointToTransferTo)
+        jointsToDelete.extend(jointsToRemove)
+
+        jointsToRemove = ["fin_scap_b_r_02"]
+        jointToTransferTo = "fin_scap_b_r_01"
+        transfer.LOD_transferWeights(meshes, jointsToRemove, jointToTransferTo)
+        jointsToDelete.extend(jointsToRemove)
+
+        jointsToRemove = ["fin_scap_c_l_02"]
+        jointToTransferTo = "fin_scap_c_l_01"
+        transfer.LOD_transferWeights(meshes, jointsToRemove, jointToTransferTo)
+        jointsToDelete.extend(jointsToRemove)
+
+        jointsToRemove = ["fin_scap_c_r_02"]
+        jointToTransferTo = "fin_scap_c_r_01"
+        transfer.LOD_transferWeights(meshes, jointsToRemove, jointToTransferTo)
+        jointsToDelete.extend(jointsToRemove)
+
+        jointsToRemove = ["fin_scap_d_l", "fin_scap_d_r", "pec_muscle_l", "pec_muscle_r"]
+        jointToTransferTo = "spine_03"
+        transfer.LOD_transferWeights(meshes, jointsToRemove, jointToTransferTo)
+        jointsToDelete.extend(jointsToRemove)
+
+        jointsToRemove = ["fin_spine_a_02"]
+        jointToTransferTo = "fin_spine_a_01"
+        transfer.LOD_transferWeights(meshes, jointsToRemove, jointToTransferTo)
+        jointsToDelete.extend(jointsToRemove)
+
+        jointsToRemove = ["fin_spine_b_02"]
+        jointToTransferTo = "fin_spine_b_01"
+        transfer.LOD_transferWeights(meshes, jointsToRemove, jointToTransferTo)
+        jointsToDelete.extend(jointsToRemove)
+
+        jointsToRemove = ["fin_spine_c_02"]
+        jointToTransferTo = "fin_spine_c_01"
+        transfer.LOD_transferWeights(meshes, jointsToRemove, jointToTransferTo)
+        jointsToDelete.extend(jointsToRemove)
+
+        jointsToRemove = ["fin_spine_d_02"]
+        jointToTransferTo = "fin_spine_d_01"
+        transfer.LOD_transferWeights(meshes, jointsToRemove, jointToTransferTo)
+        jointsToDelete.extend(jointsToRemove)
+
+        jointsToRemove = ["upperarm_twist_01_l", "upperarm_twist_02_l", "bicep_muscle_l", "tricep_muscle_l"]
+        jointToTransferTo = "upperarm_l"
+        transfer.LOD_transferWeights(meshes, jointsToRemove, jointToTransferTo)
+        jointsToDelete.extend(jointsToRemove)
+
+        jointsToRemove = ["upperarm_twist_01_r", "upperarm_twist_02_r", "bicep_muscle_r", "tricep_muscle_r"]
+        jointToTransferTo = "upperarm_r"
+        transfer.LOD_transferWeights(meshes, jointsToRemove, jointToTransferTo)
+        jointsToDelete.extend(jointsToRemove)
+
+        jointsToRemove = ["fin_arm_c_l_02"]
+        jointToTransferTo = "fin_arm_c_l_01"
+        transfer.LOD_transferWeights(meshes, jointsToRemove, jointToTransferTo)
+        jointsToDelete.extend(jointsToRemove)
+
+        jointsToRemove = ["fin_arm_d_l_02"]
+        jointToTransferTo = "fin_arm_d_l_01"
+        transfer.LOD_transferWeights(meshes, jointsToRemove, jointToTransferTo)
+        jointsToDelete.extend(jointsToRemove)
+
+        jointsToRemove = ["fin_arm_b_l_02"]
+        jointToTransferTo = "fin_arm_b_l_01"
+        transfer.LOD_transferWeights(meshes, jointsToRemove, jointToTransferTo)
+        jointsToDelete.extend(jointsToRemove)
+
+        jointsToRemove = ["chain_06", "chain_07"]
+        jointToTransferTo = "chain_05"
+        transfer.LOD_transferWeights(meshes, jointsToRemove, jointToTransferTo)
+        jointsToDelete.extend(jointsToRemove)
+
+        jointsToRemove = ["fin_arm_a_l"]
+        jointToTransferTo = "lowerarm_twist_01_l"
+        transfer.LOD_transferWeights(meshes, jointsToRemove, jointToTransferTo)
+        jointsToDelete.extend(jointsToRemove)
+
+        jointsToRemove = ["fin_arm_c_r_02"]
+        jointToTransferTo = "fin_arm_c_r_01"
+        transfer.LOD_transferWeights(meshes, jointsToRemove, jointToTransferTo)
+        jointsToDelete.extend(jointsToRemove)
+
+        jointsToRemove = ["fin_arm_d_r_02"]
+        jointToTransferTo = "fin_arm_d_r_01"
+        transfer.LOD_transferWeights(meshes, jointsToRemove, jointToTransferTo)
+        jointsToDelete.extend(jointsToRemove)
+
+        jointsToRemove = ["fin_arm_b_r_02"]
+        jointToTransferTo = "fin_arm_b_r_01"
+        transfer.LOD_transferWeights(meshes, jointsToRemove, jointToTransferTo)
+        jointsToDelete.extend(jointsToRemove)
+
+        jointsToRemove = ["fin_arm_a_r"]
+        jointToTransferTo = "lowerarm_twist_01_r"
+        transfer.LOD_transferWeights(meshes, jointsToRemove, jointToTransferTo)
+        jointsToDelete.extend(jointsToRemove)
+
+        jointsToRemove = ["chin_tent_k_l_01", "chin_tent_k_l_02", "chin_tent_k_r_01", "chin_tent_k_r_02", "eyeball_l", "eyeball_r", "eyebrow_in_l", "eyebrow_in_r", "eyebrow_mid", "eyelid_bott_l", "eyelid_bott_r", "eyelid_top_l", "eyelid_top_r", "jaw", "chin_tent_a_l_01", "chin_tent_a_l_02", "chin_tent_a_r_01", "chin_tent_a_r_02", "chin_tent_c_l_01", "chin_tent_c_l_02", "chin_tent_c_r_01", "chin_tent_c_r_02", "chin_tent_d_l_01", "chin_tent_d_l_02", "chin_tent_d_r_01", "chin_tent_d_r_02", "chin_tent_e_l_01", "chin_tent_e_l_02", "chin_tent_e_r_01", "chin_tent_e_r_02", "chin_tent_f_l_01", "chin_tent_f_l_02", "chin_tent_f_r_01", "chin_tent_f_r_02", "chin_tent_g_l_01", "chin_tent_g_l_02", "chin_tent_g_r_01", "chin_tent_g_r_02", "chin_tent_h_l_01", "chin_tent_h_l_02", "chin_tent_h_r_01", "chin_tent_h_r_02", "chin_tent_i_l_01", "chin_tent_i_l_02", "chin_tent_i_r_01", "chin_tent_i_r_02"]
+        jointToTransferTo = "head"
+        transfer.LOD_transferWeights(meshes, jointsToRemove, jointToTransferTo)
+        jointsToDelete.extend(jointsToRemove)
+
+        jointsToRemove = ["fin_head_a_l_02"]
+        jointToTransferTo = "fin_head_a_l_01"
+        transfer.LOD_transferWeights(meshes, jointsToRemove, jointToTransferTo)
+        jointsToDelete.extend(jointsToRemove)
+
+        jointsToRemove = ["fin_head_b_l_02"]
+        jointToTransferTo = "fin_head_b_l_01"
+        transfer.LOD_transferWeights(meshes, jointsToRemove, jointToTransferTo)
+        jointsToDelete.extend(jointsToRemove)
+
+        jointsToRemove = ["fin_head_c_l_02"]
+        jointToTransferTo = "fin_head_c_l_01"
+        transfer.LOD_transferWeights(meshes, jointsToRemove, jointToTransferTo)
+        jointsToDelete.extend(jointsToRemove)
+
+        jointsToRemove = ["fin_head_a_r_02"]
+        jointToTransferTo = "fin_head_a_r_01"
+        transfer.LOD_transferWeights(meshes, jointsToRemove, jointToTransferTo)
+        jointsToDelete.extend(jointsToRemove)
+
+        jointsToRemove = ["fin_head_b_r_02"]
+        jointToTransferTo = "fin_head_b_r_01"
+        transfer.LOD_transferWeights(meshes, jointsToRemove, jointToTransferTo)
+        jointsToDelete.extend(jointsToRemove)
+
+        jointsToRemove = ["fin_head_c_r_02"]
+        jointToTransferTo = "fin_head_c_r_01"
+        transfer.LOD_transferWeights(meshes, jointsToRemove, jointToTransferTo)
+        jointsToDelete.extend(jointsToRemove)
+
+        jointsToRemove = ["fin_head_mid_02"]
+        jointToTransferTo = "fin_head_mid_01"
+        transfer.LOD_transferWeights(meshes, jointsToRemove, jointToTransferTo)
+        jointsToDelete.extend(jointsToRemove)
+
+        jointsToRemove = ["thigh_twist_02_l"]
+        jointToTransferTo = "thigh_twist_01_l"
+        transfer.LOD_transferWeights(meshes, jointsToRemove, jointToTransferTo)
+        jointsToDelete.extend(jointsToRemove)
+
+        jointsToRemove = ["thigh_twist_02_r"]
+        jointToTransferTo = "thigh_twist_01_r"
+        transfer.LOD_transferWeights(meshes, jointsToRemove, jointToTransferTo)
+        jointsToDelete.extend(jointsToRemove)
+
+        jointsToRemove = ["calf_twist_01_l"]
+        jointToTransferTo = "calf_l"
+        transfer.LOD_transferWeights(meshes, jointsToRemove, jointToTransferTo)
+        jointsToDelete.extend(jointsToRemove)
+
+        jointsToRemove = ["calf_twist_01_r"]
+        jointToTransferTo = "calf_r"
+        transfer.LOD_transferWeights(meshes, jointsToRemove, jointToTransferTo)
+        jointsToDelete.extend(jointsToRemove)
+
+        jointsToRemove = ["toe_index_01_l", "toe_index_02_l", "toe_index_03_l", "toe_mid_01_l", "toe_mid_02_l", "toe_mid_03_l", "toe_pinky_01_l", "toe_pinky_02_l", "toe_pinky_03_l"]
+        jointToTransferTo = "ball_l"
+        transfer.LOD_transferWeights(meshes, jointsToRemove, jointToTransferTo)
+        jointsToDelete.extend(jointsToRemove)
+
+        jointsToRemove = ["toe_index_01_r", "toe_index_02_r", "toe_index_03_r", "toe_mid_01_r", "toe_mid_02_r", "toe_mid_03_r", "toe_pinky_01_r", "toe_pinky_02_r", "toe_pinky_03_r"]
+        jointToTransferTo = "ball_r"
+        transfer.LOD_transferWeights(meshes, jointsToRemove, jointToTransferTo)
+        jointsToDelete.extend(jointsToRemove)
+
+
+        # Delete the joints we no longer want.
+        for one in jointsToDelete:
+            if cmds.objExists(one):
+                cmds.delete(one)
+                
+        # Export the fbx file.
+        export.exportLOD("Rampage_LOD3.fbx", "LOD3")
+        
+        # Re-open the file without saving.
+        fullPath = cmds.file(q = True, sceneName = True)
+        cmds.file(fullPath, open=True, f=True)
+
+        print "LOD3 Exported Successfully!!!!!!!!!!!!!!!"
+
+RampageLOD3()
+
+# Rampage LOD4 -------------------------------------------------------------------------------------------------------------
+def RampageLOD4():
+    if saveResult == "Yes":
+
+        cmds.file(save=True)
+    
+        meshes = ["body_geo_LOD4", "chain_link_geo_01_LOD4", "chain_link_geo_02_LOD4", "chain_link_geo_03_LOD4", "chain_link_geo_04_LOD4", "chain_link_geo_05_LOD4", "chain_link_geo_06_LOD4", "chain_link_geo_07_LOD4", "cloth_geo_LOD4", "eyeball_l_geo_LOD4", "eyeball_r_geo_LOD4", "groin_armor_geo_LOD4", "legs_geo_LOD4", "lower_teeth_geo_LOD4", "shackles_geo_LOD4", "shin_armor_geo_LOD4", "thigh_armor_geo_LOD4", "upper_teeth_geo_LOD4"]
+
+        # Export a skin weights file for each mesh we plan to operate on.
+        cmds.select(meshes)
+        menu.exportSelectedSkinWeights(True)
+
+        # Set the pose if needed
+        cmds.setAttr("thumb_02_l.rotateZ", 28.70053097)
+        cmds.setAttr("thumb_01_l.rotateZ", 6.473285809)
+        cmds.setAttr("thumb_01_l.rotateY", 2.036665823)
+        cmds.setAttr("thumb_01_l.rotateX", -34.15792635)
+        cmds.setAttr("pinky_03_l.rotateZ", 64.09065997)
+        cmds.setAttr("pinky_02_l.rotateZ", 45.60472889)
+        cmds.setAttr("pinky_01_l.rotateZ", -1.391492673e-015)
+        cmds.setAttr("pinky_01_l.rotateY", -1.630034274e-014)
+        cmds.setAttr("pinky_01_l.rotateX", 1.192708006e-014)
+        cmds.setAttr("index_03_l.rotateZ", 64.09065997)
+        cmds.setAttr("index_02_l.rotateZ", 45.60472889)
+        cmds.setAttr("index_01_l.rotateZ", -8.945310042e-015)
+        cmds.setAttr("index_01_l.rotateY", -2.981770014e-014)
+        cmds.setAttr("index_01_l.rotateX", -2.266145211e-014)
+
+        cmds.setAttr("thumb_02_r.rotateZ", 28.70053097)
+        cmds.setAttr("thumb_01_r.rotateZ", 6.473285809)
+        cmds.setAttr("thumb_01_r.rotateY", 2.036665823)
+        cmds.setAttr("thumb_01_r.rotateX", -34.15792635)
+        cmds.setAttr("pinky_03_r.rotateZ", 64.09065997)
+        cmds.setAttr("pinky_02_r.rotateZ", 45.60472889)
+        cmds.setAttr("pinky_01_r.rotateZ", -1.391492673e-015)
+        cmds.setAttr("pinky_01_r.rotateY", -1.630034274e-014)
+        cmds.setAttr("pinky_01_r.rotateX", 1.192708006e-014)
+        cmds.setAttr("index_03_r.rotateZ", 64.09065997)
+        cmds.setAttr("index_02_r.rotateZ", 45.60472889)
+        cmds.setAttr("index_01_r.rotateZ", -8.945310042e-015)
+        cmds.setAttr("index_01_r.rotateY", -2.981770014e-014)
+        cmds.setAttr("index_01_r.rotateX", -2.266145211e-014)
+
+        # Delete the history on the meshes so that they get baked down into their posed positions
+        cmds.delete(meshes, ch=True)
+
+        # Move the skeleton back to its bindpose (This script REQUIRES the bind pose to be the pose that we want to export the skeleton to.  This is the case for all known characters created with A.R.T.)
+        cmds.dagPose("root", bp=True, r=True) 
+
+        # Re-import the skinning we exported before onto the mesh so thats its bound in its posed position.
+        cmds.select(meshes)
+        impw.importMultipleSkinWeights("", True)
+        
+        jointsToDelete = []
+
+        # remove unneeded joints
+        jointsToRemove = ["lat_muscle_l", "lat_muscle_r"]
+        jointToTransferTo = "spine_02"
+        transfer.LOD_transferWeights(meshes, jointsToRemove, jointToTransferTo)
+        jointsToDelete.extend(jointsToRemove)
+
+        jointsToRemove = ["fin_scap_a_l_01", "fin_scap_a_l_02", "fin_scap_a_r_01", "fin_scap_a_r_02", "fin_scap_b_l_01", "fin_scap_b_l_02", "fin_scap_b_r_01", "fin_scap_b_r_02", "fin_scap_c_l_01", "fin_scap_c_l_02", "fin_scap_c_r_01", "fin_scap_c_r_02", "fin_scap_d_l", "fin_scap_d_r", "fin_spine_a_01", "fin_spine_a_02", "fin_spine_b_01", "fin_spine_b_02", "fin_spine_c_01", "fin_spine_c_02", "fin_spine_d_01", "fin_spine_d_02", "pec_muscle_l", "pec_muscle_r"]
+        jointToTransferTo = "spine_03"
+        transfer.LOD_transferWeights(meshes, jointsToRemove, jointToTransferTo)
+        jointsToDelete.extend(jointsToRemove)
+
+        jointsToRemove = ["upperarm_twist_01_l", "upperarm_twist_02_l", "bicep_muscle_l", "tricep_muscle_l"]
+        jointToTransferTo = "upperarm_l"
+        transfer.LOD_transferWeights(meshes, jointsToRemove, jointToTransferTo)
+        jointsToDelete.extend(jointsToRemove)
+
+        jointsToRemove = ["upperarm_twist_01_r", "upperarm_twist_02_r", "bicep_muscle_r", "tricep_muscle_r"]
+        jointToTransferTo = "upperarm_r"
+        transfer.LOD_transferWeights(meshes, jointsToRemove, jointToTransferTo)
+        jointsToDelete.extend(jointsToRemove)
+
+        jointsToRemove = ["fin_arm_c_l_01", "fin_arm_c_l_02", "fin_arm_d_l_01", "fin_arm_d_l_02", "lowerarm_twist_01_l", "chain_01", "chain_02", "chain_03", "chain_04", "chain_05", "chain_06", "chain_07", "fin_arm_a_l", "fin_arm_b_l_01", "fin_arm_b_l_02"]
+        jointToTransferTo = "lowerarm_l"
+        transfer.LOD_transferWeights(meshes, jointsToRemove, jointToTransferTo)
+        jointsToDelete.extend(jointsToRemove)
+
+        jointsToRemove = ["fin_arm_c_r_01", "fin_arm_c_r_02", "fin_arm_d_r_01", "fin_arm_d_r_02", "lowerarm_twist_01_r", "fin_arm_a_r", "fin_arm_b_r_01", "fin_arm_b_r_02"]
+        jointToTransferTo = "lowerarm_r"
+        transfer.LOD_transferWeights(meshes, jointsToRemove, jointToTransferTo)
+        jointsToDelete.extend(jointsToRemove)
+
+        jointsToRemove = ["index_01_l", "index_02_l", "index_03_l", "pinky_01_l", "pinky_02_l", "pinky_03_l", "thumb_01_l", "thumb_02_l"]
+        jointToTransferTo = "hand_l"
+        transfer.LOD_transferWeights(meshes, jointsToRemove, jointToTransferTo)
+        jointsToDelete.extend(jointsToRemove)
+
+        jointsToRemove = ["index_01_r", "index_02_r", "index_03_r", "pinky_01_r", "pinky_02_r", "pinky_03_r", "thumb_01_r", "thumb_02_r"]
+        jointToTransferTo = "hand_r"
+        transfer.LOD_transferWeights(meshes, jointsToRemove, jointToTransferTo)
+        jointsToDelete.extend(jointsToRemove)
+
+        jointsToRemove = ["chin_tent_k_l_01", "chin_tent_k_l_02", "chin_tent_k_r_01", "chin_tent_k_r_02", "eyeball_l", "eyeball_r", "eyebrow_in_l", "eyebrow_in_r", "eyebrow_mid", "eyelid_bott_l", "eyelid_bott_r", "eyelid_top_l", "eyelid_top_r", "fin_head_a_l_01", "fin_head_a_l_02", "fin_head_a_r_01", "fin_head_a_r_02", "fin_head_b_l_01", "fin_head_b_l_02", "fin_head_b_r_01", "fin_head_b_r_02", "fin_head_c_l_01", "fin_head_c_l_02", "fin_head_c_r_01", "fin_head_c_r_02", "fin_head_mid_01", "fin_head_mid_02", "jaw", "chin_tent_a_l_01", "chin_tent_a_l_02", "chin_tent_a_r_01", "chin_tent_a_r_02", "chin_tent_c_l_01", "chin_tent_c_l_02", "chin_tent_c_r_01", "chin_tent_c_r_02", "chin_tent_d_l_01", "chin_tent_d_l_02", "chin_tent_d_r_01", "chin_tent_d_r_02", "chin_tent_e_l_01", "chin_tent_e_l_02", "chin_tent_e_r_01", "chin_tent_e_r_02", "chin_tent_f_l_01", "chin_tent_f_l_02", "chin_tent_f_r_01", "chin_tent_f_r_02", "chin_tent_g_l_01", "chin_tent_g_l_02", "chin_tent_g_r_01", "chin_tent_g_r_02", "chin_tent_h_l_01", "chin_tent_h_l_02", "chin_tent_h_r_01", "chin_tent_h_r_02", "chin_tent_i_l_01", "chin_tent_i_l_02", "chin_tent_i_r_01", "chin_tent_i_r_02"]
+        jointToTransferTo = "head"
+        transfer.LOD_transferWeights(meshes, jointsToRemove, jointToTransferTo)
+        jointsToDelete.extend(jointsToRemove)
+
+        jointsToRemove = ["thigh_twist_01_l", "thigh_twist_02_l"]
+        jointToTransferTo = "thigh_l"
+        transfer.LOD_transferWeights(meshes, jointsToRemove, jointToTransferTo)
+        jointsToDelete.extend(jointsToRemove)
+
+        jointsToRemove = ["thigh_twist_02_r", "thigh_twist_01_r"]
+        jointToTransferTo = "thigh_r"
+        transfer.LOD_transferWeights(meshes, jointsToRemove, jointToTransferTo)
+        jointsToDelete.extend(jointsToRemove)
+
+        jointsToRemove = ["calf_twist_01_l"]
+        jointToTransferTo = "calf_l"
+        transfer.LOD_transferWeights(meshes, jointsToRemove, jointToTransferTo)
+        jointsToDelete.extend(jointsToRemove)
+
+        jointsToRemove = ["calf_twist_01_r"]
+        jointToTransferTo = "calf_r"
+        transfer.LOD_transferWeights(meshes, jointsToRemove, jointToTransferTo)
+        jointsToDelete.extend(jointsToRemove)
+
+        jointsToRemove = ["toe_index_01_l", "toe_index_02_l", "toe_index_03_l", "toe_mid_01_l", "toe_mid_02_l", "toe_mid_03_l", "toe_pinky_01_l", "toe_pinky_02_l", "toe_pinky_03_l"]
+        jointToTransferTo = "ball_l"
+        transfer.LOD_transferWeights(meshes, jointsToRemove, jointToTransferTo)
+        jointsToDelete.extend(jointsToRemove)
+
+        jointsToRemove = ["toe_index_01_r", "toe_index_02_r", "toe_index_03_r", "toe_mid_01_r", "toe_mid_02_r", "toe_mid_03_r", "toe_pinky_01_r", "toe_pinky_02_r", "toe_pinky_03_r"]
+        jointToTransferTo = "ball_r"
+        transfer.LOD_transferWeights(meshes, jointsToRemove, jointToTransferTo)
+        jointsToDelete.extend(jointsToRemove)
+
+        # Delete the joints we no longer want.
+        for one in jointsToDelete:
+            if cmds.objExists(one):
+                cmds.delete(one)
+                
+        # Export the fbx file.
+        export.exportLOD("Rampage_LOD4.fbx", "LOD4")
+        
+        # Re-open the file without saving.
+        fullPath = cmds.file(q = True, sceneName = True)
+        cmds.file(fullPath, open=True, f=True)
+
+        print "LOD4 Exported Successfully!!!!!!!!!!!!!!!"
+
+RampageLOD4()
+
+
diff --git a/Engine/Extras/Maya_AnimationRiggingTools/ARTv1/General/ART/Projects/Orion/Scripts/Rampage_Post_Script.py b/Engine/Extras/Maya_AnimationRiggingTools/ARTv1/General/ART/Projects/Orion/Scripts/Rampage_Post_Script.py
new file mode 100644
--- /dev/null
+++ b/Engine/Extras/Maya_AnimationRiggingTools/ARTv1/General/ART/Projects/Orion/Scripts/Rampage_Post_Script.py
@@ -0,0 +1,427 @@
+import maya.cmds as cmds
+# THIS CHARACTER HAS A CUSTOM CHANGE TO HIS RIG POSE THAT NEEDS TO BE DONE IF STEPPED BACK IN ART!!!!!
+# ------------------------------------------------------------------
+# IF you have to go back to the joint mover stage, then select the old rigid mesh group and the root joint, and run the rigid mesh script at the bottom of this file.  
+# ALSO need to rotate the hands to match the locators for them in the scene.
+
+# ------------------------------------------------------------------
+# scale all of the controls to be more visible with the mesh.
+def scaleControls():
+    for side in ["_l", "_r"]:
+        # Arms        
+        cmds.select(["ik_wrist"+side+"_anim.cv[0:7]"])
+        cmds.scale(2.664125, 2.664125, 2.664125, r=True, ocp=False)
+
+        cmds.select(["thumb_finger_fk_ctrl_1"+side+".cv[0:7]"])
+        cmds.scale(5.362914, 5.362914, 5.362914, r=True, ocp=False)
+
+        cmds.select(["thumb_finger_fk_ctrl_2"+side+".cv[0:7]"])
+        cmds.scale(5.362914, 5.362914, 5.362914, r=True, ocp=False)
+
+        cmds.select(["index_finger_fk_ctrl_1"+side+".cv[0:7]"])
+        cmds.scale(5.362914, 5.362914, 5.362914, r=True, ocp=False)
+
+        cmds.select(["index_finger_fk_ctrl_2"+side+".cv[0:7]"])
+        cmds.scale(5.362914, 5.362914, 5.362914, r=True, ocp=False)
+
+        cmds.select(["index_finger_fk_ctrl_3"+side+".cv[0:7]"])
+        cmds.scale(5.362914, 5.362914, 5.362914, r=True, ocp=False)
+
+        cmds.select(["pinky_finger_fk_ctrl_1"+side+".cv[0:7]"])
+        cmds.scale(5.362914, 5.362914, 5.362914, r=True, ocp=False)
+
+        cmds.select(["pinky_finger_fk_ctrl_2"+side+".cv[0:7]"])
+        cmds.scale(5.362914, 5.362914, 5.362914, r=True, ocp=False)
+
+        cmds.select(["pinky_finger_fk_ctrl_3"+side+".cv[0:7]"])
+        cmds.scale(5.362914, 5.362914, 5.362914, r=True, ocp=False)
+
+        for one in ["fk_arm", "fk_elbow", "fk_wrist"]:
+            cmds.select([one+side+"_anim.cv[0:7]"])
+            cmds.scale(2.371997, 2.371997, 2.371997, r=True, ocp=True)
+
+        cmds.select(["fk_clavicle"+side+"_anim.cv[0:32]"])
+        cmds.scale(1, 1, 1, r=True, ocp=False)
+        cmds.move(0, 0, 13, r=True)
+
+        # Legs        
+        cmds.select(["ik_foot_anim"+side+".cv[0:17]"])
+        cmds.scale(2.0357, 2.0357, 2.0357, r=True, ocp=False)
+        cmds.move(0, 21.158168, 15.019494, r=True)
+
+        cmds.select(["toe_index_01"+side+"_anim.cv[0:32]"])
+        cmds.scale(0.514234, 0.514234, 0.514234, r=True, ocp=False)
+        cmds.move(0, 0, 3.9, r=True)
+        
+        cmds.select(["toe_index_02"+side+"_anim.cv[0:32]"])
+        cmds.scale(0.48, 0.48, 0.48, r=True, ocp=False)
+        cmds.move(0, 0, 2.9, r=True)
+
+        cmds.select(["toe_index_03"+side+"_anim.cv[0:32]"])
+        cmds.scale(0.441295, 0.441295, 0.441295, r=True, ocp=False)
+        cmds.move(0, 0, 1.9, r=True)
+
+        cmds.select(["toe_mid_01"+side+"_anim.cv[0:32]"])
+        cmds.scale(0.514234, 0.514234, 0.514234, r=True, ocp=False)
+        cmds.move(0, 0, 3.9, r=True)
+        
+        cmds.select(["toe_mid_02"+side+"_anim.cv[0:32]"])
+        cmds.scale(0.48, 0.48, 0.48, r=True, ocp=False)
+        cmds.move(0, 0, 2.9, r=True)
+
+        cmds.select(["toe_mid_03"+side+"_anim.cv[0:32]"])
+        cmds.scale(0.441295, 0.441295, 0.441295, r=True, ocp=False)
+        cmds.move(0, 0, 1.9, r=True)
+        
+        cmds.select(["toe_pinky_01"+side+"_anim.cv[0:32]"])
+        cmds.scale(0.514234, 0.514234, 0.514234, r=True, ocp=False)
+        cmds.move(0, 0, 3.9, r=True)
+        
+        cmds.select(["toe_pinky_02"+side+"_anim.cv[0:32]"])
+        cmds.scale(0.48, 0.48, 0.48, r=True, ocp=False)
+        cmds.move(0, 0, 2.9, r=True)
+
+        cmds.select(["toe_pinky_03"+side+"_anim.cv[0:32]"])
+        cmds.scale(0.441295, 0.441295, 0.441295, r=True, ocp=False)        
+        cmds.move(0, 0, 1.9, r=True)
+
+        # Face Tentacles
+        for one in ["a", "c", "d", "e", "f", "g", "h", "i", "k"]:
+            cmds.select(["fk_chin_tent_"+one+side+"_01_anim.cv[0:7]"])
+            cmds.scale(0.25, 0.25, 0.25, r=True, ocp=False)
+
+            cmds.select(["fk_chin_tent_"+one+side+"_02_anim.cv[0:7]"])
+            cmds.scale(0.16, 0.16, 0.16, r=True, ocp=False)
+        
+        # Fins
+        cmds.select(["fin_arm_a"+side+"_anim.cv[0:32]"])
+        cmds.scale(0.68, 0.68, 0.68, r=True, ocp=True)
+
+        cmds.select(["fin_scap_d"+side+"_anim.cv[0:32]"])
+        cmds.scale(0.68, 0.68, 0.68, r=True, ocp=True)
+
+        # brows
+        cmds.select(["eyebrow_in"+side+"_anim.cv[0:32]"])
+        cmds.scale(0.132495, 0.132495, 0.132495, r=True, ocp=True)
+
+        # eyes
+        cmds.select(["eyelid_bott"+side+"_anim.cv[0:32]"])
+        cmds.scale(0.1, 0.1, 0.01, r = True, ocp = False)
+        cmds.move(0, -3.757386, 0.654023, r = True)
+
+        cmds.select(["eyelid_top"+side+"_anim.cv[0:32]"])
+        cmds.scale(0.1, 0.1, 0.01, r = True, ocp = False)
+        cmds.move(0, -3.757386, 0.654023, r = True)
+
+    # Other
+    cmds.select(["eyeball_l_anim.cv[0:32]"])
+    cmds.scale(0.1, 0.1, 0.05, r = True, ocp = False)
+    cmds.move(0, -3.757386, 0, r = True)
+
+    cmds.select(["eyeball_r_anim.cv[0:32]"])
+    cmds.scale(0.1, 0.1, 0.05, r = True, ocp = False)
+    cmds.move(0, 3.757386, 0, r = True)
+
+    # jaw
+    cmds.select(["jaw_anim.cv[0:32]"])
+    cmds.scale(0.294922, 0.294922, 0.294922, r=True, ocp=True)
+    cmds.move(0, -12.8, -10.7, r=True)
+
+    cmds.select(["eyebrow_mid_anim.cv[0:32]"])
+    cmds.scale(0.132495, 0.132495, 0.132495, r=True, ocp=True)
+    
+    cmds.select(["l_global_ik_anim.cv[0:4]", "r_global_ik_anim.cv[0:4]"])
+    cmds.scale(1, 1, 1, r=True, ocp=False)
+    cmds.move(0, 0, -22.886708, r=True)
+
+
+    cmds.select(None)
+scaleControls()
+
+cmds.setAttr("Rig_Settings.lUpperarmTwistAmount", 1)
+cmds.setAttr("Rig_Settings.lUpperarmTwist2Amount", 0.4)
+cmds.setAttr("Rig_Settings.rUpperarmTwistAmount", 1)
+cmds.setAttr("Rig_Settings.rUpperarmTwist2Amount", 0.4)
+
+# ------------------------------------------------------------------
+# Adding more constraints to the Fins on the arms to follow the forearms.
+def makeFinsFollowForearms():
+    for side in ["_l", "_r"]:
+        cmds.parentConstraint("driver_lowerarm"+side, "fin_arm_a"+side+"_parent_grp", mo=True, weight=1)
+        cmds.parentConstraint("driver_lowerarm"+side, "fk_fin_arm_b"+side+"_01_grp", mo=True, weight=1)
+        cmds.parentConstraint("driver_lowerarm_twist_01"+side, "fk_fin_arm_c"+side+"_01_grp", mo=True, weight=1)
+        cmds.parentConstraint("driver_lowerarm_twist_01"+side, "fk_fin_arm_d"+side+"_01_grp", mo=True, weight=1)
+
+        cmds.setAttr("fin_arm_a"+side+"_parent_grp_parentConstraint1.interpType", 2)
+        cmds.setAttr("fk_fin_arm_b"+side+"_01_grp_parentConstraint1.interpType", 2)
+        cmds.setAttr("fk_fin_arm_c"+side+"_01_grp_parentConstraint1.interpType", 2)
+        cmds.setAttr("fk_fin_arm_d"+side+"_01_grp_parentConstraint1.interpType", 2)
+
+        cmds.setAttr("fin_arm_a"+side+"_parent_grp_parentConstraint1.driver_lowerarm_twist_01"+side+"W0", 10)
+        cmds.setAttr("fk_fin_arm_b"+side+"_01_grp_parentConstraint1.driver_lowerarm_twist_01"+side+"W0", 1.2)
+        cmds.setAttr("fk_fin_arm_c"+side+"_01_grp_parentConstraint1.driver_lowerarm"+side+"W0", 2.2)
+        cmds.setAttr("fk_fin_arm_d"+side+"_01_grp_parentConstraint1.driver_lowerarm"+side+"W0", 100)       
+makeFinsFollowForearms()
+
+# ------------------------------------------------------------------
+# Create all of the blendshapes that are needed an hook them up.
+def createBlendshapes():
+    # Create the blendshape node for the biceps - Run the following line in the export file.  This way the blendshapes are in the export to the game.
+    #blendshape = cmds.blendShape("body_geo__bicep_l", "body_geo__bicep_r", "body_geo__shoulder_up_l", "body_geo__shoulder_up_r", "body_geo", name="body_geo_blendShape", foc=True, tc=True, origin="local")
+
+    cmds.setAttr("Rig_Settings.lArmMode", 0)
+    cmds.setAttr("Rig_Settings.lClavMode", 0)
+    cmds.setAttr("Rig_Settings.rArmMode", 0)
+    cmds.setAttr("Rig_Settings.rClavMode", 0)
+
+    for side in ["_l", "_r"]:
+        # Create SDK's to drive the shoulder blendshape by thex clavicle joint
+        cmds.setAttr("fk_clavicle"+side+"_anim.ry", 7)
+        cmds.setAttr("body_geo_blendShape.body_geo__shoulder_up"+side, 0)
+        cmds.setDrivenKeyframe("body_geo_blendShape.body_geo__shoulder_up"+side, cd="clavicle"+side+".rotateY", itt="flat", ott="flat")
+        cmds.setAttr("fk_clavicle"+side+"_anim.ry", -46)
+        cmds.setAttr("body_geo_blendShape.body_geo__shoulder_up"+side, 1)
+        cmds.setDrivenKeyframe("body_geo_blendShape.body_geo__shoulder_up"+side, cd="clavicle"+side+".rotateY", itt="flat", ott="flat")
+        cmds.setAttr("fk_clavicle"+side+"_anim.ry", 0)
+
+        # Create SDK's to drive the bicep blendshape by the elbow joint
+        cmds.setAttr("fk_elbow"+side+"_anim.rz", 0)
+        cmds.setAttr("body_geo_blendShape.body_geo__bicep"+side, 0)
+        cmds.setDrivenKeyframe("body_geo_blendShape.body_geo__bicep"+side, cd="lowerarm"+side+".rotateZ", itt="flat", ott="flat")
+        cmds.setAttr("fk_elbow"+side+"_anim.rz", -100)
+        cmds.setAttr("body_geo_blendShape.body_geo__bicep"+side, 1)
+        cmds.setDrivenKeyframe("body_geo_blendShape.body_geo__bicep"+side, cd="lowerarm"+side+".rotateZ", itt="flat", ott="flat")
+        cmds.setAttr("fk_elbow"+side+"_anim.rz", 0)
+
+    cmds.setAttr("Rig_Settings.lArmMode", 1)
+    cmds.setAttr("Rig_Settings.lClavMode", 1)
+    cmds.setAttr("Rig_Settings.rArmMode", 1)
+    cmds.setAttr("Rig_Settings.rClavMode", 1)
+createBlendshapes()
+
+# ------------------------------------------------------------------
+# Adding more constraints to the chin tentacles on the face so they can partially follow the head more.
+def chinTentaclesOnJaw():
+    for side in ["_l", "_r"]:
+        tentacleDict = {"e":0.1, "f":0.5, "g":0.6, "h":1, "i":5}
+        for one in tentacleDict:
+            cmds.parentConstraint("driver_head", "fk_chin_tent_"+one+side+"_01_grp", mo=True, weight=1)
+            cmds.setAttr("fk_chin_tent_"+one+side+"_01_grp_parentConstraint1.interpType", 2)
+            cmds.setAttr("fk_chin_tent_"+one+side+"_01_grp_parentConstraint1.driver_headW1", tentacleDict[one])
+chinTentaclesOnJaw()
+
+# ------------------------------------------------------------------
+def fixUpToes():
+    # change the way the Set driven keys on the heel and ball pivots work so that the leg can work more like a quadruped leg.
+    cmds.keyframe("ik_foot_heel_pivot_l_rotateZ", edit=True, index=(2,2), absolute=True, valueChange=0)
+    cmds.keyframe("ik_foot_ball_pivot_l_rotateZ", edit=True, index=(2,2), absolute=True, valueChange=90)
+
+    cmds.keyframe("ik_foot_heel_pivot_r_rotateZ", edit=True, index=(2,2), absolute=True, valueChange=0)
+    cmds.keyframe("ik_foot_ball_pivot_r_rotateZ", edit=True, index=(2,2), absolute=True, valueChange=90)
+    
+    cmds.delete("toe_rig_l_grp", "toe_rig_r_grp")
+fixUpToes()
+
+# ------------------------------------------------------------------
+def backFins():
+    for side in ["_l", "_r"]:
+        cmds.parentConstraint("driver_spine_02", "fk_fin_scap_b"+side+"_01_grp", mo=True, weight=1)
+        cmds.setAttr("fk_fin_scap_b"+side+"_01_grp_parentConstraint1.interpType", 2)
+        cmds.setAttr("fk_fin_scap_b"+side+"_01_grp_parentConstraint1.driver_spine_02W1", 0.4)
+
+        cmds.parentConstraint("driver_spine_02", "fk_fin_scap_c"+side+"_01_grp", mo=True, weight=1)
+        cmds.setAttr("fk_fin_scap_c"+side+"_01_grp_parentConstraint1.interpType", 2)
+        cmds.setAttr("fk_fin_scap_c"+side+"_01_grp_parentConstraint1.driver_spine_02W1", 0.7)
+
+        cmds.parentConstraint("driver_spine_02", "fin_scap_d"+side+"_parent_grp", mo=True, weight=1)
+        cmds.setAttr("fin_scap_d"+side+"_parent_grp_parentConstraint1.interpType", 2)
+        cmds.setAttr("fin_scap_d"+side+"_parent_grp_parentConstraint1.driver_spine_02W1", 1.5)
+
+    tentacleDict = {"b":0.4, "c":0.7, "d":1.5}
+    for one in tentacleDict:
+        cmds.parentConstraint("driver_spine_02", "fk_fin_spine_"+one+"_01_grp", mo=True, weight=1)
+        cmds.setAttr("fk_fin_spine_"+one+"_01_grp_parentConstraint1.interpType", 2)
+        cmds.setAttr("fk_fin_spine_"+one+"_01_grp_parentConstraint1.driver_spine_02W1", tentacleDict[one])
+backFins()
+
+# ------------------------------------------------------------------
+def invertEyelids():
+    sides = ["_l", "_r"]
+    for side in sides:
+        cmds.delete("driver_eyelid_bott"+side+"_parentConstraint1")
+        cmds.rotate(0, -180, 0, "eyelid_bott"+side+"_anim_grp", r=True, os=True)
+        cmds.parentConstraint("eyelid_bott"+side+"_anim", "driver_eyelid_bott"+side, mo=True)
+invertEyelids()
+
+def invertRightEyeball():
+    cmds.delete("driver_eyeball_r_parentConstraint1")
+    cmds.rotate(0, -180, -180, "eyeball_r_anim_grp", r=True, os=True)
+    cmds.parentConstraint("eyeball_r_anim", "driver_eyeball_r", mo=True)
+invertRightEyeball()
+
+# ------------------------------------------------------------------
+def dynamicsOff():
+    for one in ["chain_dyn_anim", "chin_tent_a_l_dyn_anim", "chin_tent_a_r_dyn_anim", "chin_tent_c_l_dyn_anim", "chin_tent_c_r_dyn_anim", "chin_tent_d_l_dyn_anim", "chin_tent_d_r_dyn_anim", "chin_tent_e_l_dyn_anim", "chin_tent_e_r_dyn_anim", "chin_tent_f_l_dyn_anim", "chin_tent_f_r_dyn_anim", "chin_tent_g_l_dyn_anim", "chin_tent_g_r_dyn_anim", "chin_tent_h_l_dyn_anim", "chin_tent_h_r_dyn_anim", "chin_tent_i_l_dyn_anim", "chin_tent_i_r_dyn_anim", "chin_tent_k_l_dyn_anim", "chin_tent_k_r_dyn_anim", "fin_arm_b_l_dyn_anim", "fin_arm_b_r_dyn_anim", "fin_arm_c_l_dyn_anim", "fin_arm_c_r_dyn_anim", "fin_arm_d_l_dyn_anim", "fin_arm_d_r_dyn_anim", "fin_head_a_l_dyn_anim", "fin_head_a_r_dyn_anim", "fin_head_b_l_dyn_anim", "fin_head_b_r_dyn_anim", "fin_head_c_l_dyn_anim", "fin_head_c_r_dyn_anim", "fin_head_mid_dyn_anim", "fin_scap_a_l_dyn_anim", "fin_scap_a_r_dyn_anim", "fin_scap_b_l_dyn_anim", "fin_scap_b_r_dyn_anim", "fin_scap_c_l_dyn_anim", "fin_scap_c_r_dyn_anim", "fin_spine_a_dyn_anim", "fin_spine_b_dyn_anim", "fin_spine_c_dyn_anim", "fin_spine_d_dyn_anim"]:
+        cmds.setAttr(one+".chainStartEnvelope", 0)
+        cmds.setAttr(one+".chainStartEnvelope", l=True)
+dynamicsOff()
+
+# ------------------------------------------------------------------
+def createGameCam():
+    gameCam = cmds.camera(centerOfInterest=5, focalLength=18, lensSqueezeRatio=1, cameraScale=1, horizontalFilmAperture=1.41732, horizontalFilmOffset=0, verticalFilmAperture=0.94488, verticalFilmOffset=0, filmFit="Overscan", overscan=1.05, motionBlur=0, shutterAngle=144, nearClipPlane=1.0, farClipPlane=10000, orthographic=0, orthographicWidth=30, panZoomEnabled=0, horizontalPan=0, verticalPan=0, zoom=1, displayFilmGate=True)[0]
+    print gameCam
+    gameCam = cmds.rename(gameCam, "GameCam")
+    cmds.camera(gameCam, e=True, horizontalFilmAperture=1.6724376, focalLength=21.2399857011)
+    cmds.setAttr(gameCam+".displayGateMaskOpacity", 1)
+    cmds.setAttr(gameCam+".displayGateMaskColor", 0, 0, 0, type="double3")
+    cmds.setAttr(gameCam+".r", 90, 0, 180, type="double3")
+    
+    # Create the control
+    control = cmds.circle(name=gameCam+"_Ctrl", ch=True, o=True, r=15)[0]
+    cmds.makeIdentity(control, a=True, t=True, r=True, s=True)
+    # Create groups for the sdk and the constraints
+    const_grp = cmds.group(control, name=control+"_grp")
+    if cmds.objExists("rig_grp"):
+        cmds.parent(const_grp, "rig_grp")
+    # Parent Constrain the cam to the control
+    cmds.parentConstraint(control, gameCam, mo=True)
+    
+    cmds.connectAttr(control+".visibility", gameCam+".visibility", f=True)
+    cmds.setAttr(const_grp+".v", 0)
+    
+    cmds.setAttr(const_grp+".t", 0, 500, 200, type="double3")
+createGameCam()
+
+# ------------------------------------------------------------------
+def addCapsule(radius, halfHeight):
+    trueHalfHeight = (halfHeight*2)-(radius*2)
+    capsule = cmds.polyCylinder(name="collision_capsule", r=radius, h=trueHalfHeight, sx=18, sy=1, sz=5, ax=[0,0,1], rcp=True, cuv=3, ch=True)
+    cmds.setAttr(capsule[0]+".tz", halfHeight+2.2)
+    cmds.setAttr(capsule[0]+".rz", 90)
+
+    cmds.addAttr(capsule[0], ln="Radius", dv=radius, keyable=True)
+    cmds.addAttr(capsule[0], ln="HalfHeight", dv=halfHeight, keyable=True)
+
+    expr = capsule[1]+".radius = "+capsule[0]+".Radius;\n"+capsule[1]+".height = ("+capsule[0]+".HalfHeight*2) - ("+capsule[0]+".Radius*2);"
+    cmds.expression(s=expr, o=capsule[1], ae=1, uc="all")
+
+    cmds.setAttr(capsule[0]+".overrideEnabled", 1)
+    cmds.setAttr(capsule[0]+".overrideColor", 18)
+    cmds.setAttr(capsule[0]+".overrideShading", 0)
+
+    cmds.addAttr(capsule[0], ln="Shaded", dv=0, min=0, max=1, keyable=True)
+    cmds.connectAttr(capsule[0]+".Shaded", capsule[0]+".overrideShading")
+
+    cmds.parentConstraint("root", capsule[0], mo=True)
+    cmds.select(capsule[0], r=True)
+    capsuleDisplayLayer = cmds.createDisplayLayer(nr=True, name="Collision Capsule")
+    cmds.setAttr(capsuleDisplayLayer+".displayType", 2)
+    cmds.setAttr(capsuleDisplayLayer+".v", 0)
+    
+    for attr in ["tx", "ty", "tz", "rx", "ry", "rz", "sx", "sy", "sz"]:
+        cmds.setAttr(capsule[0]+"."+attr, lock=True, keyable=False, channelBox=False)
+
+    cmds.select(clear=True)
+addCapsule(80, 90)
+
+# ------------------------------------------------------------------
+# add ik bones Function that is called later after we open the export file.
+def addIKBones():
+    try:
+		# create the joints
+		cmds.select(clear = True)
+		ikFootRoot = cmds.joint(name = "ik_foot_root")
+		cmds.select(clear = True)
+		
+		cmds.select(clear = True)
+		ikFootLeft = cmds.joint(name = "ik_foot_l")
+		cmds.select(clear = True)
+		
+		cmds.select(clear = True)
+		ikFootRight = cmds.joint(name = "ik_foot_r")
+		cmds.select(clear = True)
+		
+		cmds.select(clear = True)
+		ikHandRoot = cmds.joint(name = "ik_hand_root")
+		cmds.select(clear = True)
+		
+		cmds.select(clear = True)
+		ikHandGun = cmds.joint(name = "ik_hand_gun")
+		cmds.select(clear = True)
+		
+		cmds.select(clear = True)
+		ikHandLeft = cmds.joint(name = "ik_hand_l")
+		cmds.select(clear = True)
+		
+		cmds.select(clear = True)
+		ikHandRight = cmds.joint(name = "ik_hand_r")
+		cmds.select(clear = True)
+		
+		# create hierarhcy
+		cmds.parent(ikFootRoot, "root")
+		cmds.parent(ikHandRoot, "root")
+		cmds.parent(ikFootLeft, ikFootRoot)
+		cmds.parent(ikFootRight, ikFootRoot)
+		cmds.parent(ikHandGun, ikHandRoot)
+		cmds.parent(ikHandLeft, ikHandGun)
+		cmds.parent(ikHandRight, ikHandGun)
+		
+		# constrain the joints
+		leftHandConstraint = cmds.parentConstraint("hand_l", ikHandLeft)[0]
+		rightHandGunConstraint = cmds.parentConstraint("hand_r", ikHandGun)[0]
+		rightHandConstraint = cmds.parentConstraint("hand_r", ikHandRight)[0]
+		leftFootConstraint = cmds.parentConstraint("foot_l", ikFootLeft)[0]
+		rightFootConstraint = cmds.parentConstraint("foot_r", ikFootRight)[0]
+	
+    except:
+		cmds.warning("Something went wrong")
+
+# ------------------------------------------------------------------
+# ADD the IK bones to the Export file.
+#save the AnimRig file before opening the export file.
+cmds.file(save = True, type = "mayaBinary", force = True)
+
+#ADD IK BONES
+#open export file
+sceneName = cmds.file(q = True, sceneName = True)
+exportFile = sceneName.partition("AnimRigs")[0] + "ExportFiles" +  sceneName.partition("AnimRigs")[2].partition(".mb")[0] + "_Export.mb"
+cmds.file(exportFile, open = True, force = True)
+
+#add ik bones
+addIKBones()
+
+# ------------------------------------------------------------------
+def addCustomAttrs():
+    attrs = cmds.listAttr("root")
+    exists = False
+    for one in attrs:
+        if one.find("DistanceCurve") != -1:
+            exists = True
+    if exists == False:
+        cmds.addAttr("root", ln="DistanceCurve", keyable=True)
+addCustomAttrs()
+
+#save the export file.
+cmds.file(save = True, type = "mayaBinary", force = True)
+
+
+
+# --------------------------------ONLY FOR MANUAL USE.  NOT PART OF THE POST SCRIPT BUILD----------------------------------
+# Select a group of geo and a skeleton heirarcy and run this script and it will search in the skeleton
+# for a joint with the same name minus "_geo".  If it finds one it will parent constrain the geo to the joint.
+def ConstrainRigidMesh():
+    selected = cmds.ls(selection=True)
+    
+    if not len(selected) is 2:
+        cmds.confirmDialog(icon = "warning", title = "Rigid Mesh Constrainer", message = "Select only 2 objects.  Geo group then skeleton root.   Geo names need to match the joint+_geo.")
+        return
+    
+    geoList = cmds.listRelatives(selected[0], ad=True, type="transform")
+    print geoList
+    skelList = cmds.listRelatives(selected[1], ad=True, type="joint")
+        
+    for geo in geoList:
+        for i in range(0,len(skelList)):
+            if skelList[i]+"_rigidgeo" == geo:
+                const = cmds.parentConstraint(skelList[i], geo, name=geo+"_PrntCnst", mo=True, weight=1)[0]
+#ConstrainRigidMesh()
\ No newline at end of file
