Branching

diff --git a/Engine/Extras/Maya_AnimationRiggingTools/ARTv1/General/ART/Projects/Orion/Scripts/RiftMage_CreateLODs.py b/Engine/Extras/Maya_AnimationRiggingTools/ARTv1/General/ART/Projects/Orion/Scripts/RiftMage_CreateLODs.py
new file mode 100644
--- /dev/null
+++ b/Engine/Extras/Maya_AnimationRiggingTools/ARTv1/General/ART/Projects/Orion/Scripts/RiftMage_CreateLODs.py
@@ -0,0 +1,645 @@
+import maya.cmds as cmds
+import customMayaMenu as menu
+import Tools.ART_ImportMultipleWeights as impw
+import Tools.ART_TransferWeights as transfer
+import Tools.ART_ExportLODtoFBX as export
+reload(export)
+
+################################################## RIFTMAGE #########################################################################
+# Warn the user that this operation will save their file.
+saveResult = cmds.confirmDialog(title = "Save Warning", message="In order to continue your file must be saved.  Would you like to save it?  If yes it will be saved and after the operation is complete your file will be re-opened.", button = ["Yes", "Cancel"], defaultButton='Yes', cancelButton='Cancel', dismissString='Cancel')
+
+# If they are ok with saving the file, warn them that it will also delete all of the current files in their skin weights directory.
+if saveResult == "Yes":
+    deleteSkinFilesResult = cmds.confirmDialog(title = "SkinWeights Warning", message="In order to continue we must delete all of the files in your ..\General\ART\SkinWeights\ directory.  Would you like to continue?.", button = ["Yes", "Cancel"], defaultButton='Yes', cancelButton='Cancel', dismissString='Cancel')
+    if deleteSkinFilesResult == "Yes":
+        toolsPath = cmds.internalVar(usd = True) + "mayaTools.txt"
+        if os.path.exists(toolsPath):
+            f = open(toolsPath, 'r')
+            mayaToolsDir = f.readline()
+            f.close()
+
+        path = mayaToolsDir + "\General\ART\SkinWeights\\"
+        files = os.listdir(path)
+        #delete all weight cache files in the system folder
+        for file in files:
+            os.remove(path + file)
+            #print path + file
+
+
+
+# RiftMage LOD1 -------------------------------------------------------------------------------------------------------------
+def RiftMageLOD1():
+    if saveResult == "Yes":
+
+        cmds.file(save=True)
+
+        meshes = ["arms_geo_LOD1", "beard_geo_LOD1", "brooch_geo_LOD1", "buttstrap_geo_LOD1", "chest_buttons_geo_LOD1", "chest_flaps_geo_LOD1", "codpiece_buttons_geo_LOD1", "codpiece_geo_LOD1", "collar_geo_LOD1", "cowl_geo_LOD1", "elbowpad_geo_LOD1", "eye_tear_duct_geo_LOD1", "eye_wetness_geo_LOD1", "eyeballs_geo_LOD1", "eyelash_geo_LOD1", "eyeshadow_geo_LOD1", "flappendant_geo_LOD1", "forearm_pads_geo_LOD1", "head_geo_LOD1", "hip_parts_geo_LOD1", "hood_geo_LOD1", "inner_ring_geo_LOD1", "kneepads_geo_LOD1", "legplate_geo_LOD1", "legs_geo_LOD1", "midflaps_geo_LOD1", "outer_ring_bott_geo_LOD1", "outer_ring_inside_geo_LOD1", "outer_ring_knobs_geo_LOD1", "outer_ring_top_geo_LOD1", "portal_geo_LOD1", "sideflaps_geo_LOD1", "teeth_lower_geo_LOD1", "teeth_upper_geo_LOD1", "tongue_geo_LOD1", "torso_detail_geo_LOD1", "torso_geo_LOD1", "wrist_armor_geo_LOD1"]
+
+        jointsToDelete = []
+
+        # remove unneeded joints
+        jointsToRemove = ["cloth_chest_bk_l", "cloth_chest_bk_r", "cloth_chest_fr_l", "cloth_chest_fr_r"]
+        jointToTransferTo = "spine_02"
+        transfer.LOD_transferWeights(meshes, jointsToRemove, jointToTransferTo)
+        jointsToDelete.extend(jointsToRemove)
+
+        jointsToRemove = ["cowl_bk", "cowl_l", "cowl_r"]
+        jointToTransferTo = "spine_03"
+        transfer.LOD_transferWeights(meshes, jointsToRemove, jointToTransferTo)
+        jointsToDelete.extend(jointsToRemove)
+
+        jointsToRemove = ["eyeball_l", "eyeball_r", "eyebrow_inner_l", "eyebrow_inner_r", "eyebrow_mid", "eyebrow_outer_l", "eyebrow_outer_r", "hood_l", "hood_r"]
+        jointToTransferTo = "head"
+        transfer.LOD_transferWeights(meshes, jointsToRemove, jointToTransferTo)
+        jointsToDelete.extend(jointsToRemove)
+        
+        jointsToRemove = ["outer_ring_A_piston01_l", "outer_ring_A_piston02_l"]
+        jointToTransferTo = "outer_ring_A_l"
+        transfer.LOD_transferWeights(meshes, jointsToRemove, jointToTransferTo)
+        jointsToDelete.extend(jointsToRemove)
+
+        jointsToRemove = ["outer_ring_B_piston01_l"]
+        jointToTransferTo = "outer_ring_B_l"
+        transfer.LOD_transferWeights(meshes, jointsToRemove, jointToTransferTo)
+        jointsToDelete.extend(jointsToRemove)
+
+        jointsToRemove = ["outer_ring_A_piston01_r", "outer_ring_A_piston02_r"]
+        jointToTransferTo = "outer_ring_A_r"
+        transfer.LOD_transferWeights(meshes, jointsToRemove, jointToTransferTo)
+        jointsToDelete.extend(jointsToRemove)
+
+        jointsToRemove = ["outer_ring_B_piston01_r"]
+        jointToTransferTo = "outer_ring_B_r"
+        transfer.LOD_transferWeights(meshes, jointsToRemove, jointToTransferTo)
+        jointsToDelete.extend(jointsToRemove)
+
+
+        # Delete the joints we no longer want.
+        for one in jointsToDelete:
+            if cmds.objExists(one):
+                cmds.delete(one)
+                
+        # Export the fbx file.
+        export.exportLOD("RiftMage_LOD1.fbx", "LOD1")
+        
+        # Re-open the file without saving.
+        fullPath = cmds.file(q = True, sceneName = True)
+        cmds.file(fullPath, open=True, f=True)
+
+        print "LOD1 Exported Successfully!!!!!!!!!!!!!!!"
+
+RiftMageLOD1()
+
+# RiftMage LOD2 -------------------------------------------------------------------------------------------------------------
+def RiftMageLOD2():
+    if saveResult == "Yes":
+
+        cmds.file(save=True)
+        
+        meshes = ["arms_geo_LOD2", "beard_geo_LOD2", "brooch_geo_LOD2", "buttstrap_geo_LOD2", "chest_buttons_geo_LOD2", "chest_flaps_geo_LOD2", "codpiece_buttons_geo_LOD2", "codpiece_geo_LOD2", "collar_geo_LOD2", "cowl_geo_LOD2", "elbowpad_geo_LOD2", "eye_tear_duct_geo_LOD2", "eye_wetness_geo_LOD2", "eyeballs_geo_LOD2", "eyelash_geo_LOD2", "eyeshadow_geo_LOD2", "flappendant_geo_LOD2", "forearm_pads_geo_LOD2", "head_geo_LOD2", "hip_parts_geo_LOD2", "hood_geo_LOD2", "inner_ring_geo_LOD2", "kneepads_geo_LOD2", "legplate_geo_LOD2", "legs_geo_LOD2", "midflaps_geo_LOD2", "outer_ring_bott_geo_LOD2", "outer_ring_inside_geo_LOD2", "outer_ring_knobs_geo_LOD2", "outer_ring_top_geo_LOD2", "portal_geo_LOD2", "sideflaps_geo_LOD2", "teeth_lower_geo_LOD2", "teeth_upper_geo_LOD2", "tongue_geo_LOD2", "torso_detail_geo_LOD2", "torso_geo_LOD2", "wrist_armor_geo_LOD2"]
+
+        # Export a skin weights file for each mesh we plan to operate on.
+        cmds.select(meshes)
+        menu.exportSelectedSkinWeights(True)
+
+        # Set the pose if needed
+        cmds.setAttr("thumb_03_l.rotateZ", 23.02290762)
+        cmds.setAttr("thumb_02_l.rotateZ", 10.06960964)
+        cmds.setAttr("ring_03_l.rotateZ", 29.40776524)
+        cmds.setAttr("ring_02_l.rotateZ", 29.40776524)
+        cmds.setAttr("pinky_03_l.rotateZ", 39.74899767)
+        cmds.setAttr("pinky_02_l.rotateZ", 39.74899767)
+        cmds.setAttr("middle_03_l.rotateZ", 23.93256607)
+        cmds.setAttr("middle_02_l.rotateZ", 23.93256607)
+        cmds.setAttr("index_03_l.rotateZ", 16.68034218)
+        cmds.setAttr("index_02_l.rotateZ", 16.68034218)
+
+        cmds.setAttr("thumb_03_r.rotateZ", 23.02290762)
+        cmds.setAttr("thumb_02_r.rotateZ", 10.06960964)
+        cmds.setAttr("ring_03_r.rotateZ", 29.40776524)
+        cmds.setAttr("ring_02_r.rotateZ", 29.40776524)
+        cmds.setAttr("pinky_03_r.rotateZ", 39.74899767)
+        cmds.setAttr("pinky_02_r.rotateZ", 39.74899767)
+        cmds.setAttr("middle_03_r.rotateZ", 23.93256607)
+        cmds.setAttr("middle_02_r.rotateZ", 23.93256607)
+        cmds.setAttr("index_03_r.rotateZ", 16.68034218)
+        cmds.setAttr("index_02_r.rotateZ", 16.68034218)
+
+        # Delete the history on the meshes so that they get baked down into their posed positions
+        cmds.delete(meshes, ch=True)
+
+        # Move the skeleton back to its bindpose (This script REQUIRES the bind pose to be the pose that we want to export the skeleton to.  This is the case for all known characters created with A.R.T.)
+        cmds.dagPose("root", bp=True, r=True) 
+
+        # Re-import the skinning we exported before onto the mesh so thats its bound in its posed position.
+        cmds.select(meshes)
+        impw.importMultipleSkinWeights("", True)
+
+        jointsToDelete = []
+
+         # remove unneeded joints
+        jointsToRemove = ["loin_cloth_bk_03"]
+        jointToTransferTo = "loin_cloth_bk_02"
+        transfer.LOD_transferWeights(meshes, jointsToRemove, jointToTransferTo)
+        jointsToDelete.extend(jointsToRemove)
+
+        jointsToRemove = ["loin_cloth_fr_03"]
+        jointToTransferTo = "loin_cloth_fr_02"
+        transfer.LOD_transferWeights(meshes, jointsToRemove, jointToTransferTo)
+        jointsToDelete.extend(jointsToRemove)
+
+        jointsToRemove = ["chest_cloth_03", "chest_cloth_04"]
+        jointToTransferTo = "chest_cloth_02"
+        transfer.LOD_transferWeights(meshes, jointsToRemove, jointToTransferTo)
+        jointsToDelete.extend(jointsToRemove)
+
+        jointsToRemove = ["cloth_chest_bk_l", "cloth_chest_bk_r", "cloth_chest_fr_l", "cloth_chest_fr_r"]
+        jointToTransferTo = "spine_02"
+        transfer.LOD_transferWeights(meshes, jointsToRemove, jointToTransferTo)
+        jointsToDelete.extend(jointsToRemove)
+
+        jointsToRemove = ["cowl_bk", "cowl_l", "cowl_r"]
+        jointToTransferTo = "spine_03"
+        transfer.LOD_transferWeights(meshes, jointsToRemove, jointToTransferTo)
+        jointsToDelete.extend(jointsToRemove)
+
+        jointsToRemove = ["upperarm_twist_01_l", "upperarm_twist_02_l"]
+        jointToTransferTo = "upperarm_l"
+        transfer.LOD_transferWeights(meshes, jointsToRemove, jointToTransferTo)
+        jointsToDelete.extend(jointsToRemove)
+
+        jointsToRemove = ["upperarm_twist_01_r", "upperarm_twist_02_r"]
+        jointToTransferTo = "upperarm_r"
+        transfer.LOD_transferWeights(meshes, jointsToRemove, jointToTransferTo)
+        jointsToDelete.extend(jointsToRemove)
+
+        jointsToRemove = ["lowerarm_twist_01_l", "lowerarm_twist_02_l"]
+        jointToTransferTo = "lowerarm_l"
+        transfer.LOD_transferWeights(meshes, jointsToRemove, jointToTransferTo)
+        jointsToDelete.extend(jointsToRemove)
+
+        jointsToRemove = ["lowerarm_twist_01_r", "lowerarm_twist_02_r"]
+        jointToTransferTo = "lowerarm_r"
+        transfer.LOD_transferWeights(meshes, jointsToRemove, jointToTransferTo)
+        jointsToDelete.extend(jointsToRemove)
+
+        jointsToRemove = ["index_02_l", "index_03_l"]
+        jointToTransferTo = "index_01_l"
+        transfer.LOD_transferWeights(meshes, jointsToRemove, jointToTransferTo)
+        jointsToDelete.extend(jointsToRemove)
+
+        jointsToRemove = ["middle_02_l", "middle_03_l"]
+        jointToTransferTo = "middle_01_l"
+        transfer.LOD_transferWeights(meshes, jointsToRemove, jointToTransferTo)
+        jointsToDelete.extend(jointsToRemove)
+
+        jointsToRemove = ["pinky_02_l", "pinky_03_l"]
+        jointToTransferTo = "pinky_03_l"
+        transfer.LOD_transferWeights(meshes, jointsToRemove, jointToTransferTo)
+        jointsToDelete.extend(jointsToRemove)
+
+        jointsToRemove = ["ring_02_l", "ring_03_l"]
+        jointToTransferTo = "ring_01_l"
+        transfer.LOD_transferWeights(meshes, jointsToRemove, jointToTransferTo)
+        jointsToDelete.extend(jointsToRemove)
+
+        jointsToRemove = ["thumb_02_l", "thumb_03_l"]
+        jointToTransferTo = "thumb_01_l"
+        transfer.LOD_transferWeights(meshes, jointsToRemove, jointToTransferTo)
+        jointsToDelete.extend(jointsToRemove)
+
+        jointsToRemove = ["index_02_r", "index_03_r"]
+        jointToTransferTo = "index_01_r"
+        transfer.LOD_transferWeights(meshes, jointsToRemove, jointToTransferTo)
+        jointsToDelete.extend(jointsToRemove)
+
+        jointsToRemove = ["middle_02_r", "middle_03_r"]
+        jointToTransferTo = "middle_01_r"
+        transfer.LOD_transferWeights(meshes, jointsToRemove, jointToTransferTo)
+        jointsToDelete.extend(jointsToRemove)
+
+        jointsToRemove = ["pinky_02_r", "pinky_03_r"]
+        jointToTransferTo = "pinky_03_r"
+        transfer.LOD_transferWeights(meshes, jointsToRemove, jointToTransferTo)
+        jointsToDelete.extend(jointsToRemove)
+
+        jointsToRemove = ["ring_02_r", "ring_03_r"]
+        jointToTransferTo = "ring_01_r"
+        transfer.LOD_transferWeights(meshes, jointsToRemove, jointToTransferTo)
+        jointsToDelete.extend(jointsToRemove)
+
+        jointsToRemove = ["thumb_02_r", "thumb_03_r"]
+        jointToTransferTo = "thumb_01_r"
+        transfer.LOD_transferWeights(meshes, jointsToRemove, jointToTransferTo)
+        jointsToDelete.extend(jointsToRemove)
+        
+        jointsToRemove = ["eyeball_l", "eyeball_r", "eyebrow_inner_l", "eyebrow_inner_r", "eyebrow_mid", "eyebrow_outer_l", "eyebrow_outer_r", "eyelid_bott_l", "eyelid_bott_r", "eyelid_top_l", "eyelid_top_r", "hood_l", "hood_r"]
+        jointToTransferTo = "head"
+        transfer.LOD_transferWeights(meshes, jointsToRemove, jointToTransferTo)
+        jointsToDelete.extend(jointsToRemove)
+
+        jointsToRemove = ["outer_ring_A_piston01_l", "outer_ring_A_piston02_l", "outer_ring_B_l", "outer_ring_B_piston01_l", "outer_ring_C_l"]
+        jointToTransferTo = "outer_ring_A_l"
+        transfer.LOD_transferWeights(meshes, jointsToRemove, jointToTransferTo)
+        jointsToDelete.extend(jointsToRemove)
+
+        jointsToRemove = ["outer_ring_A_piston01_r", "outer_ring_A_piston02_r", "outer_ring_B_r", "outer_ring_B_piston01_r", "outer_ring_C_r"]
+        jointToTransferTo = "outer_ring_A_r"
+        transfer.LOD_transferWeights(meshes, jointsToRemove, jointToTransferTo)
+        jointsToDelete.extend(jointsToRemove)
+
+        jointsToRemove = ["portal"]
+        jointToTransferTo = "portal_ring"
+        transfer.LOD_transferWeights(meshes, jointsToRemove, jointToTransferTo)
+        jointsToDelete.extend(jointsToRemove)
+
+        jointsToRemove = ["thigh_twist_01_l", "thigh_twist_02_l"]
+        jointToTransferTo = "thigh_l"
+        transfer.LOD_transferWeights(meshes, jointsToRemove, jointToTransferTo)
+        jointsToDelete.extend(jointsToRemove)
+
+        jointsToRemove = ["thigh_twist_01_r", "thigh_twist_02_r"]
+        jointToTransferTo = "thigh_r"
+        transfer.LOD_transferWeights(meshes, jointsToRemove, jointToTransferTo)
+        jointsToDelete.extend(jointsToRemove)
+
+        jointsToRemove = ["calf_twist_01_l"]
+        jointToTransferTo = "calf_l"
+        transfer.LOD_transferWeights(meshes, jointsToRemove, jointToTransferTo)
+        jointsToDelete.extend(jointsToRemove)
+
+        jointsToRemove = ["calf_twist_01_r"]
+        jointToTransferTo = "calf_r"
+        transfer.LOD_transferWeights(meshes, jointsToRemove, jointToTransferTo)
+        jointsToDelete.extend(jointsToRemove)
+
+        # Delete the joints we no longer want.
+        for one in jointsToDelete:
+            if cmds.objExists(one):
+                cmds.delete(one)
+                
+        # Export the fbx file.
+        export.exportLOD("RiftMage_LOD2.fbx", "LOD2")
+
+        # Re-open the file without saving.
+        fullPath = cmds.file(q = True, sceneName = True)
+        cmds.file(fullPath, open=True, f=True)
+
+        print "LOD2 Exported Successfully!!!!!!!!!!!!!!!"
+
+RiftMageLOD2()
+
+# RiftMage LOD3 -------------------------------------------------------------------------------------------------------------
+def RiftMageLOD3():
+    if saveResult == "Yes":
+
+        cmds.file(save=True)
+    
+        meshes = ["arms_geo_LOD3", "beard_geo_LOD3", "brooch_geo_LOD3", "buttstrap_geo_LOD3", "chest_buttons_geo_LOD3", "chest_flaps_geo_LOD3", "codpiece_buttons_geo_LOD3", "codpiece_geo_LOD3", "collar_geo_LOD3", "cowl_geo_LOD3", "elbowpad_geo_LOD3", "eye_tear_duct_geo_LOD3", "eye_wetness_geo_LOD3", "eyeballs_geo_LOD3", "eyelash_geo_LOD3", "eyeshadow_geo_LOD3", "flappendant_geo_LOD3", "forearm_pads_geo_LOD3", "head_geo_LOD3", "hip_parts_geo_LOD3", "hood_geo_LOD3", "inner_ring_geo_LOD3", "kneepads_geo_LOD3", "legplate_geo_LOD3", "legs_geo_LOD3", "midflaps_geo_LOD3", "outer_ring_bott_geo_LOD3", "outer_ring_inside_geo_LOD3", "outer_ring_knobs_geo_LOD3", "outer_ring_top_geo_LOD3", "portal_geo_LOD3", "sideflaps_geo_LOD3", "teeth_lower_geo_LOD3", "teeth_upper_geo_LOD3", "tongue_geo_LOD3", "torso_detail_geo_LOD3", "torso_geo_LOD3", "wrist_armor_geo_LOD3"]
+    
+        # Export a skin weights file for each mesh we plan to operate on.
+        cmds.select(meshes)
+        menu.exportSelectedSkinWeights(True)
+
+        # Set the pose if needed
+        cmds.setAttr("thumb_03_l.rotateZ", 23.02290762)
+        cmds.setAttr("thumb_02_l.rotateZ", 10.06960964)
+        cmds.setAttr("ring_03_l.rotateZ", 29.40776524)
+        cmds.setAttr("ring_02_l.rotateZ", 29.40776524)
+        cmds.setAttr("pinky_03_l.rotateZ", 39.74899767)
+        cmds.setAttr("pinky_02_l.rotateZ", 39.74899767)
+        cmds.setAttr("middle_03_l.rotateZ", 23.93256607)
+        cmds.setAttr("middle_02_l.rotateZ", 23.93256607)
+        cmds.setAttr("index_03_l.rotateZ", 16.68034218)
+        cmds.setAttr("index_02_l.rotateZ", 16.68034218)
+
+        cmds.setAttr("thumb_03_r.rotateZ", 23.02290762)
+        cmds.setAttr("thumb_02_r.rotateZ", 10.06960964)
+        cmds.setAttr("ring_03_r.rotateZ", 29.40776524)
+        cmds.setAttr("ring_02_r.rotateZ", 29.40776524)
+        cmds.setAttr("pinky_03_r.rotateZ", 39.74899767)
+        cmds.setAttr("pinky_02_r.rotateZ", 39.74899767)
+        cmds.setAttr("middle_03_r.rotateZ", 23.93256607)
+        cmds.setAttr("middle_02_r.rotateZ", 23.93256607)
+        cmds.setAttr("index_03_r.rotateZ", 16.68034218)
+        cmds.setAttr("index_02_r.rotateZ", 16.68034218)
+
+        # Delete the history on the meshes so that they get baked down into their posed positions
+        cmds.delete(meshes, ch=True)
+
+        # Move the skeleton back to its bindpose (This script REQUIRES the bind pose to be the pose that we want to export the skeleton to.  This is the case for all known characters created with A.R.T.)
+        cmds.dagPose("root", bp=True, r=True) 
+
+        # Re-import the skinning we exported before onto the mesh so thats its bound in its posed position.
+        cmds.select(meshes)
+        impw.importMultipleSkinWeights("", True)
+        
+        jointsToDelete = []
+
+        # remove unneeded joints
+        jointsToRemove = ["hips_cloth_main_l", "hips_cloth_bk_l_01", "hips_cloth_bk_l_02", "hips_cloth_fr_l_01", "hips_cloth_fr_l_02", "hips_cloth_main_r", "hips_cloth_bk_r_01", "hips_cloth_bk_r_02", "hips_cloth_fr_r_01", "hips_cloth_fr_r_02", "loin_cloth_bk_01", "loin_cloth_bk_02", "loin_cloth_bk_03", "loin_cloth_fr_01", "loin_cloth_fr_02", "loin_cloth_fr_03"]
+        jointToTransferTo = "pelvis"
+        transfer.LOD_transferWeights(meshes, jointsToRemove, jointToTransferTo)
+        jointsToDelete.extend(jointsToRemove)
+
+        jointsToRemove = ["chest_cloth_01", "chest_cloth_02", "chest_cloth_03", "chest_cloth_04", "cloth_chest_bk_l", "cloth_chest_bk_r", "cloth_chest_fr_l", "cloth_chest_fr_r"]
+        jointToTransferTo = "spine_02"
+        transfer.LOD_transferWeights(meshes, jointsToRemove, jointToTransferTo)
+        jointsToDelete.extend(jointsToRemove)
+
+        jointsToRemove = ["cowl_bk", "cowl_l", "cowl_r"]
+        jointToTransferTo = "spine_03"
+        transfer.LOD_transferWeights(meshes, jointsToRemove, jointToTransferTo)
+        jointsToDelete.extend(jointsToRemove)
+
+        jointsToRemove = ["upperarm_twist_01_l", "upperarm_twist_02_l"]
+        jointToTransferTo = "upperarm_l"
+        transfer.LOD_transferWeights(meshes, jointsToRemove, jointToTransferTo)
+        jointsToDelete.extend(jointsToRemove)
+
+        jointsToRemove = ["upperarm_twist_01_r", "upperarm_twist_02_r"]
+        jointToTransferTo = "upperarm_r"
+        transfer.LOD_transferWeights(meshes, jointsToRemove, jointToTransferTo)
+        jointsToDelete.extend(jointsToRemove)
+
+        jointsToRemove = ["lowerarm_twist_01_l", "lowerarm_twist_02_l"]
+        jointToTransferTo = "lowerarm_l"
+        transfer.LOD_transferWeights(meshes, jointsToRemove, jointToTransferTo)
+        jointsToDelete.extend(jointsToRemove)
+
+        jointsToRemove = ["lowerarm_twist_01_r", "lowerarm_twist_02_r"]
+        jointToTransferTo = "lowerarm_r"
+        transfer.LOD_transferWeights(meshes, jointsToRemove, jointToTransferTo)
+        jointsToDelete.extend(jointsToRemove)
+
+        jointsToRemove = ["index_02_l", "index_03_l"]
+        jointToTransferTo = "index_01_l"
+        transfer.LOD_transferWeights(meshes, jointsToRemove, jointToTransferTo)
+        jointsToDelete.extend(jointsToRemove)
+
+        jointsToRemove = ["middle_02_l", "middle_03_l"]
+        jointToTransferTo = "middle_01_l"
+        transfer.LOD_transferWeights(meshes, jointsToRemove, jointToTransferTo)
+        jointsToDelete.extend(jointsToRemove)
+
+        jointsToRemove = ["pinky_02_l", "pinky_03_l"]
+        jointToTransferTo = "pinky_03_l"
+        transfer.LOD_transferWeights(meshes, jointsToRemove, jointToTransferTo)
+        jointsToDelete.extend(jointsToRemove)
+
+        jointsToRemove = ["ring_02_l", "ring_03_l"]
+        jointToTransferTo = "ring_01_l"
+        transfer.LOD_transferWeights(meshes, jointsToRemove, jointToTransferTo)
+        jointsToDelete.extend(jointsToRemove)
+
+        jointsToRemove = ["thumb_02_l", "thumb_03_l"]
+        jointToTransferTo = "thumb_01_l"
+        transfer.LOD_transferWeights(meshes, jointsToRemove, jointToTransferTo)
+        jointsToDelete.extend(jointsToRemove)
+
+        jointsToRemove = ["index_02_r", "index_03_r"]
+        jointToTransferTo = "index_01_r"
+        transfer.LOD_transferWeights(meshes, jointsToRemove, jointToTransferTo)
+        jointsToDelete.extend(jointsToRemove)
+
+        jointsToRemove = ["middle_02_r", "middle_03_r"]
+        jointToTransferTo = "middle_01_r"
+        transfer.LOD_transferWeights(meshes, jointsToRemove, jointToTransferTo)
+        jointsToDelete.extend(jointsToRemove)
+
+        jointsToRemove = ["pinky_02_r", "pinky_03_r"]
+        jointToTransferTo = "pinky_03_r"
+        transfer.LOD_transferWeights(meshes, jointsToRemove, jointToTransferTo)
+        jointsToDelete.extend(jointsToRemove)
+
+        jointsToRemove = ["ring_02_r", "ring_03_r"]
+        jointToTransferTo = "ring_01_r"
+        transfer.LOD_transferWeights(meshes, jointsToRemove, jointToTransferTo)
+        jointsToDelete.extend(jointsToRemove)
+
+        jointsToRemove = ["thumb_02_r", "thumb_03_r"]
+        jointToTransferTo = "thumb_01_r"
+        transfer.LOD_transferWeights(meshes, jointsToRemove, jointToTransferTo)
+        jointsToDelete.extend(jointsToRemove)
+        
+        jointsToRemove = ["eyeball_l", "eyeball_r", "eyebrow_inner_l", "eyebrow_inner_r", "eyebrow_mid", "eyebrow_outer_l", "eyebrow_outer_r", "eyelid_bott_l", "eyelid_bott_r", "eyelid_top_l", "eyelid_top_r", "hood_l", "hood_r", "jaw"]
+        jointToTransferTo = "head"
+        transfer.LOD_transferWeights(meshes, jointsToRemove, jointToTransferTo)
+        jointsToDelete.extend(jointsToRemove)
+
+        jointsToRemove = ["outer_ring_A_l", "outer_ring_A_piston01_l", "outer_ring_A_piston02_l", "outer_ring_B_l", "outer_ring_B_piston01_l", "outer_ring_C_l", "outer_ring_A_r", "outer_ring_A_piston01_r", "outer_ring_A_piston02_r", "outer_ring_B_r", "outer_ring_B_piston01_r", "outer_ring_C_r"]
+        jointToTransferTo = "outer_ring_base"
+        transfer.LOD_transferWeights(meshes, jointsToRemove, jointToTransferTo)
+        jointsToDelete.extend(jointsToRemove)
+
+        jointsToRemove = ["portal"]
+        jointToTransferTo = "portal_ring"
+        transfer.LOD_transferWeights(meshes, jointsToRemove, jointToTransferTo)
+        jointsToDelete.extend(jointsToRemove)
+
+        jointsToRemove = ["thigh_twist_01_l", "thigh_twist_02_l"]
+        jointToTransferTo = "thigh_l"
+        transfer.LOD_transferWeights(meshes, jointsToRemove, jointToTransferTo)
+        jointsToDelete.extend(jointsToRemove)
+
+        jointsToRemove = ["thigh_twist_01_r", "thigh_twist_02_r"]
+        jointToTransferTo = "thigh_r"
+        transfer.LOD_transferWeights(meshes, jointsToRemove, jointToTransferTo)
+        jointsToDelete.extend(jointsToRemove)
+
+        jointsToRemove = ["calf_twist_01_l"]
+        jointToTransferTo = "calf_l"
+        transfer.LOD_transferWeights(meshes, jointsToRemove, jointToTransferTo)
+        jointsToDelete.extend(jointsToRemove)
+
+        jointsToRemove = ["calf_twist_01_r"]
+        jointToTransferTo = "calf_r"
+        transfer.LOD_transferWeights(meshes, jointsToRemove, jointToTransferTo)
+        jointsToDelete.extend(jointsToRemove)
+
+        # Delete the joints we no longer want.
+        for one in jointsToDelete:
+            if cmds.objExists(one):
+                cmds.delete(one)
+                
+        # Export the fbx file.
+        export.exportLOD("RiftMage_LOD3.fbx", "LOD3")
+        
+        # Re-open the file without saving.
+        fullPath = cmds.file(q = True, sceneName = True)
+        cmds.file(fullPath, open=True, f=True)
+
+        print "LOD3 Exported Successfully!!!!!!!!!!!!!!!"
+
+RiftMageLOD3()
+
+# RiftMage LOD4 -------------------------------------------------------------------------------------------------------------
+def RiftMageLOD4():
+    if saveResult == "Yes":
+
+        cmds.file(save=True)
+    
+        meshes = ["arms_geo_LOD4", "beard_geo_LOD4", "brooch_geo_LOD4", "buttstrap_geo_LOD4", "chest_buttons_geo_LOD4", "chest_flaps_geo_LOD4", "codpiece_buttons_geo_LOD4", "codpiece_geo_LOD4", "collar_geo_LOD4", "cowl_geo_LOD4", "elbowpad_geo_LOD4", "eye_tear_duct_geo_LOD4", "eye_wetness_geo_LOD4", "eyeballs_geo_LOD4", "eyelash_geo_LOD4", "eyeshadow_geo_LOD4", "flappendant_geo_LOD4", "forearm_pads_geo_LOD4", "head_geo_LOD4", "hip_parts_geo_LOD4", "hood_geo_LOD4", "inner_ring_geo_LOD4", "kneepads_geo_LOD4", "legplate_geo_LOD4", "legs_geo_LOD4", "midflaps_geo_LOD4", "outer_ring_bott_geo_LOD4", "outer_ring_inside_geo_LOD4", "outer_ring_knobs_geo_LOD4", "outer_ring_top_geo_LOD4", "portal_geo_LOD4", "sideflaps_geo_LOD4", "teeth_lower_geo_LOD4", "teeth_upper_geo_LOD4", "tongue_geo_LOD4", "torso_detail_geo_LOD4", "torso_geo_LOD4", "wrist_armor_geo_LOD4"]
+
+        # Export a skin weights file for each mesh we plan to operate on.
+        cmds.select(meshes)
+        menu.exportSelectedSkinWeights(True)
+
+        # Set the pose if needed
+        cmds.setAttr("thumb_03_l.rotateZ", 23.02290762)
+        cmds.setAttr("thumb_02_l.rotateZ", 10.06960964)
+        cmds.setAttr("thumb_01_l.rotateZ", 8.338113921)
+        cmds.setAttr("thumb_01_l.rotateY", 3.002087702)
+        cmds.setAttr("thumb_01_l.rotateX", -3.231699475)
+        cmds.setAttr("ring_03_l.rotateZ", 29.40776524)
+        cmds.setAttr("ring_02_l.rotateZ", 29.40776524)
+        cmds.setAttr("ring_01_l.rotateZ", 29.40776524)
+        cmds.setAttr("pinky_03_l.rotateZ", 39.74899767)
+        cmds.setAttr("pinky_02_l.rotateZ", 39.74899767)
+        cmds.setAttr("pinky_01_l.rotateZ", 39.74899767)
+        cmds.setAttr("middle_03_l.rotateZ", 23.93256607)
+        cmds.setAttr("middle_02_l.rotateZ", 23.93256607)
+        cmds.setAttr("middle_01_l.rotateZ", 23.93256607)
+        cmds.setAttr("index_03_l.rotateZ", 16.68034218)
+        cmds.setAttr("index_02_l.rotateZ", 16.68034218)
+        cmds.setAttr("index_01_l.rotateZ", 16.68034218)
+
+        cmds.setAttr("thumb_03_r.rotateZ", 23.02290762)
+        cmds.setAttr("thumb_02_r.rotateZ", 10.06960964)
+        cmds.setAttr("thumb_01_r.rotateZ", 8.338113921)
+        cmds.setAttr("thumb_01_r.rotateY", 3.002087702)
+        cmds.setAttr("thumb_01_r.rotateX", -3.231699475)
+        cmds.setAttr("ring_03_r.rotateZ", 29.40776524)
+        cmds.setAttr("ring_02_r.rotateZ", 29.40776524)
+        cmds.setAttr("ring_01_r.rotateZ", 29.40776524)
+        cmds.setAttr("pinky_03_r.rotateZ", 39.74899767)
+        cmds.setAttr("pinky_02_r.rotateZ", 39.74899767)
+        cmds.setAttr("pinky_01_r.rotateZ", 39.74899767)
+        cmds.setAttr("middle_03_r.rotateZ", 23.93256607)
+        cmds.setAttr("middle_02_r.rotateZ", 23.93256607)
+        cmds.setAttr("middle_01_r.rotateZ", 23.93256607)
+        cmds.setAttr("index_03_r.rotateZ", 16.68034218)
+        cmds.setAttr("index_02_r.rotateZ", 16.68034218)
+        cmds.setAttr("index_01_r.rotateZ", 16.68034218)
+
+        # Delete the history on the meshes so that they get baked down into their posed positions
+        cmds.delete(meshes, ch=True)
+
+        # Move the skeleton back to its bindpose (This script REQUIRES the bind pose to be the pose that we want to export the skeleton to.  This is the case for all known characters created with A.R.T.)
+        cmds.dagPose("root", bp=True, r=True) 
+
+        # Re-import the skinning we exported before onto the mesh so thats its bound in its posed position.
+        cmds.select(meshes)
+        impw.importMultipleSkinWeights("", True)
+        
+        jointsToDelete = []
+
+        # remove unneeded joints
+        jointsToRemove = ["hips_cloth_main_l", "hips_cloth_bk_l_01", "hips_cloth_bk_l_02", "hips_cloth_fr_l_01", "hips_cloth_fr_l_02", "hips_cloth_main_r", "hips_cloth_bk_r_01", "hips_cloth_bk_r_02", "hips_cloth_fr_r_01", "hips_cloth_fr_r_02", "loin_cloth_bk_01", "loin_cloth_bk_02", "loin_cloth_bk_03", "loin_cloth_fr_01", "loin_cloth_fr_02", "loin_cloth_fr_03"]
+        jointToTransferTo = "pelvis"
+        transfer.LOD_transferWeights(meshes, jointsToRemove, jointToTransferTo)
+        jointsToDelete.extend(jointsToRemove)
+
+        jointsToRemove = ["chest_cloth_01", "chest_cloth_02", "chest_cloth_03", "chest_cloth_04", "cloth_chest_bk_l", "cloth_chest_bk_r", "cloth_chest_fr_l", "cloth_chest_fr_r"]
+        jointToTransferTo = "spine_02"
+        transfer.LOD_transferWeights(meshes, jointsToRemove, jointToTransferTo)
+        jointsToDelete.extend(jointsToRemove)
+
+        jointsToRemove = ["cowl_bk", "cowl_l", "cowl_r"]
+        jointToTransferTo = "spine_03"
+        transfer.LOD_transferWeights(meshes, jointsToRemove, jointToTransferTo)
+        jointsToDelete.extend(jointsToRemove)
+
+        jointsToRemove = ["upperarm_twist_01_l", "upperarm_twist_02_l"]
+        jointToTransferTo = "upperarm_l"
+        transfer.LOD_transferWeights(meshes, jointsToRemove, jointToTransferTo)
+        jointsToDelete.extend(jointsToRemove)
+
+        jointsToRemove = ["upperarm_twist_01_r", "upperarm_twist_02_r"]
+        jointToTransferTo = "upperarm_r"
+        transfer.LOD_transferWeights(meshes, jointsToRemove, jointToTransferTo)
+        jointsToDelete.extend(jointsToRemove)
+
+        jointsToRemove = ["lowerarm_twist_01_l", "lowerarm_twist_02_l"]
+        jointToTransferTo = "lowerarm_l"
+        transfer.LOD_transferWeights(meshes, jointsToRemove, jointToTransferTo)
+        jointsToDelete.extend(jointsToRemove)
+
+        jointsToRemove = ["lowerarm_twist_01_r", "lowerarm_twist_02_r"]
+        jointToTransferTo = "lowerarm_r"
+        transfer.LOD_transferWeights(meshes, jointsToRemove, jointToTransferTo)
+        jointsToDelete.extend(jointsToRemove)
+
+        jointsToRemove = ["index_01_l", "index_02_l", "index_03_l", "middle_01_l", "middle_02_l", "middle_03_l", "pinky_01_l", "pinky_02_l", "pinky_03_l", "ring_01_l", "ring_02_l", "ring_03_l", "thumb_01_l", "thumb_02_l", "thumb_03_l"]
+        jointToTransferTo = "hand_l"
+        transfer.LOD_transferWeights(meshes, jointsToRemove, jointToTransferTo)
+        jointsToDelete.extend(jointsToRemove)
+
+        jointsToRemove = ["index_01_r", "index_02_r", "index_03_r", "middle_01_r", "middle_02_r", "middle_03_r", "pinky_01_r", "pinky_02_r", "pinky_03_r", "ring_01_r", "ring_02_r", "ring_03_r", "thumb_01_r", "thumb_02_r", "thumb_03_r"]
+        jointToTransferTo = "hand_r"
+        transfer.LOD_transferWeights(meshes, jointsToRemove, jointToTransferTo)
+        jointsToDelete.extend(jointsToRemove)
+
+        jointsToRemove = ["eyeball_l", "eyeball_r", "eyebrow_inner_l", "eyebrow_inner_r", "eyebrow_mid", "eyebrow_outer_l", "eyebrow_outer_r", "eyelid_bott_l", "eyelid_bott_r", "eyelid_top_l", "eyelid_top_r", "hood_l", "hood_r", "jaw"]
+        jointToTransferTo = "head"
+        transfer.LOD_transferWeights(meshes, jointsToRemove, jointToTransferTo)
+        jointsToDelete.extend(jointsToRemove)
+
+        jointsToRemove = ["outer_ring_A_l", "outer_ring_A_piston01_l", "outer_ring_A_piston02_l", "outer_ring_B_l", "outer_ring_B_piston01_l", "outer_ring_C_l", "outer_ring_A_r", "outer_ring_A_piston01_r", "outer_ring_A_piston02_r", "outer_ring_B_r", "outer_ring_B_piston01_r", "outer_ring_C_r"]
+        jointToTransferTo = "outer_ring_base"
+        transfer.LOD_transferWeights(meshes, jointsToRemove, jointToTransferTo)
+        jointsToDelete.extend(jointsToRemove)
+
+        jointsToRemove = ["portal"]
+        jointToTransferTo = "portal_ring"
+        transfer.LOD_transferWeights(meshes, jointsToRemove, jointToTransferTo)
+        jointsToDelete.extend(jointsToRemove)
+
+        jointsToRemove = ["thigh_twist_01_l", "thigh_twist_02_l"]
+        jointToTransferTo = "thigh_l"
+        transfer.LOD_transferWeights(meshes, jointsToRemove, jointToTransferTo)
+        jointsToDelete.extend(jointsToRemove)
+
+        jointsToRemove = ["thigh_twist_01_r", "thigh_twist_02_r"]
+        jointToTransferTo = "thigh_r"
+        transfer.LOD_transferWeights(meshes, jointsToRemove, jointToTransferTo)
+        jointsToDelete.extend(jointsToRemove)
+
+        jointsToRemove = ["calf_twist_01_l"]
+        jointToTransferTo = "calf_l"
+        transfer.LOD_transferWeights(meshes, jointsToRemove, jointToTransferTo)
+        jointsToDelete.extend(jointsToRemove)
+
+        jointsToRemove = ["calf_twist_01_r"]
+        jointToTransferTo = "calf_r"
+        transfer.LOD_transferWeights(meshes, jointsToRemove, jointToTransferTo)
+        jointsToDelete.extend(jointsToRemove)
+
+        jointsToRemove = ["ball_l"]
+        jointToTransferTo = "foot_l"
+        transfer.LOD_transferWeights(meshes, jointsToRemove, jointToTransferTo)
+        jointsToDelete.extend(jointsToRemove)
+
+        jointsToRemove = ["ball_r"]
+        jointToTransferTo = "foot_r"
+        transfer.LOD_transferWeights(meshes, jointsToRemove, jointToTransferTo)
+        jointsToDelete.extend(jointsToRemove)
+
+        # Delete the joints we no longer want.
+        for one in jointsToDelete:
+            if cmds.objExists(one):
+                cmds.delete(one)
+                
+        # Export the fbx file.
+        export.exportLOD("RiftMage_LOD4.fbx", "LOD4")
+        
+        # Re-open the file without saving.
+        fullPath = cmds.file(q = True, sceneName = True)
+        cmds.file(fullPath, open=True, f=True)
+
+        print "LOD4 Exported Successfully!!!!!!!!!!!!!!!"
+
+RiftMageLOD4()
+
+
diff --git a/Engine/Extras/Maya_AnimationRiggingTools/ARTv1/General/ART/Projects/Orion/Scripts/RiftMage_Post_Script.py b/Engine/Extras/Maya_AnimationRiggingTools/ARTv1/General/ART/Projects/Orion/Scripts/RiftMage_Post_Script.py
new file mode 100644
--- /dev/null
+++ b/Engine/Extras/Maya_AnimationRiggingTools/ARTv1/General/ART/Projects/Orion/Scripts/RiftMage_Post_Script.py
@@ -0,0 +1,289 @@
+import maya.cmds as cmds
+import maya.mel as mel
+
+# ------------------------------------------------------------------
+# scale all of the controls to be more visible with the mesh.
+def scaleControls():
+    controls = ["cowl_l_anim", "cowl_r_anim", "cowl_bk_anim", "outer_ring_A_piston01_l_anim", "outer_ring_A_piston01_r_anim", "outer_ring_A_piston02_l_anim", "outer_ring_A_piston02_r_anim", "outer_ring_B_piston01_l_anim", "outer_ring_B_piston01_r_anim", "hood_l_anim", "hood_r_anim", "outer_ring_A_l_anim", "outer_ring_A_r_anim", "outer_ring_B_l_anim", "outer_ring_B_r_anim", "outer_ring_C_l_anim", "outer_ring_C_r_anim"]
+    for control in controls:
+        cmds.select(control+".cv[0:32]")
+        cmds.scale(0.2, 0.2, 0.2, r=True)
+            
+    for side in ["_l", "_r"]:
+        cmds.select(["eyebrow_inner"+side+"_anim.cv[0:32]"])
+        cmds.scale(0.117744, 0.117744, 0.117744, r=True, ocp = False)
+        cmds.move(0, 0, 0, r=True)
+
+        cmds.select(["eyebrow_outer"+side+"_anim.cv[0:32]"])
+        cmds.scale(0.117744, 0.117744, 0.117744, r=True, ocp = False)
+        cmds.move(0, 0, 0, r=True)
+
+        cmds.select(["eyeball"+side+"_anim.cv[0:32]"])
+        cmds.scale(0.1, 0.1, 0.05, r=True, ocp = False)
+        if side == "_l":
+            cmds.move(0, -3.757386, 0, r=True, os=True)
+        else:
+            cmds.move(0, -3.757386, 0, r=True, os=True)
+
+        cmds.select(["eyelid_bott"+side+"_anim.cv[0:32]"])
+        cmds.scale(0.1, 0.1, 0.01, r=True, ocp = False)
+        if side == "_l":
+            cmds.move(0, -3.757386, 0.5, r=True, os=True)
+        else:
+            cmds.move(0, 3.757386, 0.5, r=True, os=True)
+
+        cmds.select(["eyelid_top"+side+"_anim.cv[0:32]"])
+        cmds.scale(0.1, 0.1, 0.01, r=True, ocp = False)
+        if side == "_l":
+            cmds.move(0, -3.757386, 0.5, r=True, os=True)
+        else:
+            cmds.move(0, 3.757386, 0.5, r=True, os=True)
+
+    cmds.select(["jaw_anim.cv[0:32]"])
+    cmds.scale(0.294922, 0.294922, 0.294922, r = True, ocp = True)
+    cmds.move(0, -8, -6, r = True)
+
+    cmds.select(["eyebrow_mid_anim.cv[0:32]"])
+    cmds.scale(0.132495, 0.132495, 0.132495, r = True, ocp = True)
+
+    cmds.select(["portal_anim.cv[0:32]"])
+    cmds.scale(1.5, 0.5, 1.5, r = True, ocp = True)
+
+    cmds.select(["outer_ring_base_anim.cv[0:32]"])
+    cmds.scale(4.5, 0.001, 4.5, r = True, ocp = True)
+
+    cmds.select(["portal_ring_anim.cv[0:32]"])
+    cmds.scale(2.5, 0.001, 2.5, r = True, ocp = True)
+
+    cmds.select(["outer_ring_A_l_anim.cv[0:32]"])
+    cmds.move(4, 0, 4, r = True)
+
+    cmds.select(["outer_ring_A_r_anim.cv[0:32]"])
+    cmds.move(-4, 0, 4, r = True)
+
+    cmds.select(["outer_ring_B_l_anim.cv[0:32]"])
+    cmds.move(4, 0, 0, r = True)
+
+    cmds.select(["outer_ring_B_r_anim.cv[0:32]"])
+    cmds.move(-4, 0, 0, r = True)
+
+    cmds.select(["outer_ring_C_l_anim.cv[0:32]"])
+    cmds.move(4, 0, -4, r = True)
+
+    cmds.select(["outer_ring_C_r_anim.cv[0:32]"])
+    cmds.move(-4, 0, -4, r = True)
+
+    cmds.select(["ik_foot_anim_l.cv[0:17]"])
+    cmds.move(0, 3, -2.22651, r = True)
+
+    cmds.select(["ik_foot_anim_r.cv[0:17]"])
+    cmds.move(0, 3, -2.22651, r = True)
+
+    cmds.select(["cowl_bk_anim.cv[0:32]"])
+    cmds.move(0, 0, 6.5, r = True)
+
+    cmds.select(["cowl_l_anim.cv[0:32]"])
+    cmds.move(0, 0, 2, r = True)
+
+    cmds.select(["cowl_r_anim.cv[0:32]"])
+    cmds.move(0, 0, 2, r = True)
+
+    cmds.select(None)
+
+scaleControls()
+
+# ------------------------------------------------------------------
+# Build the rig for the rings and connect them to the arm rig.
+def buildRingRig():
+    import Tools.ART_addSpaceSwitching as space
+    reload(space)
+
+    #space.addSpaceSwitching("portal_ring_anim", "outer_ring_base_anim", True)
+    
+    controlsList = {"outer_ring_A_l":"outer_ring_base_anim", "outer_ring_B_l":"outer_ring_A_l_anim", "outer_ring_C_l":"outer_ring_B_l_anim", "outer_ring_A_piston01_l":"outer_ring_A_l_anim", "outer_ring_A_piston02_l":"outer_ring_A_l_anim", "outer_ring_B_piston01_l":"outer_ring_B_l_anim", "outer_ring_A_r":"outer_ring_base_anim", "outer_ring_B_r":"outer_ring_A_r_anim", "outer_ring_C_r":"outer_ring_B_r_anim", "outer_ring_A_piston01_r":"outer_ring_A_r_anim", "outer_ring_A_piston02_r":"outer_ring_A_r_anim", "outer_ring_B_piston01_r":"outer_ring_B_r_anim", "portal":"portal_ring_anim"}
+    for one in controlsList:
+        # Re-create the scale constraint between the anim control and the driver joint.
+        sourceConn = cmds.listConnections("driver_"+one+".s", c=True, p=True)
+        cmds.disconnectAttr(sourceConn[1], sourceConn[0])
+        cmds.scaleConstraint(one+"_anim", "driver_"+one, mo=True)
+        
+        # scale constraint the control to the leaf control we want to scale it.
+        cmds.scaleConstraint(controlsList[one], one+"_parent_grp", mo=True)
+
+        # set the joints (driver and deform) to have segment scale compensate off.
+        cmds.setAttr("driver_"+one+".segmentScaleCompensate", 0)
+        cmds.setAttr(one+".segmentScaleCompensate", 0)
+            
+    cmds.select("outer_ring_base_anim")
+
+buildRingRig()
+
+# ------------------------------------------------------------------
+# Delete this once the animators lock down how big they want the ring to be.
+def TEMPRINGMORPH():
+    cmds.addAttr("outer_ring_base_anim", ln="Thickness", keyable=True)
+    cmds.connectAttr("outer_ring_base_anim.Thickness", "blendShape1.outer_ring_top_geo_scaled", f=True)
+    cmds.connectAttr("outer_ring_base_anim.Thickness", "blendShape2.outer_ring_bott_geo_scaled", f=True)
+    cmds.connectAttr("outer_ring_base_anim.Thickness", "blendShape3.outer_ring_inside_geo_scaled", f=True)
+#TEMPRINGMORPH()
+
+
+# ------------------------------------------------------------------
+def invertEyelids():
+    sides = ["_l", "_r"]
+    for side in sides:
+        cmds.delete("driver_eyelid_bott"+side+"_parentConstraint1")
+        cmds.rotate(0, -180, 0, "eyelid_bott"+side+"_anim_grp", r=True, os=True)
+        cmds.parentConstraint("eyelid_bott"+side+"_anim", "driver_eyelid_bott"+side, mo=True)
+invertEyelids()
+
+def invertRightEyeball():
+    cmds.delete("driver_eyeball_r_parentConstraint1")
+    cmds.rotate(0, -180, -180, "eyeball_r_anim_grp", r=True, os=True)
+    cmds.parentConstraint("eyeball_r_anim", "driver_eyeball_r", mo=True)
+invertRightEyeball()
+
+# ------------------------------------------------------------------
+def dynamicsOff():
+    for one in ["cloth_chest_bk_l_anim", "cloth_chest_bk_r_anim", "cloth_chest_fr_l_anim", "cloth_chest_fr_r_anim", "chest_cloth_dyn_anim", "hips_cloth_bk_l_dyn_anim", "hips_cloth_bk_r_dyn_anim", "hips_cloth_fr_l_dyn_anim", "hips_cloth_fr_r_dyn_anim", "loin_cloth_bk_dyn_anim", "loin_cloth_fr_dyn_anim"]:
+        cmds.setAttr(one+".chainStartEnvelope", 0)
+        cmds.setAttr(one+".chainStartEnvelope", l=True)
+dynamicsOff()
+
+# ------------------------------------------------------------------
+def createGameCam():
+    gameCam = cmds.camera(centerOfInterest=5, focalLength=18, lensSqueezeRatio=1, cameraScale=1, horizontalFilmAperture=1.41732, horizontalFilmOffset=0, verticalFilmAperture=0.94488, verticalFilmOffset=0, filmFit="Overscan", overscan=1.05, motionBlur=0, shutterAngle=144, nearClipPlane=1.0, farClipPlane=10000, orthographic=0, orthographicWidth=30, panZoomEnabled=0, horizontalPan=0, verticalPan=0, zoom=1, displayFilmGate=True)[0]
+    print gameCam
+    gameCam = cmds.rename(gameCam, "GameCam")
+    cmds.camera(gameCam, e=True, horizontalFilmAperture=1.6724376, focalLength=21.2399857011)
+    cmds.setAttr(gameCam+".displayGateMaskOpacity", 1)
+    cmds.setAttr(gameCam+".displayGateMaskColor", 0, 0, 0, type="double3")
+    cmds.setAttr(gameCam+".r", 90, 0, 180, type="double3")
+    
+    # Create the control
+    control = cmds.circle(name=gameCam+"_Ctrl", ch=True, o=True, r=15)[0]
+    cmds.makeIdentity(control, a=True, t=True, r=True, s=True)
+    # Create groups for the sdk and the constraints
+    const_grp = cmds.group(control, name=control+"_grp")
+    if cmds.objExists("rig_grp"):
+        cmds.parent(const_grp, "rig_grp")
+    # Parent Constrain the cam to the control
+    cmds.parentConstraint(control, gameCam, mo=True)
+    
+    cmds.connectAttr(control+".visibility", gameCam+".visibility", f=True)
+    cmds.setAttr(const_grp+".v", 0)
+    
+    cmds.setAttr(const_grp+".t", 0, 500, 200, type="double3")
+createGameCam()
+
+# ------------------------------------------------------------------
+def addCapsule(radius, halfHeight):
+    trueHalfHeight = (halfHeight*2)-(radius*2)
+    capsule = cmds.polyCylinder(name="collision_capsule", r=radius, h=trueHalfHeight, sx=18, sy=1, sz=5, ax=[0,0,1], rcp=True, cuv=3, ch=True)
+    cmds.setAttr(capsule[0]+".tz", halfHeight+2.2)
+    cmds.setAttr(capsule[0]+".rz", 90)
+
+    cmds.addAttr(capsule[0], ln="Radius", dv=radius, keyable=True)
+    cmds.addAttr(capsule[0], ln="HalfHeight", dv=halfHeight, keyable=True)
+
+    expr = capsule[1]+".radius = "+capsule[0]+".Radius;\n"+capsule[1]+".height = ("+capsule[0]+".HalfHeight*2) - ("+capsule[0]+".Radius*2);"
+    cmds.expression(s=expr, o=capsule[1], ae=1, uc="all")
+
+    cmds.setAttr(capsule[0]+".overrideEnabled", 1)
+    cmds.setAttr(capsule[0]+".overrideColor", 18)
+    cmds.setAttr(capsule[0]+".overrideShading", 0)
+
+    cmds.addAttr(capsule[0], ln="Shaded", dv=0, min=0, max=1, keyable=True)
+    cmds.connectAttr(capsule[0]+".Shaded", capsule[0]+".overrideShading")
+
+    cmds.parentConstraint("root", capsule[0], mo=True)
+    cmds.select(capsule[0], r=True)
+    capsuleDisplayLayer = cmds.createDisplayLayer(nr=True, name="Collision Capsule")
+    cmds.setAttr(capsuleDisplayLayer+".displayType", 2)
+    cmds.setAttr(capsuleDisplayLayer+".v", 0)
+    
+    for attr in ["tx", "ty", "tz", "rx", "ry", "rz", "sx", "sy", "sz"]:
+        cmds.setAttr(capsule[0]+"."+attr, lock=True, keyable=False, channelBox=False)
+
+    cmds.select(clear=True)
+addCapsule(42, 86)
+
+# ------------------------------------------------------------------
+# add ik bones Function that is called later after we open the export file.
+def addIKBones():
+    try:
+		# create the joints
+		cmds.select(clear = True)
+		ikFootRoot = cmds.joint(name = "ik_foot_root")
+		cmds.select(clear = True)
+		
+		cmds.select(clear = True)
+		ikFootLeft = cmds.joint(name = "ik_foot_l")
+		cmds.select(clear = True)
+		
+		cmds.select(clear = True)
+		ikFootRight = cmds.joint(name = "ik_foot_r")
+		cmds.select(clear = True)
+		
+		cmds.select(clear = True)
+		ikHandRoot = cmds.joint(name = "ik_hand_root")
+		cmds.select(clear = True)
+		
+		cmds.select(clear = True)
+		ikHandGun = cmds.joint(name = "ik_hand_gun")
+		cmds.select(clear = True)
+		
+		cmds.select(clear = True)
+		ikHandLeft = cmds.joint(name = "ik_hand_l")
+		cmds.select(clear = True)
+		
+		cmds.select(clear = True)
+		ikHandRight = cmds.joint(name = "ik_hand_r")
+		cmds.select(clear = True)
+		
+		# create hierarhcy
+		cmds.parent(ikFootRoot, "root")
+		cmds.parent(ikHandRoot, "root")
+		cmds.parent(ikFootLeft, ikFootRoot)
+		cmds.parent(ikFootRight, ikFootRoot)
+		cmds.parent(ikHandGun, ikHandRoot)
+		cmds.parent(ikHandLeft, ikHandGun)
+		cmds.parent(ikHandRight, ikHandGun)
+		
+		# constrain the joints
+		leftHandConstraint = cmds.parentConstraint("hand_l", ikHandLeft)[0]
+		rightHandGunConstraint = cmds.parentConstraint("hand_r", ikHandGun)[0]
+		rightHandConstraint = cmds.parentConstraint("hand_r", ikHandRight)[0]
+		leftFootConstraint = cmds.parentConstraint("foot_l", ikFootLeft)[0]
+		rightFootConstraint = cmds.parentConstraint("foot_r", ikFootRight)[0]
+	
+    except:
+		cmds.warning("Something went wrong")
+
+# ------------------------------------------------------------------
+# ADD the IK bones to the Export file.
+#save the AnimRig file before opening the export file.
+cmds.file(save = True, type = "mayaBinary", force = True)
+
+#ADD IK BONES
+#open export file
+sceneName = cmds.file(q = True, sceneName = True)
+exportFile = sceneName.partition("AnimRigs")[0] + "ExportFiles" +  sceneName.partition("AnimRigs")[2].partition(".mb")[0] + "_Export.mb"
+cmds.file(exportFile, open = True, force = True)
+
+#add ik bones
+addIKBones()
+
+# ------------------------------------------------------------------
+def addCustomAttrs():
+    attrs = cmds.listAttr("root")
+    exists = False
+    for one in attrs:
+        if one.find("DistanceCurve") != -1:
+            exists = True
+    if exists == False:
+        cmds.addAttr("root", ln="DistanceCurve", keyable=True)
+addCustomAttrs()
+
+#save the export file.
+cmds.file(save = True, type = "mayaBinary", force = True)
\ No newline at end of file
diff --git a/Engine/Extras/Maya_AnimationRiggingTools/ARTv1/General/ART/Projects/Orion/Scripts/RiftMage_Pre_Script.py b/Engine/Extras/Maya_AnimationRiggingTools/ARTv1/General/ART/Projects/Orion/Scripts/RiftMage_Pre_Script.py
new file mode 100644
--- /dev/null
+++ b/Engine/Extras/Maya_AnimationRiggingTools/ARTv1/General/ART/Projects/Orion/Scripts/RiftMage_Pre_Script.py
@@ -0,0 +1,11 @@
+import maya.cmds as cmds
+
+#cmds.parent("portal", "ring_base", "root")
+#cmds.delete("portal_ring_master")
+
+
+if cmds.objExists("ik_foot_root"):
+	cmds.delete("ik_foot_root")
+	
+if cmds.objExists("ik_hand_root"):
+	cmds.delete("ik_hand_root")
\ No newline at end of file
diff --git a/Engine/Extras/Maya_AnimationRiggingTools/ARTv1/General/ART/Projects/Orion/Scripts/Sparrow_CreateLODs.py b/Engine/Extras/Maya_AnimationRiggingTools/ARTv1/General/ART/Projects/Orion/Scripts/Sparrow_CreateLODs.py
new file mode 100644
--- /dev/null
+++ b/Engine/Extras/Maya_AnimationRiggingTools/ARTv1/General/ART/Projects/Orion/Scripts/Sparrow_CreateLODs.py
@@ -0,0 +1,534 @@
+import maya.cmds as cmds
+import customMayaMenu as menu
+import Tools.ART_ImportMultipleWeights as impw
+import Tools.ART_TransferWeights as transfer
+import Tools.ART_ExportLODtoFBX as export
+reload(export)
+
+################################################## SPARROW #########################################################################
+# Warn the user that this operation will save their file.
+saveResult = cmds.confirmDialog(title = "Save Warning", message="In order to continue your file must be saved.  Would you like to save it?  If yes it will be saved and after the operation is complete your file will be re-opened.", button = ["Yes", "Cancel"], defaultButton='Yes', cancelButton='Cancel', dismissString='Cancel')
+
+# If they are ok with saving the file, warn them that it will also delete all of the current files in their skin weights directory.
+if saveResult == "Yes":
+    deleteSkinFilesResult = cmds.confirmDialog(title = "SkinWeights Warning", message="In order to continue we must delete all of the files in your ..\General\ART\SkinWeights\ directory.  Would you like to continue?.", button = ["Yes", "Cancel"], defaultButton='Yes', cancelButton='Cancel', dismissString='Cancel')
+    if deleteSkinFilesResult == "Yes":
+        toolsPath = cmds.internalVar(usd = True) + "mayaTools.txt"
+        if os.path.exists(toolsPath):
+            f = open(toolsPath, 'r')
+            mayaToolsDir = f.readline()
+            f.close()
+
+        path = mayaToolsDir + "\General\ART\SkinWeights\\"
+        files = os.listdir(path)
+        #delete all weight cache files in the system folder
+        for file in files:
+            os.remove(path + file)
+            #print path + file
+
+
+
+# Sparrow LOD1 -------------------------------------------------------------------------------------------------------------
+def SparrowLOD1():
+    if saveResult == "Yes":
+
+        cmds.file(save=True)
+
+        meshes = ["arm_l_detail_geo_LOD1", "arms_geo_LOD1", "arrow_geo_LOD1", "belt_geo_LOD1", "bow_geo_LOD1", "brooch_chain_geo_LOD1", "brooch_geo_LOD1", "eye_tear_duct_geo_LOD1", "eye_wetness_geo_LOD1", "eyeballs_geo_LOD1", "eyebrows_geo_LOD1", "eyelashes_geo_LOD1", "eyeshadow_geo_LOD1", "feathers_geo_LOD1", "hair_core_geo_LOD1", "hair_ponytail_holder_geo_LOD1", "hair_sheets_geo_LOD1", "head_geo_LOD1", "legs_details_geo_LOD1", "legs_geo_LOD1", "skirt_geo_LOD1", "teeth_lower_geo_LOD1", "teeth_upper_geo_LOD1", "tongue_geo_LOD1", "torso_geo_LOD1"]
+
+        jointsToDelete = []
+
+        # remove unneeded joints
+        jointsToRemove = ["skirt_back_cloth_01", "skirt_back_cloth_02", "skirt_cloth_l_01", "skirt_cloth_l_02", "skirt_front_cloth_01", "skirt_front_cloth_02"]
+        jointToTransferTo = "pelvis"
+        transfer.LOD_transferWeights(meshes, jointsToRemove, jointToTransferTo)
+        jointsToDelete.extend(jointsToRemove)
+
+        jointsToRemove = ["eyeball_l", "eyeball_r", "eyebrow_in_l", "eyebrow_in_r", "eyebrow_mid", "eyebrow_out_l", "eyebrow_out_r", "eyelid_bott_l", "eyelid_bott_r", "eyelid_top_l", "eyelid_top_r"]
+        jointToTransferTo = "head"
+        transfer.LOD_transferWeights(meshes, jointsToRemove, jointToTransferTo)
+        jointsToDelete.extend(jointsToRemove)
+
+        jointsToRemove = ["quiver"]
+        jointToTransferTo = "thigh_r"
+        transfer.LOD_transferWeights(meshes, jointsToRemove, jointToTransferTo)
+        jointsToDelete.extend(jointsToRemove)
+        
+        # Delete the joints we no longer want.
+        for one in jointsToDelete:
+            if cmds.objExists(one):
+                cmds.delete(one)
+                
+        # Export the fbx file.
+        export.exportLOD("Sparrow_LOD1.fbx", "LOD1")
+        
+        # Re-open the file without saving.
+        fullPath = cmds.file(q = True, sceneName = True)
+        cmds.file(fullPath, open=True, f=True)
+
+        print "LOD1 Exported Successfully!!!!!!!!!!!!!!!"
+
+SparrowLOD1()
+
+# Sparrow LOD2 -------------------------------------------------------------------------------------------------------------
+def SparrowLOD2():
+    if saveResult == "Yes":
+
+        cmds.file(save=True)
+        
+        meshes = ["arm_l_detail_geo_LOD2", "arms_geo_LOD2", "arrow_geo_LOD2", "belt_geo_LOD2", "bow_geo_LOD2", "brooch_chain_geo_LOD2", "brooch_geo_LOD2", "eye_tear_duct_geo_LOD2", "eye_wetness_geo_LOD2", "eyeballs_geo_LOD2", "eyebrows_geo_LOD2", "eyelashes_geo_LOD2", "eyeshadow_geo_LOD2", "feathers_geo_LOD2", "hair_core_geo_LOD2", "hair_ponytail_holder_geo_LOD2", "hair_sheets_geo_LOD2", "head_geo_LOD2", "legs_details_geo_LOD2", "legs_geo_LOD2", "skirt_geo_LOD2", "teeth_lower_geo_LOD2", "teeth_upper_geo_LOD2", "tongue_geo_LOD2", "torso_geo_LOD2"]
+
+        # Export a skin weights file for each mesh we plan to operate on.
+        cmds.select(meshes)
+        menu.exportSelectedSkinWeights(True)
+
+        # Set the pose if needed
+        cmds.setAttr("thumb_03_r.rotateZ", 17.55735447)
+        cmds.setAttr("thumb_03_r.rotateY", 2.752475475)
+        cmds.setAttr("thumb_03_r.rotateX", -8.614445938)
+        cmds.setAttr("ring_03_r.rotateZ", 58.10260223)
+        cmds.setAttr("pinky_03_r.rotateZ", 58.10260223)
+        cmds.setAttr("middle_03_r.rotateZ", 60.1391109)
+        cmds.setAttr("index_03_r.rotateZ", 60.1391109)
+
+        cmds.setAttr("thumb_03_l.rotateZ", 49.02926105)
+        cmds.setAttr("ring_03_l.rotateZ", 60.61060457)
+        cmds.setAttr("pinky_03_l.rotateZ", 60.61060457)
+        cmds.setAttr("middle_03_l.rotateZ", 60.61060457)
+        cmds.setAttr("index_03_l.rotateZ", 60.61060457)
+
+        # Delete the history on the meshes so that they get baked down into their posed positions
+        cmds.delete(meshes, ch=True)
+
+        # Move the skeleton back to its bindpose (This script REQUIRES the bind pose to be the pose that we want to export the skeleton to.  This is the case for all known characters created with A.R.T.)
+        cmds.dagPose("root", bp=True, r=True) 
+
+        # Re-import the skinning we exported before onto the mesh so thats its bound in its posed position.
+        cmds.select(meshes)
+        impw.importMultipleSkinWeights("", True)
+
+        jointsToDelete = []
+
+         # remove unneeded joints
+        jointsToRemove = ["front_strap_01", "front_strap_02", "front_strap_small", "skirt_back_cloth_01", "skirt_back_cloth_02", "skirt_cloth_l_01", "skirt_cloth_l_02", "skirt_front_cloth_01", "skirt_front_cloth_02"]
+        jointToTransferTo = "pelvis"
+        transfer.LOD_transferWeights(meshes, jointsToRemove, jointToTransferTo)
+        jointsToDelete.extend(jointsToRemove)
+
+        jointsToRemove = ["breast_chain", "brooch_back_l", "brooch_back_r", "brooch_front_l", "brooch_front_r", "chest_cloth_l_01", "chest_cloth_l_02", "chest_cloth_r_01", "chest_cloth_r_02"]
+        jointToTransferTo = "spine_03"
+        transfer.LOD_transferWeights(meshes, jointsToRemove, jointToTransferTo)
+        jointsToDelete.extend(jointsToRemove)
+
+        jointsToRemove = ["eyeball_l", "eyeball_r", "eyebrow_in_l", "eyebrow_in_r", "eyebrow_mid", "eyebrow_out_l", "eyebrow_out_r", "eyelid_bott_l", "eyelid_bott_r", "eyelid_top_l", "eyelid_top_r", "feathers", "hair_curls_A_l_01", "hair_curls_A_r_01", "hair_curls_B_l_01", "hair_curls_B_r_01", "hair_top"]
+        jointToTransferTo = "head"
+        transfer.LOD_transferWeights(meshes, jointsToRemove, jointToTransferTo)
+        jointsToDelete.extend(jointsToRemove)
+
+        jointsToRemove = ["upperarm_twist_02_l"]
+        jointToTransferTo = "upperarm_twist_01_l"
+        transfer.LOD_transferWeights(meshes, jointsToRemove, jointToTransferTo)
+        jointsToDelete.extend(jointsToRemove)
+
+        jointsToRemove = ["upperarm_twist_02_r"]
+        jointToTransferTo = "upperarm_twist_01_r"
+        transfer.LOD_transferWeights(meshes, jointsToRemove, jointToTransferTo)
+        jointsToDelete.extend(jointsToRemove)
+
+        jointsToRemove = ["lowerarm_twist_02_l"]
+        jointToTransferTo = "lowerarm_twist_01_l"
+        transfer.LOD_transferWeights(meshes, jointsToRemove, jointToTransferTo)
+        jointsToDelete.extend(jointsToRemove)
+
+        jointsToRemove = ["lowerarm_twist_02_r"]
+        jointToTransferTo = "lowerarm_twist_01_r"
+        transfer.LOD_transferWeights(meshes, jointsToRemove, jointToTransferTo)
+        jointsToDelete.extend(jointsToRemove)
+
+        jointsToRemove = ["index_03_l"]
+        jointToTransferTo = "index_02_l"
+        transfer.LOD_transferWeights(meshes, jointsToRemove, jointToTransferTo)
+        jointsToDelete.extend(jointsToRemove)
+
+        jointsToRemove = ["middle_03_l"]
+        jointToTransferTo = "middle_02_l"
+        transfer.LOD_transferWeights(meshes, jointsToRemove, jointToTransferTo)
+        jointsToDelete.extend(jointsToRemove)
+
+        jointsToRemove = ["pinky_03_l"]
+        jointToTransferTo = "pinky_02_l"
+        transfer.LOD_transferWeights(meshes, jointsToRemove, jointToTransferTo)
+        jointsToDelete.extend(jointsToRemove)
+
+        jointsToRemove = ["ring_03_l"]
+        jointToTransferTo = "ring_02_l"
+        transfer.LOD_transferWeights(meshes, jointsToRemove, jointToTransferTo)
+        jointsToDelete.extend(jointsToRemove)
+
+        jointsToRemove = ["thumb_03_l"]
+        jointToTransferTo = "thumb_02_l"
+        transfer.LOD_transferWeights(meshes, jointsToRemove, jointToTransferTo)
+        jointsToDelete.extend(jointsToRemove)
+
+        jointsToRemove = ["index_03_r"]
+        jointToTransferTo = "index_02_r"
+        transfer.LOD_transferWeights(meshes, jointsToRemove, jointToTransferTo)
+        jointsToDelete.extend(jointsToRemove)
+
+        jointsToRemove = ["middle_03_r"]
+        jointToTransferTo = "middle_02_r"
+        transfer.LOD_transferWeights(meshes, jointsToRemove, jointToTransferTo)
+        jointsToDelete.extend(jointsToRemove)
+
+        jointsToRemove = ["pinky_03_r"]
+        jointToTransferTo = "pinky_02_r"
+        transfer.LOD_transferWeights(meshes, jointsToRemove, jointToTransferTo)
+        jointsToDelete.extend(jointsToRemove)
+
+        jointsToRemove = ["ring_03_r"]
+        jointToTransferTo = "ring_02_r"
+        transfer.LOD_transferWeights(meshes, jointsToRemove, jointToTransferTo)
+        jointsToDelete.extend(jointsToRemove)
+
+        jointsToRemove = ["thumb_03_r"]
+        jointToTransferTo = "thumb_02_r"
+        transfer.LOD_transferWeights(meshes, jointsToRemove, jointToTransferTo)
+        jointsToDelete.extend(jointsToRemove)
+
+        jointsToRemove = ["quiver"]
+        jointToTransferTo = "thigh_r"
+        transfer.LOD_transferWeights(meshes, jointsToRemove, jointToTransferTo)
+        jointsToDelete.extend(jointsToRemove)
+
+        # Delete the joints we no longer want.
+        for one in jointsToDelete:
+            if cmds.objExists(one):
+                cmds.delete(one)
+                
+        # Export the fbx file.
+        export.exportLOD("Sparrow_LOD2.fbx", "LOD2")
+
+        # Re-open the file without saving.
+        fullPath = cmds.file(q = True, sceneName = True)
+        cmds.file(fullPath, open=True, f=True)
+
+        print "LOD2 Exported Successfully!!!!!!!!!!!!!!!"
+
+SparrowLOD2()
+
+# Sparrow LOD3 -------------------------------------------------------------------------------------------------------------
+def SparrowLOD3():
+    if saveResult == "Yes":
+
+        cmds.file(save=True)
+    
+        meshes = ["arm_l_detail_geo_LOD3", "arms_geo_LOD3", "arrow_geo_LOD3", "belt_geo_LOD3", "bow_geo_LOD3", "brooch_chain_geo_LOD3", "brooch_geo_LOD3", "eye_tear_duct_geo_LOD3", "eye_wetness_geo_LOD3", "eyeballs_geo_LOD3", "eyebrows_geo_LOD3", "eyelashes_geo_LOD3", "eyeshadow_geo_LOD3", "feathers_geo_LOD3", "hair_core_geo_LOD3", "hair_ponytail_holder_geo_LOD3", "hair_sheets_geo_LOD3", "head_geo_LOD3", "legs_details_geo_LOD3", "legs_geo_LOD3", "skirt_geo_LOD3", "teeth_lower_geo_LOD3", "teeth_upper_geo_LOD3", "tongue_geo_LOD3", "torso_geo_LOD3"]
+    
+        # Export a skin weights file for each mesh we plan to operate on.
+        cmds.select(meshes)
+        menu.exportSelectedSkinWeights(True)
+
+        # Set the pose if needed
+        cmds.setAttr("thumb_03_r.rotateZ", 17.55735447)
+        cmds.setAttr("thumb_03_r.rotateY", 2.752475475)
+        cmds.setAttr("thumb_03_r.rotateX", -8.614445938)
+        cmds.setAttr("thumb_02_r.rotateZ", 15.15922943)
+        cmds.setAttr("thumb_01_r.rotateZ", 18.03248389)
+        cmds.setAttr("thumb_01_r.rotateY", 7.009164082)
+        cmds.setAttr("thumb_01_r.rotateX", -0.3484935255)
+        cmds.setAttr("ring_03_r.rotateZ", 58.10260223)
+        cmds.setAttr("ring_02_r.rotateZ", 58.10260223)
+        cmds.setAttr("ring_01_r.rotateZ", 58.10260223)
+        cmds.setAttr("pinky_03_r.rotateZ", 58.10260223)
+        cmds.setAttr("pinky_02_r.rotateZ", 64.22260217)
+        cmds.setAttr("pinky_01_r.rotateZ", 71.7111965)
+        cmds.setAttr("middle_03_r.rotateZ", 60.1391109)
+        cmds.setAttr("middle_02_r.rotateZ", 50.6928414)
+        cmds.setAttr("middle_01_r.rotateZ", 6.782023659)
+        cmds.setAttr("index_03_r.rotateZ", 60.1391109)
+        cmds.setAttr("index_02_r.rotateZ", 50.6928414)
+        cmds.setAttr("index_01_r.rotateZ", 6.782023659)
+
+        cmds.setAttr("thumb_03_l.rotateZ", 49.02926105)
+        cmds.setAttr("thumb_02_l.rotateZ", 29.37554526)
+        cmds.setAttr("thumb_01_l.rotateZ", -1.060679943)
+        cmds.setAttr("thumb_01_l.rotateY", 16.78500021)
+        cmds.setAttr("thumb_01_l.rotateX", 5.616322924)
+        cmds.setAttr("ring_03_l.rotateZ", 60.61060457)
+        cmds.setAttr("ring_02_l.rotateZ", 60.61060457)
+        cmds.setAttr("ring_01_l.rotateZ", 60.61060457)
+        cmds.setAttr("pinky_03_l.rotateZ", 60.61060457)
+        cmds.setAttr("pinky_02_l.rotateZ", 60.61060457)
+        cmds.setAttr("pinky_01_l.rotateZ", 60.61060457)
+        cmds.setAttr("middle_03_l.rotateZ", 60.61060457)
+        cmds.setAttr("middle_02_l.rotateZ", 60.61060457)
+        cmds.setAttr("middle_01_l.rotateZ", 60.61060457)
+        cmds.setAttr("index_03_l.rotateZ", 60.61060457)
+        cmds.setAttr("index_02_l.rotateZ", 60.61060457)
+        cmds.setAttr("index_01_l.rotateZ", 60.61060457)
+
+        # Delete the history on the meshes so that they get baked down into their posed positions
+        cmds.delete(meshes, ch=True)
+
+        # Move the skeleton back to its bindpose (This script REQUIRES the bind pose to be the pose that we want to export the skeleton to.  This is the case for all known characters created with A.R.T.)
+        cmds.dagPose("root", bp=True, r=True) 
+
+        # Re-import the skinning we exported before onto the mesh so thats its bound in its posed position.
+        cmds.select(meshes)
+        impw.importMultipleSkinWeights("", True)
+
+        jointsToDelete = []
+
+        # remove unneeded joints
+        jointsToRemove = ["front_strap_01", "front_strap_02", "front_strap_small", "skirt_back_cloth_01", "skirt_back_cloth_02", "skirt_cloth_l_01", "skirt_cloth_l_02", "skirt_front_cloth_01", "skirt_front_cloth_02"]
+        jointToTransferTo = "pelvis"
+        transfer.LOD_transferWeights(meshes, jointsToRemove, jointToTransferTo)
+        jointsToDelete.extend(jointsToRemove)
+
+        jointsToRemove = ["back_tail_cloth_01", "back_tail_cloth_02", "breast_chain", "brooch_back_l", "brooch_back_r", "brooch_front_l", "brooch_front_r", "chest_cloth_l_01", "chest_cloth_l_02", "chest_cloth_r_01", "chest_cloth_r_02"]
+        jointToTransferTo = "spine_03"
+        transfer.LOD_transferWeights(meshes, jointsToRemove, jointToTransferTo)
+        jointsToDelete.extend(jointsToRemove)
+
+        jointsToRemove = ["eyeball_l", "eyeball_r", "eyebrow_in_l", "eyebrow_in_r", "eyebrow_mid", "eyebrow_out_l", "eyebrow_out_r", "eyelid_bott_l", "eyelid_bott_r", "eyelid_top_l", "eyelid_top_r", "feathers", "hair_curls_A_l_01", "hair_curls_A_r_01", "hair_curls_B_l_01", "hair_curls_B_r_01", "hair_top", "jaw", "ponytail_01", "ponytail_02", "ponytail_03", "ponytail_04"]
+        jointToTransferTo = "head"
+        transfer.LOD_transferWeights(meshes, jointsToRemove, jointToTransferTo)
+        jointsToDelete.extend(jointsToRemove)
+
+        jointsToRemove = ["upperarm_twist_01_l", "upperarm_twist_02_l"]
+        jointToTransferTo = "upperarm_l"
+        transfer.LOD_transferWeights(meshes, jointsToRemove, jointToTransferTo)
+        jointsToDelete.extend(jointsToRemove)
+
+        jointsToRemove = ["upperarm_twist_02_r", "upperarm_twist_01_r"]
+        jointToTransferTo = "upperarm_r"
+        transfer.LOD_transferWeights(meshes, jointsToRemove, jointToTransferTo)
+        jointsToDelete.extend(jointsToRemove)
+
+        jointsToRemove = ["lowerarm_twist_01_l", "lowerarm_twist_02_l"]
+        jointToTransferTo = "lowerarm_l"
+        transfer.LOD_transferWeights(meshes, jointsToRemove, jointToTransferTo)
+        jointsToDelete.extend(jointsToRemove)
+
+        jointsToRemove = ["lowerarm_twist_02_r", "lowerarm_twist_01_r"]
+        jointToTransferTo = "lowerarm_r"
+        transfer.LOD_transferWeights(meshes, jointsToRemove, jointToTransferTo)
+        jointsToDelete.extend(jointsToRemove)
+
+        jointsToRemove = ["index_01_l", "index_02_l", "index_03_l", "middle_01_l", "middle_02_l", "middle_03_l", "pinky_01_l", "pinky_02_l", "pinky_03_l", "ring_01_l", "ring_02_l", "ring_03_l", "thumb_01_l", "thumb_02_l", "thumb_03_l"]
+        jointToTransferTo = "hand_l"
+        transfer.LOD_transferWeights(meshes, jointsToRemove, jointToTransferTo)
+        jointsToDelete.extend(jointsToRemove)
+
+        jointsToRemove = ["index_01_r", "index_02_r", "index_03_r", "middle_01_r", "middle_02_r", "middle_03_r", "pinky_01_r", "pinky_02_r", "pinky_03_r", "ring_01_r", "ring_02_r", "ring_03_r", "thumb_01_r", "thumb_02_r", "thumb_03_r"]
+        jointToTransferTo = "hand_r"
+        transfer.LOD_transferWeights(meshes, jointsToRemove, jointToTransferTo)
+        jointsToDelete.extend(jointsToRemove)
+
+        jointsToRemove = ["bow_shaft_bott_02", "bow_shaft_bott_03", "bow_string_bott"]
+        jointToTransferTo = "bow_shaft_bott_01"
+        transfer.LOD_transferWeights(meshes, jointsToRemove, jointToTransferTo)
+        jointsToDelete.extend(jointsToRemove)
+
+        jointsToRemove = ["bow_shaft_top_02", "bow_shaft_top_03", "bow_string_top"]
+        jointToTransferTo = "bow_shaft_top_01"
+        transfer.LOD_transferWeights(meshes, jointsToRemove, jointToTransferTo)
+        jointsToDelete.extend(jointsToRemove)
+
+        jointsToRemove = ["thigh_twist_01_l", "thigh_twist_02_l"]
+        jointToTransferTo = "thigh_l"
+        transfer.LOD_transferWeights(meshes, jointsToRemove, jointToTransferTo)
+        jointsToDelete.extend(jointsToRemove)
+
+        jointsToRemove = ["thigh_twist_01_r", "thigh_twist_02_r", "quiver"]
+        jointToTransferTo = "thigh_r"
+        transfer.LOD_transferWeights(meshes, jointsToRemove, jointToTransferTo)
+        jointsToDelete.extend(jointsToRemove)
+
+        jointsToRemove = ["calf_twist_01_l"]
+        jointToTransferTo = "calf_l"
+        transfer.LOD_transferWeights(meshes, jointsToRemove, jointToTransferTo)
+        jointsToDelete.extend(jointsToRemove)
+
+        jointsToRemove = ["calf_twist_01_r"]
+        jointToTransferTo = "calf_r"
+        transfer.LOD_transferWeights(meshes, jointsToRemove, jointToTransferTo)
+        jointsToDelete.extend(jointsToRemove)
+
+        # Delete the joints we no longer want.
+        for one in jointsToDelete:
+            if cmds.objExists(one):
+                cmds.delete(one)
+                
+        # Export the fbx file.
+        export.exportLOD("Sparrow_LOD3.fbx", "LOD3")
+        
+        # Re-open the file without saving.
+        fullPath = cmds.file(q = True, sceneName = True)
+        cmds.file(fullPath, open=True, f=True)
+
+        print "LOD3 Exported Successfully!!!!!!!!!!!!!!!"
+
+SparrowLOD3()
+
+# Sparrow LOD4 -------------------------------------------------------------------------------------------------------------
+def SparrowLOD4():
+    if saveResult == "Yes":
+
+        cmds.file(save=True)
+    
+        meshes = ["arm_l_detail_geo_LOD4", "arms_geo_LOD4", "arrow_geo_LOD4", "belt_geo_LOD4", "bow_geo_LOD4", "brooch_chain_geo_LOD4", "brooch_geo_LOD4", "eye_tear_duct_geo_LOD4", "eye_wetness_geo_LOD4", "eyeballs_geo_LOD4", "eyebrows_geo_LOD4", "eyelashes_geo_LOD4", "eyeshadow_geo_LOD4", "feathers_geo_LOD4", "hair_core_geo_LOD4", "hair_ponytail_holder_geo_LOD4", "hair_sheets_geo_LOD4", "head_geo_LOD4", "legs_details_geo_LOD4", "legs_geo_LOD4", "skirt_geo_LOD4", "teeth_lower_geo_LOD4", "teeth_upper_geo_LOD4", "tongue_geo_LOD4", "torso_geo_LOD4"]
+
+        # Export a skin weights file for each mesh we plan to operate on.
+        cmds.select(meshes)
+        menu.exportSelectedSkinWeights(True)
+
+        # Set the pose if needed
+        cmds.setAttr("thumb_03_r.rotateZ", 17.55735447)
+        cmds.setAttr("thumb_03_r.rotateY", 2.752475475)
+        cmds.setAttr("thumb_03_r.rotateX", -8.614445938)
+        cmds.setAttr("thumb_02_r.rotateZ", 15.15922943)
+        cmds.setAttr("thumb_01_r.rotateZ", 18.03248389)
+        cmds.setAttr("thumb_01_r.rotateY", 7.009164082)
+        cmds.setAttr("thumb_01_r.rotateX", -0.3484935255)
+        cmds.setAttr("ring_03_r.rotateZ", 58.10260223)
+        cmds.setAttr("ring_02_r.rotateZ", 58.10260223)
+        cmds.setAttr("ring_01_r.rotateZ", 58.10260223)
+        cmds.setAttr("pinky_03_r.rotateZ", 58.10260223)
+        cmds.setAttr("pinky_02_r.rotateZ", 64.22260217)
+        cmds.setAttr("pinky_01_r.rotateZ", 71.7111965)
+        cmds.setAttr("middle_03_r.rotateZ", 60.1391109)
+        cmds.setAttr("middle_02_r.rotateZ", 50.6928414)
+        cmds.setAttr("middle_01_r.rotateZ", 6.782023659)
+        cmds.setAttr("index_03_r.rotateZ", 60.1391109)
+        cmds.setAttr("index_02_r.rotateZ", 50.6928414)
+        cmds.setAttr("index_01_r.rotateZ", 6.782023659)
+
+        cmds.setAttr("thumb_03_l.rotateZ", 49.02926105)
+        cmds.setAttr("thumb_02_l.rotateZ", 29.37554526)
+        cmds.setAttr("thumb_01_l.rotateZ", -1.060679943)
+        cmds.setAttr("thumb_01_l.rotateY", 16.78500021)
+        cmds.setAttr("thumb_01_l.rotateX", 5.616322924)
+        cmds.setAttr("ring_03_l.rotateZ", 60.61060457)
+        cmds.setAttr("ring_02_l.rotateZ", 60.61060457)
+        cmds.setAttr("ring_01_l.rotateZ", 60.61060457)
+        cmds.setAttr("pinky_03_l.rotateZ", 60.61060457)
+        cmds.setAttr("pinky_02_l.rotateZ", 60.61060457)
+        cmds.setAttr("pinky_01_l.rotateZ", 60.61060457)
+        cmds.setAttr("middle_03_l.rotateZ", 60.61060457)
+        cmds.setAttr("middle_02_l.rotateZ", 60.61060457)
+        cmds.setAttr("middle_01_l.rotateZ", 60.61060457)
+        cmds.setAttr("index_03_l.rotateZ", 60.61060457)
+        cmds.setAttr("index_02_l.rotateZ", 60.61060457)
+        cmds.setAttr("index_01_l.rotateZ", 60.61060457)
+
+        # Delete the history on the meshes so that they get baked down into their posed positions
+        cmds.delete(meshes, ch=True)
+
+        # Move the skeleton back to its bindpose (This script REQUIRES the bind pose to be the pose that we want to export the skeleton to.  This is the case for all known characters created with A.R.T.)
+        cmds.dagPose("root", bp=True, r=True) 
+
+        # Re-import the skinning we exported before onto the mesh so thats its bound in its posed position.
+        cmds.select(meshes)
+        impw.importMultipleSkinWeights("", True)
+        
+        jointsToDelete = []
+
+        # remove unneeded joints
+        jointsToRemove = ["front_strap_01", "front_strap_02", "front_strap_small", "skirt_back_cloth_01", "skirt_back_cloth_02", "skirt_cloth_l_01", "skirt_cloth_l_02", "skirt_front_cloth_01", "skirt_front_cloth_02"]
+        jointToTransferTo = "pelvis"
+        transfer.LOD_transferWeights(meshes, jointsToRemove, jointToTransferTo)
+        jointsToDelete.extend(jointsToRemove)
+
+        jointsToRemove = ["back_tail_cloth_01", "back_tail_cloth_02", "breast_chain", "brooch_back_l", "brooch_back_r", "brooch_front_l", "brooch_front_r", "chest_cloth_l_01", "chest_cloth_l_02", "chest_cloth_r_01", "chest_cloth_r_02"]
+        jointToTransferTo = "spine_03"
+        transfer.LOD_transferWeights(meshes, jointsToRemove, jointToTransferTo)
+        jointsToDelete.extend(jointsToRemove)
+
+        jointsToRemove = ["eyeball_l", "eyeball_r", "eyebrow_in_l", "eyebrow_in_r", "eyebrow_mid", "eyebrow_out_l", "eyebrow_out_r", "eyelid_bott_l", "eyelid_bott_r", "eyelid_top_l", "eyelid_top_r", "feathers", "hair_curls_A_l_01", "hair_curls_A_r_01", "hair_curls_B_l_01", "hair_curls_B_r_01", "hair_top", "jaw", "ponytail_01", "ponytail_02", "ponytail_03", "ponytail_04"]
+        jointToTransferTo = "head"
+        transfer.LOD_transferWeights(meshes, jointsToRemove, jointToTransferTo)
+        jointsToDelete.extend(jointsToRemove)
+
+        jointsToRemove = ["upperarm_twist_01_l", "upperarm_twist_02_l"]
+        jointToTransferTo = "upperarm_l"
+        transfer.LOD_transferWeights(meshes, jointsToRemove, jointToTransferTo)
+        jointsToDelete.extend(jointsToRemove)
+
+        jointsToRemove = ["upperarm_twist_02_r", "upperarm_twist_01_r"]
+        jointToTransferTo = "upperarm_r"
+        transfer.LOD_transferWeights(meshes, jointsToRemove, jointToTransferTo)
+        jointsToDelete.extend(jointsToRemove)
+
+        jointsToRemove = ["lowerarm_twist_01_l", "lowerarm_twist_02_l"]
+        jointToTransferTo = "lowerarm_l"
+        transfer.LOD_transferWeights(meshes, jointsToRemove, jointToTransferTo)
+        jointsToDelete.extend(jointsToRemove)
+
+        jointsToRemove = ["lowerarm_twist_02_r", "lowerarm_twist_01_r"]
+        jointToTransferTo = "lowerarm_r"
+        transfer.LOD_transferWeights(meshes, jointsToRemove, jointToTransferTo)
+        jointsToDelete.extend(jointsToRemove)
+
+        jointsToRemove = ["index_01_l", "index_02_l", "index_03_l", "middle_01_l", "middle_02_l", "middle_03_l", "pinky_01_l", "pinky_02_l", "pinky_03_l", "ring_01_l", "ring_02_l", "ring_03_l", "thumb_01_l", "thumb_02_l", "thumb_03_l"]
+        jointToTransferTo = "hand_l"
+        transfer.LOD_transferWeights(meshes, jointsToRemove, jointToTransferTo)
+        jointsToDelete.extend(jointsToRemove)
+
+        jointsToRemove = ["index_01_r", "index_02_r", "index_03_r", "middle_01_r", "middle_02_r", "middle_03_r", "pinky_01_r", "pinky_02_r", "pinky_03_r", "ring_01_r", "ring_02_r", "ring_03_r", "thumb_01_r", "thumb_02_r", "thumb_03_r"]
+        jointToTransferTo = "hand_r"
+        transfer.LOD_transferWeights(meshes, jointsToRemove, jointToTransferTo)
+        jointsToDelete.extend(jointsToRemove)
+
+        jointsToRemove = ["bow_shaft_bott_01", "bow_shaft_bott_02", "bow_shaft_bott_03", "bow_string_bott", "bow_shaft_top_01", "bow_shaft_top_02", "bow_shaft_top_03", "bow_string_top"]
+        jointToTransferTo = "bow_base"
+        transfer.LOD_transferWeights(meshes, jointsToRemove, jointToTransferTo)
+        jointsToDelete.extend(jointsToRemove)
+
+        jointsToRemove = ["thigh_twist_01_l", "thigh_twist_02_l"]
+        jointToTransferTo = "thigh_l"
+        transfer.LOD_transferWeights(meshes, jointsToRemove, jointToTransferTo)
+        jointsToDelete.extend(jointsToRemove)
+
+        jointsToRemove = ["thigh_twist_01_r", "thigh_twist_02_r", "quiver"]
+        jointToTransferTo = "thigh_r"
+        transfer.LOD_transferWeights(meshes, jointsToRemove, jointToTransferTo)
+        jointsToDelete.extend(jointsToRemove)
+
+        jointsToRemove = ["calf_twist_01_l"]
+        jointToTransferTo = "calf_l"
+        transfer.LOD_transferWeights(meshes, jointsToRemove, jointToTransferTo)
+        jointsToDelete.extend(jointsToRemove)
+
+        jointsToRemove = ["calf_twist_01_r"]
+        jointToTransferTo = "calf_r"
+        transfer.LOD_transferWeights(meshes, jointsToRemove, jointToTransferTo)
+        jointsToDelete.extend(jointsToRemove)
+
+        jointsToRemove = ["ball_l"]
+        jointToTransferTo = "foot_l"
+        transfer.LOD_transferWeights(meshes, jointsToRemove, jointToTransferTo)
+        jointsToDelete.extend(jointsToRemove)
+
+        jointsToRemove = ["ball_r"]
+        jointToTransferTo = "foot_r"
+        transfer.LOD_transferWeights(meshes, jointsToRemove, jointToTransferTo)
+        jointsToDelete.extend(jointsToRemove)
+
+        # Delete the joints we no longer want.
+        for one in jointsToDelete:
+            if cmds.objExists(one):
+                cmds.delete(one)
+                
+        # Export the fbx file.
+        export.exportLOD("Sparrow_LOD4.fbx", "LOD4")
+        
+        # Re-open the file without saving.
+        fullPath = cmds.file(q = True, sceneName = True)
+        cmds.file(fullPath, open=True, f=True)
+
+        print "LOD4 Exported Successfully!!!!!!!!!!!!!!!"
+
+SparrowLOD4()
+
+
