Branching

diff --git a/Engine/Extras/Maya_AnimationRiggingTools/ARTv1/General/ART/Projects/Orion/Scripts/Sparrow_Post_Script.py b/Engine/Extras/Maya_AnimationRiggingTools/ARTv1/General/ART/Projects/Orion/Scripts/Sparrow_Post_Script.py
new file mode 100644
--- /dev/null
+++ b/Engine/Extras/Maya_AnimationRiggingTools/ARTv1/General/ART/Projects/Orion/Scripts/Sparrow_Post_Script.py
@@ -0,0 +1,409 @@
+import maya.cmds as cmds
+
+# NOTE: This character has special steps to run during Rig Pose creation.  See end of script.
+
+# ------------------------------------------------------------------
+# scale all of the controls to be more visible with the mesh.
+def scaleControls():
+    sides = ["_l", "_r"]
+    for side in sides:
+        # feet
+        cmds.select(["ik_foot_anim"+side+".cv[*]"])
+        cmds.scale(1, 1, 1, r = True, ocp = False)
+        cmds.move(0, 0, 0, r = True)
+
+        cmds.select(["ik_wrist"+side+"_anim.cv[*]"])
+        cmds.scale(1, 1, 1, r = True, ocp = False)
+
+        cmds.select(["eyebrow_in"+side+"_anim.cv[*]"])
+        cmds.scale(0.117744, 0.117744, 0.117744, r=True, ocp = False)
+        cmds.move(0, 0, 0, r=True)
+
+        cmds.select(["eyebrow_out"+side+"_anim.cv[*]"])
+        cmds.scale(0.117744, 0.117744, 0.117744, r=True, ocp = False)
+        cmds.move(0, 0, 0, r=True)
+
+        cmds.select(["eyeball"+side+"_anim.cv[*]"])
+        cmds.scale(0.1, 0.1, 0.05, r=True, ocp = False)
+        if side == "_l":
+            cmds.move(0, -3.757386, 0, r=True, os=True)
+        else:
+            cmds.move(0, -3.757386, 0, r=True, os=True)
+
+        cmds.select(["eyelid_bott"+side+"_anim.cv[*]"])
+        cmds.scale(0.1, 0.1, 0.01, r=True, ocp = False)
+        if side == "_l":
+            cmds.move(0, -3.757386, 0.5, r=True, os=True)
+        else:
+            cmds.move(0, 3.757386, 0.5, r=True, os=True)
+
+        cmds.select(["eyelid_top"+side+"_anim.cv[*]"])
+        cmds.scale(0.1, 0.1, 0.01, r=True, ocp = False)
+        if side == "_l":
+            cmds.move(0, -3.757386, 0.5, r=True, os=True)
+        else:
+            cmds.move(0, 3.757386, 0.5, r=True, os=True)
+
+    objects = ["hair_top_anim", "fk_ponytail_01_anim", "fk_ponytail_02_anim", "fk_ponytail_03_anim", "fk_ponytail_04_anim", "quiver_anim", "arrow_nock_anim", "bow_base_anim", "bow_shaft_top_01_anim", "bow_shaft_top_02_anim", "bow_shaft_top_03_anim", "bow_shaft_bott_01_anim", "bow_shaft_bott_02_anim", "bow_shaft_bott_03_anim", "bow_string_bott_anim", "bow_string_mid_anim", "bow_string_top_anim", "front_strap_small_anim", "fk_front_strap_01_anim", "fk_front_strap_02_anim", "fk_skirt_back_cloth_01_anim", "fk_skirt_back_cloth_02_anim", "fk_skirt_cloth_l_01_anim", "fk_skirt_cloth_l_02_anim", "fk_skirt_front_cloth_01_anim", "fk_skirt_front_cloth_02_anim", "breast_chain_anim", "feathers_anim", "hair_curls_A_l_01_anim", "hair_curls_A_r_01_anim", "hair_curls_B_l_01_anim", "hair_curls_B_l_01_anim", "fk_chest_cloth_l_01_anim", "fk_chest_cloth_l_02_anim", "fk_chest_cloth_r_01_anim", "fk_chest_cloth_r_02_anim", "fk_back_tail_cloth_01_anim", "fk_back_tail_cloth_02_anim", "brooch_front_l_anim", "brooch_front_r_anim", "brooch_back_l_anim", "brooch_back_r_anim"]
+    for one in objects:
+        cmds.select([one+".cv[*]"])
+        cmds.scale(0.227057, 0.227057, 0.227057, r=True, ocp=False)
+        cmds.move(0, 0, 0, r = True, os=True, wd=True)
+
+    cmds.select(["jaw_anim.cv[*]"])
+    cmds.scale(0.2, 0.2, 0.2, r = True, ocp = True)
+    cmds.move(0, -7, -7, r = True)
+
+    cmds.select(["eyebrow_mid_anim.cv[*]"])
+    cmds.scale(0.117744, 0.117744, 0.117744, r = True, ocp = True)
+
+    cmds.select(["arrow_nock_anim.cv[*]"])
+    cmds.scale(0, 0.3, 4, r = True, ocp = True)
+
+    cmds.select(["bow_string_mid_anim.cv[*]"])
+    cmds.scale(0, 0.6, 4, r = True, ocp = True)
+
+    cmds.select(None)
+scaleControls()
+
+# ------------------------------------------------------------------
+def bowAndArrowRig():
+
+    import Tools.ART_addSpaceSwitching as space
+    reload(space)
+        
+    # Create Arrow base and tip controls
+    arrow_base_anim = cmds.duplicate("arrow_nock_anim", n="arrow_base_anim")
+    cmds.xform(arrow_base_anim, t=[0, 0, -35.51354])
+    cmds.makeIdentity(arrow_base_anim, a=True, t=True)
+
+    arrow_tip_anim = cmds.duplicate("arrow_nock_anim", n="arrow_tip_anim")
+    cmds.xform(arrow_tip_anim, t=[0, 0, -78.67631])
+    cmds.makeIdentity(arrow_tip_anim, a=True, t=True)
+
+    cmds.parent(arrow_base_anim, "arrow_nock_anim")
+    cmds.parent(arrow_tip_anim, arrow_base_anim)
+
+    cmds.delete("driver_arrow_nock_parentConstraint1")
+    cmds.parentConstraint(arrow_tip_anim, "driver_arrow_nock", mo=True)
+
+    # Create aimConstraints for the ends of the strings
+    cmds.aimConstraint("bow_string_mid_anim", "bow_string_bott_anim_grp", mo=True, aimVector=(0, 0, -1), upVector=(0, -1, 0), worldUpType="object", worldUpObject="bow_base_anim")
+    cmds.aimConstraint("bow_string_mid_anim", "bow_string_top_anim_grp", mo=True, aimVector=(0, 0, -1), upVector=(0, -1, 0), worldUpType="object", worldUpObject="bow_base_anim")
+
+    # Create a dummy node that represents the original position of the arrow string.
+    bow_string_mid_anim_WORLD = cmds.duplicate("bow_string_mid_anim", n="bow_string_mid_anim_WORLD", po=True)[0]
+    cmds.parent(bow_string_mid_anim_WORLD, "bow_string_mid_parent_grp")
+    cmds.parentConstraint("bow_string_mid_anim", bow_string_mid_anim_WORLD, mo=True)
+
+    bowshaftDict = {"bow_shaft_top_01_anim_grp":-10, "bow_shaft_top_02_anim_grp":-14, "bow_shaft_top_03_anim_grp":-19.4, "bow_shaft_bott_01_anim_grp":-10, "bow_shaft_bott_02_anim_grp":-14, "bow_shaft_bott_03_anim_grp":-19.4}
+    for one in bowshaftDict:
+        mdNode = cmds.shadingNode("multiplyDivide", asUtility=True, name=one+"_mdNode")
+        clampNode = cmds.shadingNode("clamp", asUtility=True, name=one+"_clampNode")
+        remapValueNode = cmds.shadingNode("remapValue", asUtility=True, name=one+"_remapValueNode")
+
+        cmds.connectAttr(bow_string_mid_anim_WORLD+".translateY", clampNode+".inputR", f=True)
+        cmds.connectAttr(clampNode+".outputR", mdNode+".input1X", f=True)
+        cmds.connectAttr(mdNode+".outputX", remapValueNode+".inputValue", f=True)
+        cmds.connectAttr(remapValueNode+".outValue", one+".rotateX", f=True)
+        cmds.setAttr(mdNode+".input2X", -1)
+        cmds.setAttr(clampNode+".maxR", 120)
+        cmds.setAttr(remapValueNode+".value[1].value_Position", -120)
+        cmds.setAttr(remapValueNode+".value[1].value_FloatValue", bowshaftDict[one])
+
+    # Move the arrow to the string
+    cmds.setAttr("arrow_nock_anim_grp.translateX", 0.05019)
+    cmds.setAttr("arrow_nock_anim_grp.translateY", 15.07779)
+    cmds.setAttr("arrow_nock_anim_grp.translateZ", -35.42169)
+    cmds.setAttr("arrow_nock_anim_grp.rotateX", -90)
+    cmds.setAttr("arrow_nock_anim_grp.rotateY", 0)
+    cmds.setAttr("arrow_nock_anim_grp.rotateZ", 0)
+
+    # Make the arrow follow the string using space switching
+    space.addSpaceSwitching("arrow_nock_anim", "bow_string_mid_anim", True)
+
+    #Move the string to an appropriate position relative to where it might be aligned when shooting.
+    cmds.setAttr("bow_string_mid_anim.translateX", -2.45216)
+    cmds.setAttr("bow_string_mid_anim.translateZ", 5.9665)
+
+    # Create an aim constraint rig so that the arrow will always stay aimed at its position on the shaft.
+    arrow_aim_target = cmds.circle(n="arrow_notched_target", r=3.4)[0]
+    cmds.xform(arrow_aim_target, t=(-2.18154, 0.91122, 6.11286))
+    arrow_up_target = cmds.spaceLocator(n="arrow_up_target")
+    cmds.xform(arrow_up_target, t=(-2.18154, 0.91122, 36))
+    cmds.parent(arrow_aim_target, arrow_up_target, "bow_base_anim")
+    cmds.makeIdentity(arrow_aim_target, a=True, t=True, r=True)
+    cmds.aimConstraint(arrow_aim_target, "arrow_nock_anim_grp", mo=True, aimVector=(0, 0, -1), upVector=(0, -1, 0), worldUpType="object", worldUpObject=arrow_up_target[0])
+
+    # Add space switching for the string to follow the hand.  Off by default.
+    space.addSpaceSwitching("bow_string_mid_anim", "driver_hand_r", False)
+
+    # Move the bow to the left hand.
+    cmds.setAttr("bow_base_anim.translateX", 74.70517)
+    cmds.setAttr("bow_base_anim.translateY", 6.12268)
+    cmds.setAttr("bow_base_anim.translateZ", 144.79125)
+    cmds.setAttr("bow_base_anim.rotateX", 183.95171)
+    cmds.setAttr("bow_base_anim.rotateY", -86.42158)
+    cmds.setAttr("bow_base_anim.rotateZ", -83.33129)
+
+bowAndArrowRig()
+
+# ------------------------------------------------------------------
+def armsLegsRig():
+    cmds.setAttr("Rig_Settings.lUpperarmTwist2Amount", 0.3)
+    cmds.setAttr("Rig_Settings.rUpperarmTwist2Amount", 0.3)
+
+    cmds.setAttr("Rig_Settings.lForearmTwist2Amount", 0.3)
+    cmds.setAttr("Rig_Settings.rForearmTwist2Amount", 0.3)
+
+    cmds.setAttr("Rig_Settings.lThighTwist2Amount", 0.3)
+    cmds.setAttr("Rig_Settings.rThighTwist2Amount", 0.3)
+
+    cmds.setAttr("Rig_Settings.lCalfTwistAmount", 0.7)
+    cmds.setAttr("Rig_Settings.rCalfTwistAmount", 0.7)
+
+armsLegsRig()
+
+# ------------------------------------------------------------------
+def broochRig():
+    # Create SDK's for the brooch.
+    cmds.setDrivenKeyframe(["brooch_front_l_anim_grp.rx", "brooch_back_l_anim_grp.rx", "brooch_front_l_anim_grp.rz", "brooch_back_l_anim_grp.rz"], cd = "driver_clavicle_l.rz", itt = "spline", ott = "spline")
+    cmds.setAttr("clavicle_l_anim.ty", -25)
+    cmds.setAttr("brooch_front_l_anim_grp.rz", -43.5)
+    cmds.setAttr("brooch_back_l_anim_grp.rz", -20.8)
+    cmds.setDrivenKeyframe(["brooch_front_l_anim_grp.rx", "brooch_back_l_anim_grp.rx", "brooch_front_l_anim_grp.rz", "brooch_back_l_anim_grp.rz"], cd = "driver_clavicle_l.rz", itt = "spline", ott = "spline")
+
+    cmds.setAttr("clavicle_l_anim.ty", -7)
+    cmds.setAttr("brooch_front_l_anim_grp.rx", 0.17028)
+    cmds.setAttr("brooch_front_l_anim_grp.rz", -6.47505)
+    cmds.setAttr("brooch_back_l_anim_grp.rx", 1.06856)
+    cmds.setAttr("brooch_back_l_anim_grp.rz", -6.61083)
+    cmds.setDrivenKeyframe(["brooch_front_l_anim_grp.rx", "brooch_back_l_anim_grp.rx", "brooch_front_l_anim_grp.rz", "brooch_back_l_anim_grp.rz"], cd = "driver_clavicle_l.rz", itt = "spline", ott = "spline")
+
+    cmds.setAttr("clavicle_l_anim.ty", 18.0)
+    cmds.setAttr("brooch_front_l_anim_grp.rx", -2.3)
+    cmds.setAttr("brooch_front_l_anim_grp.rz", 8.4)
+    cmds.setAttr("brooch_back_l_anim_grp.rx", -14.1)
+    cmds.setAttr("brooch_back_l_anim_grp.rz", 29.7)
+    cmds.setDrivenKeyframe(["brooch_front_l_anim_grp.rx", "brooch_back_l_anim_grp.rx", "brooch_front_l_anim_grp.rz", "brooch_back_l_anim_grp.rz"], cd = "driver_clavicle_l.rz", itt = "spline", ott = "spline")
+
+    cmds.setAttr("clavicle_l_anim.ty", 2.0)
+    cmds.setAttr("brooch_front_l_anim_grp.rx", -3.47122)
+    cmds.setAttr("brooch_front_l_anim_grp.rz", 3.23346)
+    cmds.setAttr("brooch_back_l_anim_grp.rx", -1.9647)
+    cmds.setAttr("brooch_back_l_anim_grp.rz", 9.7907)
+    cmds.setDrivenKeyframe(["brooch_front_l_anim_grp.rx", "brooch_back_l_anim_grp.rx", "brooch_front_l_anim_grp.rz", "brooch_back_l_anim_grp.rz"], cd = "driver_clavicle_l.rz", itt = "spline", ott = "spline")
+
+    cmds.setAttr("clavicle_l_anim.ty", 0)
+    #---
+    cmds.setDrivenKeyframe(["brooch_front_r_anim_grp.rx", "brooch_back_r_anim_grp.rx", "brooch_front_r_anim_grp.rz", "brooch_back_r_anim_grp.rz"], cd = "driver_clavicle_r.rz", itt = "spline", ott = "spline")
+    cmds.setAttr("clavicle_r_anim.ty", -25)
+    cmds.setAttr("brooch_front_r_anim_grp.rz", -43.5)
+    cmds.setAttr("brooch_back_r_anim_grp.rz", -20.8)
+    cmds.setDrivenKeyframe(["brooch_front_r_anim_grp.rx", "brooch_back_r_anim_grp.rx", "brooch_front_r_anim_grp.rz", "brooch_back_r_anim_grp.rz"], cd = "driver_clavicle_r.rz", itt = "spline", ott = "spline")
+
+    cmds.setAttr("clavicle_r_anim.ty", -7)
+    cmds.setAttr("brooch_front_r_anim_grp.rx", 0.17028)
+    cmds.setAttr("brooch_front_r_anim_grp.rz", -6.47505)
+    cmds.setAttr("brooch_back_r_anim_grp.rx", 1.06856)
+    cmds.setAttr("brooch_back_r_anim_grp.rz", -6.61083)
+    cmds.setDrivenKeyframe(["brooch_front_r_anim_grp.rx", "brooch_back_r_anim_grp.rx", "brooch_front_r_anim_grp.rz", "brooch_back_r_anim_grp.rz"], cd = "driver_clavicle_r.rz", itt = "spline", ott = "spline")
+
+    cmds.setAttr("clavicle_r_anim.ty", 18.0)
+    cmds.setAttr("brooch_front_r_anim_grp.rx", -2.3)
+    cmds.setAttr("brooch_front_r_anim_grp.rz", 8.4)
+    cmds.setAttr("brooch_back_r_anim_grp.rx", -14.1)
+    cmds.setAttr("brooch_back_r_anim_grp.rz", 29.7)
+    cmds.setDrivenKeyframe(["brooch_front_r_anim_grp.rx", "brooch_back_r_anim_grp.rx", "brooch_front_r_anim_grp.rz", "brooch_back_r_anim_grp.rz"], cd = "driver_clavicle_r.rz", itt = "spline", ott = "spline")
+
+    cmds.setAttr("clavicle_r_anim.ty", 2.0)
+    cmds.setAttr("brooch_front_r_anim_grp.rx", -3.47122)
+    cmds.setAttr("brooch_front_r_anim_grp.rz", 3.23346)
+    cmds.setAttr("brooch_back_r_anim_grp.rx", -1.9647)
+    cmds.setAttr("brooch_back_r_anim_grp.rz", 9.7907)
+    cmds.setDrivenKeyframe(["brooch_front_r_anim_grp.rx", "brooch_back_r_anim_grp.rx", "brooch_front_r_anim_grp.rz", "brooch_back_r_anim_grp.rz"], cd = "driver_clavicle_r.rz", itt = "spline", ott = "spline")
+
+    cmds.setAttr("clavicle_r_anim.ty", 0)    
+
+
+broochRig()
+
+# ------------------------------------------------------------------
+def invertEyelids():
+    sides = ["_l", "_r"]
+    for side in sides:
+        cmds.delete("driver_eyelid_bott"+side+"_parentConstraint1")
+        cmds.rotate(0, -180, 0, "eyelid_bott"+side+"_anim_grp", r=True, os=True)
+        cmds.parentConstraint("eyelid_bott"+side+"_anim", "driver_eyelid_bott"+side, mo=True)
+invertEyelids()
+
+def invertRightEyeball():
+    cmds.delete("driver_eyeball_r_parentConstraint1")
+    cmds.rotate(0, -180, -180, "eyeball_r_anim_grp", r=True, os=True)
+    cmds.parentConstraint("eyeball_r_anim", "driver_eyeball_r", mo=True)
+invertRightEyeball()
+
+# ------------------------------------------------------------------
+def dynamicsOff():
+    for one in ["back_tail_cloth_dyn_anim", "chest_cloth_l_dyn_anim", "chest_cloth_r_dyn_anim", "front_strap_dyn_anim", "ponytail_dyn_anim", "skirt_back_cloth_dyn_anim", "skirt_cloth_l_dyn_anim", "skirt_front_cloth_dyn_anim"]:
+        cmds.setAttr(one+".chainStartEnvelope", 0)
+        cmds.setAttr(one+".chainStartEnvelope", l=True)
+dynamicsOff()
+
+# ------------------------------------------------------------------
+def addCapsule(radius, halfHeight):
+    trueHalfHeight = (halfHeight*2)-(radius*2)
+    capsule = cmds.polyCylinder(name="collision_capsule", r=radius, h=trueHalfHeight, sx=18, sy=1, sz=5, ax=[0,0,1], rcp=True, cuv=3, ch=True)
+    cmds.setAttr(capsule[0]+".tz", halfHeight+2.2)
+    cmds.setAttr(capsule[0]+".rz", 90)
+
+    cmds.addAttr(capsule[0], ln="Radius", dv=radius, keyable=True)
+    cmds.addAttr(capsule[0], ln="HalfHeight", dv=halfHeight, keyable=True)
+
+    expr = capsule[1]+".radius = "+capsule[0]+".Radius;\n"+capsule[1]+".height = ("+capsule[0]+".HalfHeight*2) - ("+capsule[0]+".Radius*2);"
+    cmds.expression(s=expr, o=capsule[1], ae=1, uc="all")
+
+    cmds.setAttr(capsule[0]+".overrideEnabled", 1)
+    cmds.setAttr(capsule[0]+".overrideColor", 18)
+    cmds.setAttr(capsule[0]+".overrideShading", 0)
+
+    cmds.addAttr(capsule[0], ln="Shaded", dv=0, min=0, max=1, keyable=True)
+    cmds.connectAttr(capsule[0]+".Shaded", capsule[0]+".overrideShading")
+
+    cmds.parentConstraint("root", capsule[0], mo=True)
+    cmds.select(capsule[0], r=True)
+    capsuleDisplayLayer = cmds.createDisplayLayer(nr=True, name="Collision Capsule")
+    cmds.setAttr(capsuleDisplayLayer+".displayType", 2)
+    cmds.setAttr(capsuleDisplayLayer+".v", 0)
+    
+    for attr in ["tx", "ty", "tz", "rx", "ry", "rz", "sx", "sy", "sz"]:
+        cmds.setAttr(capsule[0]+"."+attr, lock=True, keyable=False, channelBox=False)
+
+    cmds.select(clear=True)
+addCapsule(42, 86)
+
+# ------------------------------------------------------------------
+def createGameCam():
+    gameCam = cmds.camera(centerOfInterest=5, focalLength=18, lensSqueezeRatio=1, cameraScale=1, horizontalFilmAperture=1.41732, horizontalFilmOffset=0, verticalFilmAperture=0.94488, verticalFilmOffset=0, filmFit="Overscan", overscan=1.05, motionBlur=0, shutterAngle=144, nearClipPlane=1.0, farClipPlane=10000, orthographic=0, orthographicWidth=30, panZoomEnabled=0, horizontalPan=0, verticalPan=0, zoom=1, displayFilmGate=True)[0]
+    print gameCam
+    gameCam = cmds.rename(gameCam, "GameCam")
+    cmds.camera(gameCam, e=True, horizontalFilmAperture=1.6724376, focalLength=21.2399857011)
+    cmds.setAttr(gameCam+".displayGateMaskOpacity", 1)
+    cmds.setAttr(gameCam+".displayGateMaskColor", 0, 0, 0, type="double3")
+    cmds.setAttr(gameCam+".r", 90, 0, 180, type="double3")
+    
+    # Create the control
+    control = cmds.circle(name=gameCam+"_Ctrl", ch=True, o=True, r=15)[0]
+    cmds.makeIdentity(control, a=True, t=True, r=True, s=True)
+    # Create groups for the sdk and the constraints
+    const_grp = cmds.group(control, name=control+"_grp")
+    if cmds.objExists("rig_grp"):
+        cmds.parent(const_grp, "rig_grp")
+    # Parent Constrain the cam to the control
+    cmds.parentConstraint(control, gameCam, mo=True)
+    
+    cmds.connectAttr(control+".visibility", gameCam+".visibility", f=True)
+    cmds.setAttr(const_grp+".v", 0)
+    
+    cmds.setAttr(const_grp+".t", 0, 500, 200, type="double3")
+createGameCam()
+
+# ------------------------------------------------------------------
+# add ik bones Function that is called later after we open the export file.
+def addIKBones():
+    try:
+		# create the joints
+		cmds.select(clear = True)
+		ikFootRoot = cmds.joint(name = "ik_foot_root")
+		cmds.select(clear = True)
+		
+		cmds.select(clear = True)
+		ikFootLeft = cmds.joint(name = "ik_foot_l")
+		cmds.select(clear = True)
+		
+		cmds.select(clear = True)
+		ikFootRight = cmds.joint(name = "ik_foot_r")
+		cmds.select(clear = True)
+		
+		cmds.select(clear = True)
+		ikHandRoot = cmds.joint(name = "ik_hand_root")
+		cmds.select(clear = True)
+		
+		cmds.select(clear = True)
+		ikHandGun = cmds.joint(name = "ik_hand_gun")
+		cmds.select(clear = True)
+		
+		cmds.select(clear = True)
+		ikHandLeft = cmds.joint(name = "ik_hand_l")
+		cmds.select(clear = True)
+		
+		cmds.select(clear = True)
+		ikHandRight = cmds.joint(name = "ik_hand_r")
+		cmds.select(clear = True)
+		
+		# create hierarhcy
+		cmds.parent(ikFootRoot, "root")
+		cmds.parent(ikHandRoot, "root")
+		cmds.parent(ikFootLeft, ikFootRoot)
+		cmds.parent(ikFootRight, ikFootRoot)
+		cmds.parent(ikHandGun, ikHandRoot)
+		cmds.parent(ikHandLeft, ikHandGun)
+		cmds.parent(ikHandRight, ikHandGun)
+		
+		# constrain the joints
+		leftHandConstraint = cmds.parentConstraint("hand_l", ikHandLeft)[0]
+		rightHandGunConstraint = cmds.parentConstraint("hand_r", ikHandGun)[0]
+		rightHandConstraint = cmds.parentConstraint("hand_r", ikHandRight)[0]
+		leftFootConstraint = cmds.parentConstraint("foot_l", ikFootLeft)[0]
+		rightFootConstraint = cmds.parentConstraint("foot_r", ikFootRight)[0]
+	
+    except:
+		cmds.warning("Something went wrong")
+
+# ------------------------------------------------------------------
+# ADD the IK bones to the Export file.
+#save the AnimRig file before opening the export file.
+cmds.file(save = True, type = "mayaBinary", force = True)
+
+#ADD IK BONES
+#open export file
+sceneName = cmds.file(q = True, sceneName = True)
+exportFile = sceneName.partition("AnimRigs")[0] + "ExportFiles" +  sceneName.partition("AnimRigs")[2].partition(".mb")[0] + "_Export.mb"
+cmds.file(exportFile, open = True, force = True)
+
+#add ik bones
+addIKBones()
+
+# ------------------------------------------------------------------
+def addCustomAttrs():
+    attrs = cmds.listAttr("root")
+    exists = False
+    for one in attrs:
+        if one.find("DistanceCurve") != -1:
+            exists = True
+    if exists == False:
+        cmds.addAttr("root", ln="DistanceCurve", keyable=True)
+addCustomAttrs()
+
+# ------------------------------------------------------------------
+#save the export file.
+cmds.file(save = True, type = "mayaBinary", force = True)
+
+
+
+#########################################
+# RUN THIS COMMAND DURING THE RIG POSE CREATION STEP TO GET THE SHIELD TO THE CORRECT LOCATION.  LEAVE IT COMMENTED OUT SO THAT IT DOESNT RUN DURING RIGGING.
+'''
+# The following 2 lines are used to find the locations and then input below.
+#cmds.xform(q=True, ws=True, rp=True)
+#cmds.xform(q=True, ro=True)
+
+
+cmds.xform("arrow_nock_mover", ws=True, t=[9.611692837947271e-14, -2.0994091418383203e-16, 35.41695785522444])
+cmds.xform("arrow_nock_mover", ro=[-13.78281392154952, -64.119833503541, 16.763202342514436])
+
+cmds.xform("bow_base_mover", ws=True, t=[3.3590546269413172e-15, -3.762654592984341e-15, -7.492745461319391e-15])
+cmds.xform("bow_base_mover", ro=[-13.782813921549694, 64.11983350354099, -16.763202342514536])
+
+'''
+#########################################
\ No newline at end of file
diff --git a/Engine/Extras/Maya_AnimationRiggingTools/ARTv1/General/ART/Projects/Orion/Scripts/Steel_CreateLODs.py b/Engine/Extras/Maya_AnimationRiggingTools/ARTv1/General/ART/Projects/Orion/Scripts/Steel_CreateLODs.py
new file mode 100644
--- /dev/null
+++ b/Engine/Extras/Maya_AnimationRiggingTools/ARTv1/General/ART/Projects/Orion/Scripts/Steel_CreateLODs.py
@@ -0,0 +1,524 @@
+import maya.cmds as cmds
+import customMayaMenu as menu
+import Tools.ART_ImportMultipleWeights as impw
+import Tools.ART_TransferWeights as transfer
+import Tools.ART_ExportLODtoFBX as export
+reload(export)
+
+################################################## STEEL #########################################################################
+# Warn the user that this operation will save their file.
+saveResult = cmds.confirmDialog(title = "Save Warning", message="In order to continue your file must be saved.  Would you like to save it?  If yes it will be saved and after the operation is complete your file will be re-opened.", button = ["Yes", "Cancel"], defaultButton='Yes', cancelButton='Cancel', dismissString='Cancel')
+
+# If they are ok with saving the file, warn them that it will also delete all of the current files in their skin weights directory.
+if saveResult == "Yes":
+    deleteSkinFilesResult = cmds.confirmDialog(title = "SkinWeights Warning", message="In order to continue we must delete all of the files in your ..\General\ART\SkinWeights\ directory.  Would you like to continue?.", button = ["Yes", "Cancel"], defaultButton='Yes', cancelButton='Cancel', dismissString='Cancel')
+    if deleteSkinFilesResult == "Yes":
+        toolsPath = cmds.internalVar(usd = True) + "mayaTools.txt"
+        if os.path.exists(toolsPath):
+            f = open(toolsPath, 'r')
+            mayaToolsDir = f.readline()
+            f.close()
+
+        path = mayaToolsDir + "\General\ART\SkinWeights\\"
+        files = os.listdir(path)
+        #delete all weight cache files in the system folder
+        for file in files:
+            os.remove(path + file)
+            #print path + file
+
+# Steel LOD1 -------------------------------------------------------------------------------------------------------------
+def SteelLOD1():
+    if saveResult == "Yes":
+       
+        cmds.file(save=True)
+
+        meshes = ["geo_acc_LOD1", "geo_achiles_LOD1", "geo_arm_shield_bott_l_LOD1", "geo_arm_shield_top_l_LOD1", "geo_boot_cover_LOD1", "geo_chest_plate_LOD1", "geo_eyeballs_LOD1", "geo_fingers_base_LOD1", "geo_fingers_tips_LOD1", "geo_forearms_LOD1", "geo_hand_plates_LOD1", "geo_head_LOD1", "geo_heels_LOD1", "geo_helmet_LOD1", "geo_knee_pads_LOD1", "geo_legs_LOD1", "geo_lower_teeth_LOD1", "geo_nose_ring_LOD1", "geo_shield_LOD1", "geo_shins_LOD1", "geo_shoulders_LOD1", "geo_straps_LOD1", "geo_toe_plates_LOD1", "geo_toes_LOD1", "geo_tongue_LOD1", "geo_torso_LOD1", "geo_turbine_center_LOD1", "geo_turbine_outer_LOD1", "geo_upper_teeth_LOD1"]
+
+        jointsToDelete = []
+
+        # remove unneeded joints
+        jointsToRemove = ["loarm_shield_back_l", "lowerarm_twist_01_l"]
+        jointToTransferTo = "lowerarm_l"
+        transfer.LOD_transferWeights(meshes, jointsToRemove, jointToTransferTo)
+        jointsToDelete.extend(jointsToRemove)
+
+        jointsToRemove = ["lowerarm_twist_01_r"]
+        jointToTransferTo = "lowerarm_r"
+        transfer.LOD_transferWeights(meshes, jointsToRemove, jointToTransferTo)
+        jointsToDelete.extend(jointsToRemove)
+
+        jointsToRemove = ["lowerback_muscle_l", "lowerback_muscle_r"]
+        jointToTransferTo = "spine_02"
+        transfer.LOD_transferWeights(meshes, jointsToRemove, jointToTransferTo)
+        jointsToDelete.extend(jointsToRemove)
+
+        jointsToRemove = ["tricep_l"]
+        jointToTransferTo = "upperarm_twist_02_l"
+        transfer.LOD_transferWeights(meshes, jointsToRemove, jointToTransferTo)
+        jointsToDelete.extend(jointsToRemove)
+        
+        jointsToRemove = ["tricep_r"]
+        jointToTransferTo = "upperarm_twist_02_r"
+        transfer.LOD_transferWeights(meshes, jointsToRemove, jointToTransferTo)
+        jointsToDelete.extend(jointsToRemove)
+        
+        jointsToRemove = ["toe_l"]
+        jointToTransferTo = "foot_l"
+        transfer.LOD_transferWeights(meshes, jointsToRemove, jointToTransferTo)
+        jointsToDelete.extend(jointsToRemove)
+        
+        jointsToRemove = ["toe_r"]
+        jointToTransferTo = "foot_r"
+        transfer.LOD_transferWeights(meshes, jointsToRemove, jointToTransferTo)
+        jointsToDelete.extend(jointsToRemove)
+        
+        jointsToRemove = ["calf_armor_l"]
+        jointToTransferTo = "calf_l"
+        transfer.LOD_transferWeights(meshes, jointsToRemove, jointToTransferTo)
+        jointsToDelete.extend(jointsToRemove)
+        
+        jointsToRemove = ["calf_armor_r"]
+        jointToTransferTo = "calf_r"
+        transfer.LOD_transferWeights(meshes, jointsToRemove, jointToTransferTo)
+        jointsToDelete.extend(jointsToRemove)
+        
+        # Delete the joints we no longer want.
+        for one in jointsToDelete:
+            if cmds.objExists(one):
+                cmds.delete(one)
+                
+        # Export the fbx file.
+        export.exportLOD("Steel_LOD1.fbx", "LOD1")
+        
+        # Re-open the file without saving.
+        fullPath = cmds.file(q = True, sceneName = True)
+        cmds.file(fullPath, open=True, f=True)
+
+        print "LOD1 Exported Successfully!!!!!!!!!!!!!!!"
+
+SteelLOD1()
+
+# Steel LOD2 -------------------------------------------------------------------------------------------------------------
+def SteelLOD2():
+    if saveResult == "Yes":
+        
+        cmds.file(save=True)
+        
+        meshes = ["geo_acc_LOD2", "geo_achiles_LOD2", "geo_arm_shield_bott_l_LOD2", "geo_arm_shield_top_l_LOD2", "geo_boot_cover_LOD2", "geo_chest_plate_LOD2", "geo_eyeballs_LOD2", "geo_fingers_base_LOD2", "geo_fingers_tips_LOD2", "geo_forearms_LOD2", "geo_hand_plates_LOD2", "geo_head_LOD2", "geo_heels_LOD2", "geo_helmet_LOD2", "geo_knee_pads_LOD2", "geo_legs_LOD2", "geo_shield_LOD2", "geo_shins_LOD2", "geo_shoulders_LOD2", "geo_straps_LOD2", "geo_toe_plates_LOD2", "geo_toes_LOD2", "geo_torso_LOD2", "geo_turbine_center_LOD2", "geo_turbine_outer_LOD2"]
+
+        jointsToDelete = []
+
+         # remove unneeded joints
+        jointsToRemove = ["loarm_shield_back_l", "loarm_shield_bot_l", "loarm_shield_top_l", "lowerarm_twist_01_l"]
+        jointToTransferTo = "lowerarm_l"
+        transfer.LOD_transferWeights(meshes, jointsToRemove, jointToTransferTo)
+        jointsToDelete.extend(jointsToRemove)
+
+        jointsToRemove = ["lowerback_muscle_l", "lowerback_muscle_r"]
+        jointToTransferTo = "spine_02"
+        transfer.LOD_transferWeights(meshes, jointsToRemove, jointToTransferTo)
+        jointsToDelete.extend(jointsToRemove)
+
+        jointsToRemove = ["tricep_l"]
+        jointToTransferTo = "upperarm_twist_02_l"
+        transfer.LOD_transferWeights(meshes, jointsToRemove, jointToTransferTo)
+        jointsToDelete.extend(jointsToRemove)
+
+        jointsToRemove = ["lowerarm_twist_01_r"]
+        jointToTransferTo = "lowerarm_r"
+        transfer.LOD_transferWeights(meshes, jointsToRemove, jointToTransferTo)
+        jointsToDelete.extend(jointsToRemove)
+
+        jointsToRemove = ["tricep_r"]
+        jointToTransferTo = "upperarm_twist_02_r"
+        transfer.LOD_transferWeights(meshes, jointsToRemove, jointToTransferTo)
+        jointsToDelete.extend(jointsToRemove)
+
+        jointsToRemove = ["football_pads_back", "turbine_back_main", "turbine_back_inner", "football_pads_front", "chest_strap_l", "chest_strap_r"]
+        jointToTransferTo = "spine_03"
+        transfer.LOD_transferWeights(meshes, jointsToRemove, jointToTransferTo)
+        jointsToDelete.extend(jointsToRemove)
+
+        jointsToRemove = ["eyeball_l", "eyeball_r", "eyebrow_in_l", "eyebrow_in_r", "eyebrow_out_l", "eyebrow_out_r", "eyelid_bott_l", "eyelid_bott_r", "eyelid_top_l", "eyelid_top_r", "jaw"]
+        jointToTransferTo = "head"
+        transfer.LOD_transferWeights(meshes, jointsToRemove, jointToTransferTo)
+        jointsToDelete.extend(jointsToRemove)
+
+        jointsToRemove = ["flap_back", "flap_front", "hip_strap_l", "hip_strap_r", "tail_pad"]
+        jointToTransferTo = "pelvis"
+        transfer.LOD_transferWeights(meshes, jointsToRemove, jointToTransferTo)
+        jointsToDelete.extend(jointsToRemove)
+
+        jointsToRemove = ["toe_plate_l"]
+        jointToTransferTo = "ball_l"
+        transfer.LOD_transferWeights(meshes, jointsToRemove, jointToTransferTo)
+        jointsToDelete.extend(jointsToRemove)
+
+        jointsToRemove = ["calf_armor_l", "shinpad_l"]
+        jointToTransferTo = "calf_l"
+        transfer.LOD_transferWeights(meshes, jointsToRemove, jointToTransferTo)
+        jointsToDelete.extend(jointsToRemove)
+
+        jointsToRemove = ["toe_l"]
+        jointToTransferTo = "foot_l"
+        transfer.LOD_transferWeights(meshes, jointsToRemove, jointToTransferTo)
+        jointsToDelete.extend(jointsToRemove)
+
+        jointsToRemove = ["thigh_twist_01_l", "thigh_twist_02_l"]
+        jointToTransferTo = "thigh_l"
+        transfer.LOD_transferWeights(meshes, jointsToRemove, jointToTransferTo)
+        jointsToDelete.extend(jointsToRemove)
+
+        jointsToRemove = ["shinpad_r", "calf_armor_r"]
+        jointToTransferTo = "calf_r"
+        transfer.LOD_transferWeights(meshes, jointsToRemove, jointToTransferTo)
+        jointsToDelete.extend(jointsToRemove)
+
+        jointsToRemove = ["toe_plate_r"]
+        jointToTransferTo = "ball_r"
+        transfer.LOD_transferWeights(meshes, jointsToRemove, jointToTransferTo)
+        jointsToDelete.extend(jointsToRemove)
+
+        jointsToRemove = ["toe_r"]
+        jointToTransferTo = "foot_r"
+        transfer.LOD_transferWeights(meshes, jointsToRemove, jointToTransferTo)
+        jointsToDelete.extend(jointsToRemove)
+
+        jointsToRemove = ["thigh_twist_01_r", "thigh_twist_02_r"]
+        jointToTransferTo = "thigh_r"
+        transfer.LOD_transferWeights(meshes, jointsToRemove, jointToTransferTo)
+        jointsToDelete.extend(jointsToRemove)
+
+        # Delete the joints we no longer want.
+        for one in jointsToDelete:
+            if cmds.objExists(one):
+                cmds.delete(one)
+                
+        # Export the fbx file.
+        export.exportLOD("Steel_LOD2.fbx", "LOD2")
+
+        # Re-open the file without saving.
+        fullPath = cmds.file(q = True, sceneName = True)
+        cmds.file(fullPath, open=True, f=True)
+
+        print "LOD2 Exported Successfully!!!!!!!!!!!!!!!"
+
+SteelLOD2()
+
+# Steel LOD3 -------------------------------------------------------------------------------------------------------------
+def SteelLOD3():
+    if saveResult == "Yes":
+
+        cmds.file(save=True)
+    
+        meshes = ["geo_achiles_LOD3", "geo_arm_shield_bott_l_LOD3", "geo_arm_shield_top_l_LOD3", "geo_boot_cover_LOD3", "geo_chest_plate_LOD3", "geo_fingers_base_LOD3", "geo_fingers_tips_LOD3", "geo_forearms_LOD3", "geo_hand_plates_LOD3", "geo_head_LOD3", "geo_heels_LOD3", "geo_helmet_LOD3", "geo_knee_pads_LOD3", "geo_legs_LOD3", "geo_shield_LOD3", "geo_shins_LOD3", "geo_shoulders_LOD3", "geo_toe_plates_LOD3", "geo_toes_LOD3", "geo_torso_LOD3", "geo_turbine_center_LOD3", "geo_turbine_outer_LOD3"]
+    
+        # Export a skin weights file for each mesh we plan to operate on.
+        cmds.select(meshes)
+        menu.exportSelectedSkinWeights(True)
+
+        # Set the pose if needed
+        cmds.setAttr("thumb_02_r.rotateZ", 14.0)
+        cmds.setAttr("ring_02_r.rotateZ", 45.0)
+        cmds.setAttr("pinky_02_r.rotateZ", 45.0)
+        cmds.setAttr("middle_02_r.rotateZ", 45.0)
+        cmds.setAttr("index_02_r.rotateZ", 45.0)
+        cmds.setAttr("pinky_02_l.rotateZ", 45.0)
+        cmds.setAttr("middle_02_l.rotateZ", 45.0)
+        cmds.setAttr("index_02_l.rotateZ", 45.0)
+        cmds.setAttr("ring_02_l.rotateZ", 45.0)
+        cmds.setAttr("thumb_02_l.rotateZ", 14.0)
+
+        # Delete the history on the meshes so that they get baked down into their posed positions
+        cmds.delete(meshes, ch=True)
+
+        # Move the skeleton back to its bindpose (This script REQUIRES the bind pose to be the pose that we want to export the skeleton to.  This is the case for all known characters created with A.R.T.)
+        cmds.dagPose("root", bp=True, r=True) 
+
+        # Re-import the skinning we exported before onto the mesh so thats its bound in its posed position.
+        cmds.select(meshes)
+        impw.importMultipleSkinWeights("", True)
+        
+        jointsToDelete = []
+
+        # remove unneeded joints
+        jointsToRemove = ["flap_back", "flap_front", "hip_strap_l", "hip_strap_r", "tail_pad"]
+        jointToTransferTo = "pelvis"
+        transfer.LOD_transferWeights(meshes, jointsToRemove, jointToTransferTo)
+        jointsToDelete.extend(jointsToRemove)
+
+        jointsToRemove = ["lowerback_muscle_l", "lowerback_muscle_r"]
+        jointToTransferTo = "spine_02"
+        transfer.LOD_transferWeights(meshes, jointsToRemove, jointToTransferTo)
+        jointsToDelete.extend(jointsToRemove)
+
+        jointsToRemove = ["football_pads_back", "turbine_back_main", "turbine_back_inner", "football_pads_front", "chest_strap_l", "chest_strap_r"]
+        jointToTransferTo = "spine_03"
+        transfer.LOD_transferWeights(meshes, jointsToRemove, jointToTransferTo)
+        jointsToDelete.extend(jointsToRemove)
+
+        jointsToRemove = ["eyeball_l", "eyeball_r", "eyebrow_in_l", "eyebrow_in_r", "eyebrow_out_l", "eyebrow_out_r", "eyelid_bott_l", "eyelid_bott_r", "eyelid_top_l", "eyelid_top_r", "jaw"]
+        jointToTransferTo = "head"
+        transfer.LOD_transferWeights(meshes, jointsToRemove, jointToTransferTo)
+        jointsToDelete.extend(jointsToRemove)
+
+        jointsToRemove = ["upperarm_twist_01_l", "upperarm_twist_02_l", "arm_strap_bott_l", "arm_strap_top_l", "tricep_l"]
+        jointToTransferTo = "upperarm_l"
+        transfer.LOD_transferWeights(meshes, jointsToRemove, jointToTransferTo)
+        jointsToDelete.extend(jointsToRemove)
+
+        jointsToRemove = ["upperarm_twist_01_r", "upperarm_twist_02_r", "arm_strap_bott_r", "arm_strap_top_r", "tricep_r"]
+        jointToTransferTo = "upperarm_r"
+        transfer.LOD_transferWeights(meshes, jointsToRemove, jointToTransferTo)
+        jointsToDelete.extend(jointsToRemove)
+
+        jointsToRemove = ["loarm_shield_back_l", "loarm_shield_bot_l", "loarm_shield_top_l", "lowerarm_twist_01_l"]
+        jointToTransferTo = "lowerarm_l"
+        transfer.LOD_transferWeights(meshes, jointsToRemove, jointToTransferTo)
+        jointsToDelete.extend(jointsToRemove)
+
+        jointsToRemove = ["lowerarm_twist_01_r"]
+        jointToTransferTo = "lowerarm_r"
+        transfer.LOD_transferWeights(meshes, jointsToRemove, jointToTransferTo)
+        jointsToDelete.extend(jointsToRemove)
+
+        jointsToRemove = ["index_02_l"]
+        jointToTransferTo = "index_01_l"
+        transfer.LOD_transferWeights(meshes, jointsToRemove, jointToTransferTo)
+        jointsToDelete.extend(jointsToRemove)
+
+        jointsToRemove = ["middle_02_l"]
+        jointToTransferTo = "middle_01_l"
+        transfer.LOD_transferWeights(meshes, jointsToRemove, jointToTransferTo)
+        jointsToDelete.extend(jointsToRemove)
+
+        jointsToRemove = ["pinky_02_l"]
+        jointToTransferTo = "pinky_01_l"
+        transfer.LOD_transferWeights(meshes, jointsToRemove, jointToTransferTo)
+        jointsToDelete.extend(jointsToRemove)
+
+        jointsToRemove = ["ring_02_l"]
+        jointToTransferTo = "ring_01_l"
+        transfer.LOD_transferWeights(meshes, jointsToRemove, jointToTransferTo)
+        jointsToDelete.extend(jointsToRemove)
+
+        jointsToRemove = ["thumb_02_l"]
+        jointToTransferTo = "thumb_01_l"
+        transfer.LOD_transferWeights(meshes, jointsToRemove, jointToTransferTo)
+        jointsToDelete.extend(jointsToRemove)
+
+        jointsToRemove = ["index_02_r"]
+        jointToTransferTo = "index_01_r"
+        transfer.LOD_transferWeights(meshes, jointsToRemove, jointToTransferTo)
+        jointsToDelete.extend(jointsToRemove)
+
+        jointsToRemove = ["middle_02_r"]
+        jointToTransferTo = "middle_01_r"
+        transfer.LOD_transferWeights(meshes, jointsToRemove, jointToTransferTo)
+        jointsToDelete.extend(jointsToRemove)
+
+        jointsToRemove = ["pinky_02_r"]
+        jointToTransferTo = "pinky_01_r"
+        transfer.LOD_transferWeights(meshes, jointsToRemove, jointToTransferTo)
+        jointsToDelete.extend(jointsToRemove)
+
+        jointsToRemove = ["ring_02_r"]
+        jointToTransferTo = "ring_01_r"
+        transfer.LOD_transferWeights(meshes, jointsToRemove, jointToTransferTo)
+        jointsToDelete.extend(jointsToRemove)
+
+        jointsToRemove = ["thumb_02_r"]
+        jointToTransferTo = "thumb_01_r"
+        transfer.LOD_transferWeights(meshes, jointsToRemove, jointToTransferTo)
+        jointsToDelete.extend(jointsToRemove)
+        
+        jointsToRemove = ["thigh_twist_01_l", "thigh_twist_02_l"]
+        jointToTransferTo = "thigh_l"
+        transfer.LOD_transferWeights(meshes, jointsToRemove, jointToTransferTo)
+        jointsToDelete.extend(jointsToRemove)
+
+        jointsToRemove = ["calf_armor_l", "shinpad_l"]
+        jointToTransferTo = "calf_l"
+        transfer.LOD_transferWeights(meshes, jointsToRemove, jointToTransferTo)
+        jointsToDelete.extend(jointsToRemove)
+
+        jointsToRemove = ["toe_plate_l", "toe_l"]
+        jointToTransferTo = "foot_l"
+        transfer.LOD_transferWeights(meshes, jointsToRemove, jointToTransferTo)
+        jointsToDelete.extend(jointsToRemove)
+
+        jointsToRemove = ["thigh_twist_01_r", "thigh_twist_02_r"]
+        jointToTransferTo = "thigh_r"
+        transfer.LOD_transferWeights(meshes, jointsToRemove, jointToTransferTo)
+        jointsToDelete.extend(jointsToRemove)
+
+        jointsToRemove = ["calf_armor_r", "shinpad_r"]
+        jointToTransferTo = "calf_r"
+        transfer.LOD_transferWeights(meshes, jointsToRemove, jointToTransferTo)
+        jointsToDelete.extend(jointsToRemove)
+
+        jointsToRemove = ["toe_plate_r", "toe_r"]
+        jointToTransferTo = "foot_r"
+        transfer.LOD_transferWeights(meshes, jointsToRemove, jointToTransferTo)
+        jointsToDelete.extend(jointsToRemove)
+
+        # Delete the joints we no longer want.
+        for one in jointsToDelete:
+            if cmds.objExists(one):
+                cmds.delete(one)
+                
+        # Export the fbx file.
+        export.exportLOD("Steel_LOD3.fbx", "LOD3")
+        
+        # Re-open the file without saving.
+        fullPath = cmds.file(q = True, sceneName = True)
+        cmds.file(fullPath, open=True, f=True)
+
+        print "LOD3 Exported Successfully!!!!!!!!!!!!!!!"
+
+SteelLOD3()
+
+# Steel LOD4 -------------------------------------------------------------------------------------------------------------
+def SteelLOD4():
+    if saveResult == "Yes":
+
+        cmds.file(save=True)
+    
+        meshes = ["geo_acc_LOD4", "geo_achiles_LOD4", "geo_arm_shield_bott_l_LOD4", "geo_arm_shield_top_l_LOD4", "geo_boot_cover_LOD4", "geo_chest_plate_LOD4", "geo_fingers_base_LOD4", "geo_fingers_tips_LOD4", "geo_forearms_LOD4", "geo_hand_plates_LOD4", "geo_head_LOD4", "geo_heels_LOD4", "geo_helmet_LOD4", "geo_knee_pads_LOD4", "geo_legs_LOD4", "geo_shield_LOD4", "geo_shins_LOD4", "geo_shoulders_LOD4", "geo_toe_plates_LOD4", "geo_toes_LOD4", "geo_torso_LOD4"]
+
+        # Export a skin weights file for each mesh we plan to operate on.
+        cmds.select(meshes)
+        menu.exportSelectedSkinWeights(True)
+
+        # Set the pose if needed
+        cmds.setAttr("thumb_02_r.rotateZ", 27.33517469)
+        cmds.setAttr("thumb_01_r.rotateZ", 27.33517469)
+        cmds.setAttr("ring_02_r.rotateZ", 81.26462407)
+        cmds.setAttr("ring_01_r.rotateZ", 81.26462407)
+        cmds.setAttr("pinky_02_r.rotateZ", 81.26462407)
+        cmds.setAttr("pinky_01_r.rotateZ", 81.26462407)
+        cmds.setAttr("middle_02_r.rotateZ", 81.26462407)
+        cmds.setAttr("middle_01_r.rotateZ", 81.26462407)
+        cmds.setAttr("index_02_r.rotateZ", 81.26462407)
+        cmds.setAttr("index_01_r.rotateZ", 81.26462407)
+        cmds.setAttr("pinky_02_l.rotateZ", 77.13683837)
+        cmds.setAttr("pinky_01_l.rotateZ", 77.13683837)
+        cmds.setAttr("middle_02_l.rotateZ", 77.13683837)
+        cmds.setAttr("middle_01_l.rotateZ", 77.13683837)
+        cmds.setAttr("index_02_l.rotateZ", 77.13683837)
+        cmds.setAttr("index_01_l.rotateZ", 77.13683837)
+        cmds.setAttr("ring_02_l.rotateZ", 77.13683837)
+        cmds.setAttr("ring_01_l.rotateZ", 77.13683837)
+        cmds.setAttr("thumb_02_l.rotateZ", 31.92338728)
+        cmds.setAttr("thumb_01_l.rotateZ", 31.92338728)
+
+        # Delete the history on the meshes so that they get baked down into their posed positions
+        cmds.delete(meshes, ch=True)
+
+        # Move the skeleton back to its bindpose (This script REQUIRES the bind pose to be the pose that we want to export the skeleton to.  This is the case for all known characters created with A.R.T.)
+        cmds.dagPose("root", bp=True, r=True) 
+
+        # Re-import the skinning we exported before onto the mesh so thats its bound in its posed position.
+        cmds.select(meshes)
+        impw.importMultipleSkinWeights("", True)
+        
+        jointsToDelete = []
+
+        # remove unneeded joints
+        jointsToRemove = ["flap_back", "flap_front", "hip_strap_l", "hip_strap_r", "tail_pad"]
+        jointToTransferTo = "pelvis"
+        transfer.LOD_transferWeights(meshes, jointsToRemove, jointToTransferTo)
+        jointsToDelete.extend(jointsToRemove)
+
+        jointsToRemove = ["lowerback_muscle_l", "lowerback_muscle_r"]
+        jointToTransferTo = "spine_02"
+        transfer.LOD_transferWeights(meshes, jointsToRemove, jointToTransferTo)
+        jointsToDelete.extend(jointsToRemove)
+
+        jointsToRemove = ["football_pads_back", "turbine_back_main", "turbine_back_inner", "football_pads_front", "chest_strap_l", "chest_strap_r"]
+        jointToTransferTo = "spine_03"
+        transfer.LOD_transferWeights(meshes, jointsToRemove, jointToTransferTo)
+        jointsToDelete.extend(jointsToRemove)
+
+        jointsToRemove = ["eyeball_l", "eyeball_r", "eyebrow_in_l", "eyebrow_in_r", "eyebrow_out_l", "eyebrow_out_r", "eyelid_bott_l", "eyelid_bott_r", "eyelid_top_l", "eyelid_top_r", "jaw"]
+        jointToTransferTo = "head"
+        transfer.LOD_transferWeights(meshes, jointsToRemove, jointToTransferTo)
+        jointsToDelete.extend(jointsToRemove)
+
+        jointsToRemove = ["upperarm_twist_01_l", "upperarm_twist_02_l", "arm_strap_bott_l", "arm_strap_top_l", "tricep_l"]
+        jointToTransferTo = "upperarm_l"
+        transfer.LOD_transferWeights(meshes, jointsToRemove, jointToTransferTo)
+        jointsToDelete.extend(jointsToRemove)
+
+        jointsToRemove = ["upperarm_twist_01_r", "upperarm_twist_02_r", "arm_strap_bott_r", "arm_strap_top_r", "tricep_r"]
+        jointToTransferTo = "upperarm_r"
+        transfer.LOD_transferWeights(meshes, jointsToRemove, jointToTransferTo)
+        jointsToDelete.extend(jointsToRemove)
+
+        jointsToRemove = ["loarm_shield_back_l", "loarm_shield_bot_l", "loarm_shield_top_l", "lowerarm_twist_01_l"]
+        jointToTransferTo = "lowerarm_l"
+        transfer.LOD_transferWeights(meshes, jointsToRemove, jointToTransferTo)
+        jointsToDelete.extend(jointsToRemove)
+
+        jointsToRemove = ["lowerarm_twist_01_r"]
+        jointToTransferTo = "lowerarm_r"
+        transfer.LOD_transferWeights(meshes, jointsToRemove, jointToTransferTo)
+        jointsToDelete.extend(jointsToRemove)
+
+        jointsToRemove = ["index_01_l", "index_02_l", "middle_01_l", "middle_02_l", "pinky_01_l", "pinky_02_l", "ring_01_l", "ring_02_l", "thumb_01_l", "thumb_02_l"]
+        jointToTransferTo = "hand_l"
+        transfer.LOD_transferWeights(meshes, jointsToRemove, jointToTransferTo)
+        jointsToDelete.extend(jointsToRemove)
+
+        jointsToRemove = ["index_01_r", "index_02_r", "middle_01_r", "middle_02_r", "pinky_01_r", "pinky_02_r", "ring_01_r", "ring_02_r", "thumb_01_r", "thumb_02_r"]
+        jointToTransferTo = "hand_r"
+        transfer.LOD_transferWeights(meshes, jointsToRemove, jointToTransferTo)
+        jointsToDelete.extend(jointsToRemove)
+
+        jointsToRemove = ["thigh_twist_01_l", "thigh_twist_02_l"]
+        jointToTransferTo = "thigh_l"
+        transfer.LOD_transferWeights(meshes, jointsToRemove, jointToTransferTo)
+        jointsToDelete.extend(jointsToRemove)
+
+        jointsToRemove = ["calf_armor_l", "shinpad_l"]
+        jointToTransferTo = "calf_l"
+        transfer.LOD_transferWeights(meshes, jointsToRemove, jointToTransferTo)
+        jointsToDelete.extend(jointsToRemove)
+
+        jointsToRemove = ["ball_l", "toe_plate_l", "toe_l"]
+        jointToTransferTo = "foot_l"
+        transfer.LOD_transferWeights(meshes, jointsToRemove, jointToTransferTo)
+        jointsToDelete.extend(jointsToRemove)
+
+        jointsToRemove = ["thigh_twist_01_r", "thigh_twist_02_r"]
+        jointToTransferTo = "thigh_r"
+        transfer.LOD_transferWeights(meshes, jointsToRemove, jointToTransferTo)
+        jointsToDelete.extend(jointsToRemove)
+
+        jointsToRemove = ["calf_armor_r", "shinpad_r"]
+        jointToTransferTo = "calf_r"
+        transfer.LOD_transferWeights(meshes, jointsToRemove, jointToTransferTo)
+        jointsToDelete.extend(jointsToRemove)
+
+        jointsToRemove = ["ball_r", "toe_plate_r", "toe_r"]
+        jointToTransferTo = "foot_r"
+        transfer.LOD_transferWeights(meshes, jointsToRemove, jointToTransferTo)
+        jointsToDelete.extend(jointsToRemove)
+        
+        # Delete the joints we no longer want.
+        for one in jointsToDelete:
+            if cmds.objExists(one):
+                cmds.delete(one)
+                
+        # Export the fbx file.
+        export.exportLOD("Steel_LOD4.fbx", "LOD4")
+        
+        # Re-open the file without saving.
+        fullPath = cmds.file(q = True, sceneName = True)
+        cmds.file(fullPath, open=True, f=True)
+
+        print "LOD4 Exported Successfully!!!!!!!!!!!!!!!"
+
+SteelLOD4()
+
+
diff --git a/Engine/Extras/Maya_AnimationRiggingTools/ARTv1/General/ART/Projects/Orion/Scripts/Steel_Post_Script.py b/Engine/Extras/Maya_AnimationRiggingTools/ARTv1/General/ART/Projects/Orion/Scripts/Steel_Post_Script.py
new file mode 100644
--- /dev/null
+++ b/Engine/Extras/Maya_AnimationRiggingTools/ARTv1/General/ART/Projects/Orion/Scripts/Steel_Post_Script.py
@@ -0,0 +1,656 @@
+import maya.cmds as cmds
+
+# NOTE: This character has special steps to run during Rig Pose creation.  See end of script.
+# ------------------------------------------------------------------
+def shieldRig():
+    #Move Shield to correct position
+    cmds.setAttr("sheild_main_anim_grp.tx", 190.409*0.65)
+    cmds.setAttr("sheild_main_anim_grp.ty", -101.619*0.65)
+    cmds.setAttr("sheild_main_anim_grp.tz", -130.408*0.65)
+    cmds.setAttr("sheild_main_anim_grp.ry", 90)
+
+    #Shield Anim Control Scale
+    cmds.select("sheild_main_anim" + ".cv[0:32]")
+    cmds.scale(6*0.65, 6*0.65, 0, r=True)
+
+    #add attrs to shield anim
+    attrsToAdd = ["shieldDeploy1", "shieldDeploy2", "shieldDeploy3", "shieldDeploy4", "panelA_deploy1", "panelA_deploy2", "panelB_deploy1", "panelB_deploy2", "panelC_deploy1", "panelC_deploy2", "panelD_deploy1", "panelD_deploy2"] 
+    for attr in attrsToAdd:
+        cmds.addAttr("sheild_main_anim", ln = attr, dv = 0, min = 0, keyable = True)
+
+    #Shield SDKS
+    #deploy1
+    cmds.setAttr("sheild_main_anim.shieldDeploy1", 0)
+    cmds.setDrivenKeyframe(["loarm_shield_top_l_anim.tz", "loarm_shield_top_l_anim.ry", "loarm_shield_bot_l_anim.tz", "loarm_shield_bot_l_anim.ry"], cd = "sheild_main_anim.shieldDeploy1", itt = "spline", ott = "spline")
+
+    cmds.setAttr("sheild_main_anim.shieldDeploy1", 10)
+    cmds.setAttr("loarm_shield_top_l_anim.tz", 13.22*0.65)
+    cmds.setAttr("loarm_shield_top_l_anim.ry", -11.846)
+    cmds.setAttr("loarm_shield_bot_l_anim.tz", -14.016*0.65)
+    cmds.setAttr("loarm_shield_bot_l_anim.ry", 14.991)
+    cmds.setDrivenKeyframe(["loarm_shield_top_l_anim.tz", "loarm_shield_top_l_anim.ry", "loarm_shield_bot_l_anim.tz", "loarm_shield_bot_l_anim.ry"], cd = "sheild_main_anim.shieldDeploy1", itt = "spline", ott = "spline")
+    cmds.setAttr("sheild_main_anim.shieldDeploy1", 0)
+
+    cmds.setInfinity("loarm_shield_top_l_anim.tz", pri = "linear", poi = "linear")
+    cmds.setInfinity("loarm_shield_top_l_anim.ry", pri = "linear", poi = "linear")
+    cmds.setInfinity("loarm_shield_bot_l_anim.tz", pri = "linear", poi = "linear")
+    cmds.setInfinity("loarm_shield_bot_l_anim.ry", pri = "linear", poi = "linear")
+
+    cmds.setAttr("loarm_shield_top_l_anim.tz", keyable = False, channelBox = False)
+    cmds.setAttr("loarm_shield_top_l_anim.ry", keyable = False, channelBox = False)
+    cmds.setAttr("loarm_shield_bot_l_anim.tz", keyable = False, channelBox = False)
+    cmds.setAttr("loarm_shield_bot_l_anim.ry", keyable = False, channelBox = False)
+
+    #deploy2
+    cmds.setAttr("sheild_main_anim.shieldDeploy2", 0)
+    cmds.setDrivenKeyframe(["sheild_main_anim_grp.ty"], cd = "sheild_main_anim.shieldDeploy2", itt = "spline", ott = "spline")
+
+    cmds.setAttr("sheild_main_anim.shieldDeploy2", 10)
+    cmds.setAttr("sheild_main_anim_grp.ty", -55*0.65)
+    cmds.setDrivenKeyframe(["sheild_main_anim_grp.ty"], cd = "sheild_main_anim.shieldDeploy2", itt = "spline", ott = "spline")
+    cmds.setAttr("sheild_main_anim.shieldDeploy2", 0)
+    cmds.setInfinity("sheild_main_anim_grp.ty", pri = "linear", poi = "linear")
+    cmds.setAttr("sheild_main_anim_grp.ty", keyable = False, channelBox = False)
+
+    #deploy3
+    cmds.setAttr("sheild_main_anim.shieldDeploy3", 10)
+    cmds.setDrivenKeyframe(["panel_bot_root_l_anim.sy", "panel_bot_root_r_anim.sy", "panel_top_root_l_anim.sy", "panel_top_root_r_anim.sy"], cd = "sheild_main_anim.shieldDeploy3", itt = "spline", ott = "spline")
+
+    cmds.setAttr("sheild_main_anim.shieldDeploy3", 0)
+
+    cmds.setAttr("panel_bot_root_l_anim.sy", .5*0.65)
+    cmds.setAttr("panel_bot_root_r_anim.sy", .5*0.65)
+    cmds.setAttr("panel_top_root_l_anim.sy", .5*0.65)
+    cmds.setAttr("panel_top_root_r_anim.sy", .5*0.65)
+    cmds.setDrivenKeyframe(["panel_bot_root_l_anim.sy", "panel_bot_root_r_anim.sy", "panel_top_root_l_anim.sy", "panel_top_root_r_anim.sy"], cd = "sheild_main_anim.shieldDeploy3", itt = "spline", ott = "spline")
+    cmds.setInfinity("panel_bot_root_l_anim.sy", pri = "linear", poi = "linear")
+    cmds.setInfinity("panel_bot_root_r_anim.sy", pri = "linear", poi = "linear")
+    cmds.setInfinity("panel_top_root_l_anim.sy", pri = "linear", poi = "linear")
+    cmds.setInfinity("panel_top_root_r_anim.sy", pri = "linear", poi = "linear")
+
+    cmds.setAttr("panel_bot_root_l_anim.sy", keyable = False, channelBox = False)
+    cmds.setAttr("panel_bot_root_r_anim.sy", keyable = False, channelBox = False)
+    cmds.setAttr("panel_top_root_l_anim.sy", keyable = False, channelBox = False)
+    cmds.setAttr("panel_top_root_r_anim.sy", keyable = False, channelBox = False)
+
+    #deploy4
+    cmds.setAttr("sheild_main_anim.shieldDeploy4", 10)
+    cmds.setDrivenKeyframe(["panel_bot_r_anim.sz", "panel_bot_l_anim.sz", "panel_top_r_anim.sz", "panel_top_l_anim.sz"], cd = "sheild_main_anim.shieldDeploy4", itt = "spline", ott = "spline")
+
+    cmds.setAttr("sheild_main_anim.shieldDeploy4", 0)
+
+    cmds.setAttr("panel_bot_r_anim.sz", .3*0.65)
+    cmds.setAttr("panel_bot_l_anim.sz", .3*0.65)
+    cmds.setAttr("panel_top_r_anim.sz", .3*0.65)
+    cmds.setAttr("panel_top_l_anim.sz", .3*0.65)
+    cmds.setDrivenKeyframe(["panel_bot_r_anim.sz", "panel_bot_l_anim.sz", "panel_top_r_anim.sz", "panel_top_l_anim.sz"], cd = "sheild_main_anim.shieldDeploy4", itt = "spline", ott = "spline")
+    cmds.setInfinity("panel_bot_r_anim.sz", pri = "linear", poi = "linear")
+    cmds.setInfinity("panel_bot_l_anim.sz", pri = "linear", poi = "linear")
+    cmds.setInfinity("panel_top_r_anim.sz", pri = "linear", poi = "linear")
+    cmds.setInfinity("panel_top_l_anim.sz", pri = "linear", poi = "linear")
+
+    cmds.setAttr("panel_bot_r_anim.sz", keyable = False, channelBox = False)
+    cmds.setAttr("panel_bot_l_anim.sz", keyable = False, channelBox = False)
+    cmds.setAttr("panel_top_r_anim.sz", keyable = False, channelBox = False)
+    cmds.setAttr("panel_top_l_anim.sz", keyable = False, channelBox = False)
+
+    #panelA 1
+    cmds.setAttr("sheild_main_anim.panelA_deploy1", 10)
+    cmds.setDrivenKeyframe(["panel_bot_root_l_anim.rz"], cd = "sheild_main_anim.panelA_deploy1", itt = "spline", ott = "spline")
+
+    cmds.setAttr("sheild_main_anim.panelA_deploy1", 0)
+    cmds.setAttr("panel_bot_root_l_anim.rz", -90)
+    cmds.setDrivenKeyframe(["panel_bot_root_l_anim.rz"], cd = "sheild_main_anim.panelA_deploy1", itt = "spline", ott = "spline")
+    cmds.setInfinity("panel_bot_root_l_anim.rz", pri = "linear", poi = "linear")
+    cmds.setAttr("panel_bot_root_l_anim.rz", keyable = False, channelBox = False)
+
+    #panelA 2
+    cmds.setAttr("sheild_main_anim.panelA_deploy2", 10)
+    cmds.setDrivenKeyframe(["panel_bot_l_anim.sx"], cd = "sheild_main_anim.panelA_deploy2", itt = "spline", ott = "spline")
+
+    cmds.setAttr("sheild_main_anim.panelA_deploy2", 0)
+    cmds.setAttr("panel_bot_l_anim.sx", 0.05*0.65)
+    cmds.setDrivenKeyframe(["panel_bot_l_anim.sx"], cd = "sheild_main_anim.panelA_deploy2", itt = "spline", ott = "spline")
+    cmds.setInfinity("panel_bot_l_anim.sx", pri = "linear", poi = "linear")
+    cmds.setAttr("panel_bot_l_anim.sx", keyable = False, channelBox = False)
+
+    #panelB 1
+    cmds.setAttr("sheild_main_anim.panelB_deploy1", 10)
+    cmds.setDrivenKeyframe(["panel_top_root_l_anim.rz"], cd = "sheild_main_anim.panelB_deploy1", itt = "spline", ott = "spline")
+
+    cmds.setAttr("sheild_main_anim.panelB_deploy1", 0)
+    cmds.setAttr("panel_top_root_l_anim.rz", -90)
+    cmds.setDrivenKeyframe(["panel_top_root_l_anim.rz"], cd = "sheild_main_anim.panelB_deploy1", itt = "spline", ott = "spline")
+    cmds.setInfinity("panel_top_root_l_anim.rz", pri = "linear", poi = "linear")
+    cmds.setAttr("panel_top_root_l_anim.rz", keyable = False, channelBox = False)
+
+    #panelB 2
+    cmds.setAttr("sheild_main_anim.panelB_deploy2", 10)
+    cmds.setDrivenKeyframe(["panel_top_l_anim.sx"], cd = "sheild_main_anim.panelB_deploy2", itt = "spline", ott = "spline")
+
+    cmds.setAttr("sheild_main_anim.panelB_deploy2", 0)
+    cmds.setAttr("panel_top_l_anim.sx", 0.05*0.65)
+    cmds.setDrivenKeyframe(["panel_top_l_anim.sx"], cd = "sheild_main_anim.panelB_deploy2", itt = "spline", ott = "spline")
+    cmds.setInfinity("panel_top_l_anim.sx", pri = "linear", poi = "linear")
+    cmds.setAttr("panel_top_l_anim.sx", keyable = False, channelBox = False)
+
+    #panelC 1
+    cmds.setAttr("sheild_main_anim.panelC_deploy1", 10)
+    cmds.setDrivenKeyframe(["panel_bot_root_r_anim.rz"], cd = "sheild_main_anim.panelC_deploy1", itt = "spline", ott = "spline")
+
+    cmds.setAttr("sheild_main_anim.panelC_deploy1", 0)
+    cmds.setAttr("panel_bot_root_r_anim.rz", -90)
+    cmds.setDrivenKeyframe(["panel_bot_root_r_anim.rz"], cd = "sheild_main_anim.panelC_deploy1", itt = "spline", ott = "spline")
+    cmds.setInfinity("panel_bot_root_r_anim.rz", pri = "linear", poi = "linear")
+    cmds.setAttr("panel_bot_root_r_anim.rz", keyable = False, channelBox = False)
+
+    #panelC 2
+    cmds.setAttr("sheild_main_anim.panelC_deploy2", 10)
+    cmds.setDrivenKeyframe(["panel_bot_r_anim.sx"], cd = "sheild_main_anim.panelC_deploy2", itt = "spline", ott = "spline")
+
+    cmds.setAttr("sheild_main_anim.panelC_deploy2", 0)
+    cmds.setAttr("panel_bot_r_anim.sx", 0.05*0.65)
+    cmds.setDrivenKeyframe(["panel_bot_r_anim.sx"], cd = "sheild_main_anim.panelC_deploy2", itt = "spline", ott = "spline")
+    cmds.setInfinity("panel_bot_r_anim.sx", pri = "linear", poi = "linear")
+    cmds.setAttr("panel_bot_r_anim.sx", keyable = False, channelBox = False)
+
+    #panelD 1
+    cmds.setAttr("sheild_main_anim.panelD_deploy1", 10)
+    cmds.setDrivenKeyframe(["panel_top_root_r_anim.rz"], cd = "sheild_main_anim.panelD_deploy1", itt = "spline", ott = "spline")
+
+    cmds.setAttr("sheild_main_anim.panelD_deploy1", 0)
+    cmds.setAttr("panel_top_root_r_anim.rz", -90)
+    cmds.setDrivenKeyframe(["panel_top_root_r_anim.rz"], cd = "sheild_main_anim.panelD_deploy1", itt = "spline", ott = "spline")
+    cmds.setInfinity("panel_top_root_r_anim.rz", pri = "linear", poi = "linear")
+    cmds.setAttr("panel_top_root_r_anim.rz", keyable = False, channelBox = False)
+
+    #panelD 2
+    cmds.setAttr("sheild_main_anim.panelD_deploy2", 10)
+    cmds.setDrivenKeyframe(["panel_top_r_anim.sx"], cd = "sheild_main_anim.panelD_deploy2", itt = "spline", ott = "spline")
+
+    cmds.setAttr("sheild_main_anim.panelD_deploy2", 0)
+    cmds.setAttr("panel_top_r_anim.sx", 0.05*0.65)
+    cmds.setDrivenKeyframe(["panel_top_r_anim.sx"], cd = "sheild_main_anim.panelD_deploy2", itt = "spline", ott = "spline")
+    cmds.setInfinity("panel_top_r_anim.sx", pri = "linear", poi = "linear")
+    cmds.setAttr("panel_top_r_anim.sx", keyable = False, channelBox = False)
+
+shieldRig()
+
+# ------------------------------------------------------------------
+def legRig():
+    #Boot SDKS
+    for side in ["_l", "_r"]:
+        cmds.setDrivenKeyframe(["calf_armor"+side+"_anim_grp.rx", "toe_plate"+side+"_anim_grp.rz"], cd = "driver_foot"+side+".rz", itt = "linear", ott = "linear")
+        cmds.setAttr("driver_foot"+side+".rz", 30)
+        cmds.setAttr("calf_armor"+side+"_anim_grp.rx", -90)
+        cmds.setAttr("toe_plate"+side+"_anim_grp.rz", -90)
+        cmds.setDrivenKeyframe(["calf_armor"+side+"_anim_grp.rx", "toe_plate"+side+"_anim_grp.rz"], cd = "driver_foot"+side+".rz", itt = "linear", ott = "linear")
+        cmds.setAttr("driver_foot"+side+".rz", 0)
+
+        cmds.setDrivenKeyframe(["calf_armor"+side+"_anim_grp.rx", "calf_armor"+side+"_anim_grp.ry", "calf_armor"+side+"_anim_grp.rz"], cd = "driver_foot"+side+".rx", itt = "linear", ott = "linear")
+        cmds.setAttr("driver_foot"+side+".rx", 45)
+        cmds.setAttr("calf_armor"+side+"_anim_grp.rx", -11.7)
+        cmds.setAttr("calf_armor"+side+"_anim_grp.ry", -41)
+        cmds.setAttr("calf_armor"+side+"_anim_grp.rz", 24.9)
+        cmds.setDrivenKeyframe(["calf_armor"+side+"_anim_grp.rx", "calf_armor"+side+"_anim_grp.ry", "calf_armor"+side+"_anim_grp.rz"], cd = "driver_foot"+side+".rx", itt = "linear", ott = "linear")
+        cmds.setAttr("driver_foot"+side+".rx", -45)
+        cmds.setAttr("calf_armor"+side+"_anim_grp.rx", -11.7)
+        cmds.setAttr("calf_armor"+side+"_anim_grp.ry", 41)
+        cmds.setAttr("calf_armor"+side+"_anim_grp.rz", -24.9)
+        cmds.setDrivenKeyframe(["calf_armor"+side+"_anim_grp.rx", "calf_armor"+side+"_anim_grp.ry", "calf_armor"+side+"_anim_grp.rz"], cd = "driver_foot"+side+".rx", itt = "linear", ott = "linear")
+        cmds.setAttr("driver_foot"+side+".rx", 0)
+        
+        cmds.setDrivenKeyframe(["calf_armor"+side+"_anim_grp.rx", "calf_armor"+side+"_anim_grp.ry", "calf_armor"+side+"_anim_grp.rz"], cd = "driver_foot"+side+".ry", itt = "linear", ott = "linear")
+        cmds.setAttr("driver_foot"+side+".ry", 45)
+        cmds.setAttr("calf_armor"+side+"_anim_grp.rx", -17)
+        cmds.setAttr("calf_armor"+side+"_anim_grp.ry", 24)
+        cmds.setAttr("calf_armor"+side+"_anim_grp.rz", 0.3)
+        cmds.setDrivenKeyframe(["calf_armor"+side+"_anim_grp.rx", "calf_armor"+side+"_anim_grp.ry", "calf_armor"+side+"_anim_grp.rz"], cd = "driver_foot"+side+".ry", itt = "linear", ott = "linear")
+        cmds.setAttr("driver_foot"+side+".ry", -45)
+        cmds.setAttr("calf_armor"+side+"_anim_grp.rx", -17)
+        cmds.setAttr("calf_armor"+side+"_anim_grp.ry", -24)
+        cmds.setAttr("calf_armor"+side+"_anim_grp.rz", -0.3)
+        cmds.setDrivenKeyframe(["calf_armor"+side+"_anim_grp.rx", "calf_armor"+side+"_anim_grp.ry", "calf_armor"+side+"_anim_grp.rz"], cd = "driver_foot"+side+".ry", itt = "linear", ott = "linear")
+        cmds.setAttr("driver_foot"+side+".ry", 0)
+    
+legRig()
+
+# ------------------------------------------------------------------
+def tailFlapRig():
+    cmds.setDrivenKeyframe(["tail_pad_anim_grp.rx"], cd = "driver_spine_01.rz", itt = "linear", ott = "linear")
+    cmds.setAttr("driver_spine_01.rz", 90)
+    cmds.setAttr("tail_pad_anim_grp.rx", -90)
+    cmds.setDrivenKeyframe(["tail_pad_anim_grp.rx"], cd = "driver_spine_01.rz", itt = "linear", ott = "linear")
+    cmds.setAttr("driver_spine_01.rz", 15)
+    cmds.setAttr("tail_pad_anim_grp.rx", -38)
+    cmds.setDrivenKeyframe(["tail_pad_anim_grp.rx"], cd = "driver_spine_01.rz", itt = "spline", ott = "spline")
+    cmds.setAttr("driver_spine_01.rz", 0)
+
+    cmds.setDrivenKeyframe(["tail_pad_anim_grp.ry"], cd = "driver_spine_01.ry", itt = "linear", ott = "linear")
+    cmds.setAttr("driver_spine_01.ry", 90)
+    cmds.setAttr("tail_pad_anim_grp.ry", 25)
+    cmds.setDrivenKeyframe(["tail_pad_anim_grp.ry"], cd = "driver_spine_01.ry", itt = "linear", ott = "linear")
+    cmds.setAttr("driver_spine_01.ry", -90)
+    cmds.setAttr("tail_pad_anim_grp.ry", -25)
+    cmds.setDrivenKeyframe(["tail_pad_anim_grp.ry"], cd = "driver_spine_01.ry", itt = "linear", ott = "linear")
+    cmds.setAttr("driver_spine_01.ry", 0)
+
+    cmds.setDrivenKeyframe(["tail_pad_anim_grp.rz"], cd = "driver_spine_01.rx", itt = "linear", ott = "linear")
+    cmds.setAttr("driver_spine_01.rx", 90)
+    cmds.setAttr("tail_pad_anim_grp.rz", 20)
+    cmds.setDrivenKeyframe(["tail_pad_anim_grp.rz"], cd = "driver_spine_01.rx", itt = "linear", ott = "linear")
+    cmds.setAttr("driver_spine_01.rx", -90)
+    cmds.setAttr("tail_pad_anim_grp.rz", -20)
+    cmds.setDrivenKeyframe(["tail_pad_anim_grp.rz"], cd = "driver_spine_01.rx", itt = "linear", ott = "linear")
+    cmds.setAttr("driver_spine_01.rx", 0)
+tailFlapRig()
+
+# ------------------------------------------------------------------
+def scaleControls():
+    #Hide controls
+    '''controlsToHide = ["toe_l_anim", "toe_r_anim", "loarm_shield_top_l_anim", "loarm_shield_bot_l_anim", "loarm_shield_back_l_anim", "panels_l_anim", "panels_main_anim", "panels_r_anim", "panel_top_root_l_anim",
+    "panel_top_root_r_anim",  "panel_bot_root_l_anim", "panel_bot_root_r_anim", "panel_bot_l_anim", "panel_bot_r_anim", "panel_top_l_anim", "panel_top_r_anim", "toe_plate_l_anim", "toe_plate_r_anim"]
+    for ctrl in controlsToHide:
+        if cmds.objExists(ctrl):
+            grp = ctrl+"_grp"
+            cmds.setAttr(grp + ".visibility", 0)
+            cmds.setAttr(grp + ".visibility", lock = True)'''
+
+    for side in ["_l", "_r"]:
+        cmds.select(["eyebrow_in"+side+"_anim.cv[0:32]"])
+        cmds.scale(0.117744, 0.117744, 0.117744, r=True, ocp = False)
+        cmds.move(0, 0, 0, r=True)
+
+        cmds.select(["eyebrow_out"+side+"_anim.cv[0:32]"])
+        cmds.scale(0.117744, 0.117744, 0.117744, r=True, ocp = False)
+        cmds.move(0, 0, 0, r=True)
+
+        cmds.select(["eyeball"+side+"_anim.cv[0:32]"])
+        cmds.scale(0.1, 0.1, 0.05, r=True, ocp = False)
+        if side == "_l":
+            cmds.move(0, -3.757386, 0, r=True, os=True)
+        else:
+            cmds.move(0, -3.757386, 0, r=True, os=True)
+
+        cmds.select(["eyelid_bott"+side+"_anim.cv[0:32]"])
+        cmds.scale(0.1, 0.1, 0.01, r=True, ocp = False)
+        if side == "_l":
+            cmds.move(0, -3.757386, 0.5, r=True, os=True)
+        else:
+            cmds.move(0, 3.757386, 0.5, r=True, os=True)
+
+        cmds.select(["eyelid_top"+side+"_anim.cv[0:32]"])
+        cmds.scale(0.1, 0.1, 0.01, r=True, ocp = False)
+        if side == "_l":
+            cmds.move(0, -3.757386, 0.5, r=True, os=True)
+        else:
+            cmds.move(0, 3.757386, 0.5, r=True, os=True)
+
+        #scale shin controls
+        cmds.select("shinpad"+side+"_anim" + ".cv[0:32]")
+        cmds.scale(0, 2*0.65, 2*0.65, r=True)
+
+        #calf armor controls
+        cmds.select("calf_armor"+side+"_anim" + ".cv[0:32]")
+        cmds.scale(2*0.65, 2*0.65, 0, r=True)
+
+        #straps controls
+        cmds.select("arm_strap_top"+side+"_anim" + ".cv[0:32]")
+        cmds.scale(0.22, 0.22, 0.22, r=True)
+
+        cmds.select("arm_strap_bott"+side+"_anim" + ".cv[0:32]")
+        cmds.scale(0.22, 0.22, 0.22, r=True)
+
+        cmds.select("chest_strap"+side+"_anim" + ".cv[0:32]")
+        cmds.scale(0.22, 0.22, 0.22, r=True)
+
+        cmds.select("hip_strap"+side+"_anim" + ".cv[0:32]")
+        cmds.scale(0.22, 0.22, 0.22, r=True)
+
+        #wrist anims
+        cmds.select("ik_wrist"+side+"_anim" + ".cv[0:32]")
+        cmds.scale(2.8*0.65, 2.8*0.65, 2.8*0.65, r=True)
+
+        cmds.select("ik_elbow"+side+"_anim" + ".cv[0:32]")
+        cmds.scale(3*0.65, 3*0.65, 3*0.65, r=True)
+
+        #ik feet
+        cmds.select("ik_foot_anim"+side + ".cv[0:32]")
+        cmds.scale(1.7*0.65, 1.75*0.65, 1*0.65, r=True)
+    
+        #fingers
+        fingers = ["index_finger_fk_ctrl_1", "index_finger_fk_ctrl_2", "middle_finger_fk_ctrl_1", "middle_finger_fk_ctrl_2", "ring_finger_fk_ctrl_1", "ring_finger_fk_ctrl_2", "pinky_finger_fk_ctrl_1", "pinky_finger_fk_ctrl_2", "thumb_finger_fk_ctrl_1", "thumb_finger_fk_ctrl_2"]
+        for finger in fingers:
+            cmds.select(finger+side+".cv[0:32]")
+            cmds.scale(3.7*0.65, 3.7*0.65, 3.7*0.65, r=True)
+
+    # Other
+    cmds.select(["jaw_anim.cv[0:32]"])
+    cmds.scale(0.2, 0.2, 0.2, r=True, ocp = True)
+    cmds.move(0, -9, -6, r=True)
+
+    cmds.select("tail_pad_anim" + ".cv[0:32]")
+    cmds.scale(0.3, 0.3, 0.3, r=True)
+    cmds.move(0, 3, 7.5, r=True)
+
+    cmds.select("flap_back_anim" + ".cv[0:32]")
+    cmds.scale(0.3, 0.3, 0.3, r=True)
+    cmds.move(0, 5, 0, r=True)
+
+    cmds.select("flap_front_anim" + ".cv[0:32]")
+    cmds.scale(0.3, 0.3, 0.3, r=True)
+    cmds.move(0, -7, 0, r=True)
+
+    #turbine controls
+    cmds.select("turbine_back_main_anim" + ".cv[0:32]")
+    cmds.scale(1, 0.07, 1, r=True)
+    cmds.move(0, 5.6, 0, r=True)
+
+    cmds.select("turbine_back_inner_anim" + ".cv[0:32]")
+    cmds.scale(0.75, 0.07, 0.75, r=True)
+    cmds.move(0, -0.6, 0, r=True)
+
+    #chest ik
+    cmds.select("chest_ik_anim" + ".cv[0:32]")
+    cmds.scale(1.5*0.65, 1.5*0.65, 1.5*0.65, r=True)
+    
+    #fk spine
+    cmds.select("spine_01_anim" + ".cv[0:32]")
+    cmds.scale(1.28*0.65, 1.28*0.65, 1.28*0.65, r=True)
+        
+    cmds.select("spine_02_anim" + ".cv[0:32]")
+    cmds.scale(1.32*0.65, 1.32*0.65, 1.32*0.65, r=True)
+
+    cmds.select("spine_03_anim" + ".cv[0:32]")
+    cmds.scale(1.54*0.65, 1.54*0.65, 1.54*0.65, r=True)    
+    
+    #neck
+    cmds.select("neck_01_fk_anim" + ".cv[0:32]")
+    cmds.scale(2.3*0.65, 2.3*0.65, 2.3*0.65, r=True)
+    
+    # football pads
+    cmds.select("football_pads_front_anim" + ".cv[0:32]")
+    cmds.scale(0.2, 0.2, 0.2, r=True)
+    cmds.move(0, 9, 16, r=True)
+
+    cmds.select("football_pads_back_anim" + ".cv[0:32]")
+    cmds.scale(0.2, 0.2, 0.2, r=True)
+    cmds.move(0, 9, 16, r=True)
+
+scaleControls()
+
+# ------------------------------------------------------------------
+def rigArms():
+    cmds.setAttr("Rig_Settings.lUpperarmTwistAmount", 1)
+    cmds.setAttr("Rig_Settings.lUpperarmTwist2Amount", 0.5)
+    cmds.setAttr("Rig_Settings.rUpperarmTwistAmount", 1)
+    cmds.setAttr("Rig_Settings.rUpperarmTwist2Amount", 0.5)    
+
+    cmds.setAttr("Rig_Settings.lArmMode", 0)
+    cmds.setAttr("Rig_Settings.lClavMode", 0)
+    cmds.setAttr("Rig_Settings.rArmMode", 0)
+    cmds.setAttr("Rig_Settings.rClavMode", 0)
+
+    for side in ["_l", "_r"]:
+        # Create SDK's to drive the shoulder blendshapes by thex upperarm joint
+        cmds.setAttr("fk_arm"+side+"_anim.ry", 62)
+        cmds.setDrivenKeyframe(["shoulders_geo_blendShape.shoulders_geo_shoulders_up"+side], cd="upperarm"+side+".rotateY", itt="flat", ott="flat")
+        cmds.setDrivenKeyframe(["shoulders_geo_blendShape.shoulders_geo_shoulders_down"+side], cd="upperarm"+side+".rotateY", itt="flat", ott="flat")
+
+        cmds.setAttr("fk_arm"+side+"_anim.ry", 0)
+        cmds.setAttr("shoulders_geo_blendShape.shoulders_geo_shoulders_up"+side, 1)
+        cmds.setDrivenKeyframe(["shoulders_geo_blendShape.shoulders_geo_shoulders_up"+side], cd="upperarm"+side+".rotateY", itt="flat", ott="flat")
+        cmds.setDrivenKeyframe(["shoulders_geo_blendShape.shoulders_geo_shoulders_down"+side], cd="upperarm"+side+".rotateY", itt="flat", ott="flat")
+        cmds.setAttr("shoulders_geo_blendShape.shoulders_geo_shoulders_up"+side, 0)
+
+        cmds.setAttr("fk_arm"+side+"_anim.ry", 82)
+        cmds.setAttr("shoulders_geo_blendShape.shoulders_geo_shoulders_down"+side, 1)
+        cmds.setDrivenKeyframe(["shoulders_geo_blendShape.shoulders_geo_shoulders_up"+side], cd="upperarm"+side+".rotateY", itt="flat", ott="flat")
+        cmds.setDrivenKeyframe(["shoulders_geo_blendShape.shoulders_geo_shoulders_down"+side], cd="upperarm"+side+".rotateY", itt="flat", ott="flat")
+        cmds.setAttr("shoulders_geo_blendShape.shoulders_geo_shoulders_down"+side, 0)
+
+        cmds.setAttr("fk_arm"+side+"_anim.ry", 0)
+
+        # Create SDK's to drive the bicep blendshape by the elbow joint
+        cmds.setAttr("fk_elbow"+side+"_anim.rz", 0)
+        cmds.setAttr("shoulders_geo_blendShape.shoulders_geo_elbows_bent"+side, 0)
+        cmds.setDrivenKeyframe(["shoulders_geo_blendShape.shoulders_geo_elbows_bent"+side], cd="lowerarm"+side+".rotateZ", itt="flat", ott="flat")
+        
+        cmds.setAttr("fk_elbow"+side+"_anim.rz", -100)
+        cmds.setAttr("shoulders_geo_blendShape.shoulders_geo_elbows_bent"+side, 1)
+        cmds.setDrivenKeyframe(["shoulders_geo_blendShape.shoulders_geo_elbows_bent"+side], cd="lowerarm"+side+".rotateZ", itt="flat", ott="flat")
+        
+        cmds.setAttr("fk_elbow"+side+"_anim.rz", 0)
+
+    cmds.setAttr("Rig_Settings.lArmMode", 1)
+    cmds.setAttr("Rig_Settings.lClavMode", 1)
+    cmds.setAttr("Rig_Settings.rArmMode", 1)
+    cmds.setAttr("Rig_Settings.rClavMode", 1)
+rigArms()
+
+# ------------------------------------------------------------------
+def invertEyelids():
+    sides = ["_l", "_r"]
+    for side in sides:
+        cmds.delete("driver_eyelid_bott"+side+"_parentConstraint1")
+        cmds.rotate(0, -180, 0, "eyelid_bott"+side+"_anim_grp", r=True, os=True)
+        cmds.parentConstraint("eyelid_bott"+side+"_anim", "driver_eyelid_bott"+side, mo=True)
+invertEyelids()
+
+def invertRightEyeball():
+    cmds.delete("driver_eyeball_r_parentConstraint1")
+    cmds.rotate(0, -180, -180, "eyeball_r_anim_grp", r=True, os=True)
+    cmds.parentConstraint("eyeball_r_anim", "driver_eyeball_r", mo=True)
+invertRightEyeball()
+
+# ------------------------------------------------------------------
+def createGameCam():
+    gameCam = cmds.camera(centerOfInterest=5, focalLength=18, lensSqueezeRatio=1, cameraScale=1, horizontalFilmAperture=1.41732, horizontalFilmOffset=0, verticalFilmAperture=0.94488, verticalFilmOffset=0, filmFit="Overscan", overscan=1.05, motionBlur=0, shutterAngle=144, nearClipPlane=1.0, farClipPlane=10000, orthographic=0, orthographicWidth=30, panZoomEnabled=0, horizontalPan=0, verticalPan=0, zoom=1, displayFilmGate=True)[0]
+    print gameCam
+    gameCam = cmds.rename(gameCam, "GameCam")
+    cmds.camera(gameCam, e=True, horizontalFilmAperture=1.6724376, focalLength=21.2399857011)
+    cmds.setAttr(gameCam+".displayGateMaskOpacity", 1)
+    cmds.setAttr(gameCam+".displayGateMaskColor", 0, 0, 0, type="double3")
+    cmds.setAttr(gameCam+".r", 90, 0, 180, type="double3")
+    
+    # Create the control
+    control = cmds.circle(name=gameCam+"_Ctrl", ch=True, o=True, r=15)[0]
+    cmds.makeIdentity(control, a=True, t=True, r=True, s=True)
+    # Create groups for the sdk and the constraints
+    const_grp = cmds.group(control, name=control+"_grp")
+    if cmds.objExists("rig_grp"):
+        cmds.parent(const_grp, "rig_grp")
+    # Parent Constrain the cam to the control
+    cmds.parentConstraint(control, gameCam, mo=True)
+    
+    cmds.connectAttr(control+".visibility", gameCam+".visibility", f=True)
+    cmds.setAttr(const_grp+".v", 0)
+    
+    cmds.setAttr(const_grp+".t", 0, 500, 200, type="double3")
+createGameCam()
+
+# ------------------------------------------------------------------
+def addCapsule(radius, halfHeight):
+    trueHalfHeight = (halfHeight*2)-(radius*2)
+    capsule = cmds.polyCylinder(name="collision_capsule", r=radius, h=trueHalfHeight, sx=18, sy=1, sz=5, ax=[0,0,1], rcp=True, cuv=3, ch=True)
+    cmds.setAttr(capsule[0]+".tz", halfHeight+2.2)
+    cmds.setAttr(capsule[0]+".rz", 90)
+
+    cmds.addAttr(capsule[0], ln="Radius", dv=radius, keyable=True)
+    cmds.addAttr(capsule[0], ln="HalfHeight", dv=halfHeight, keyable=True)
+
+    expr = capsule[1]+".radius = "+capsule[0]+".Radius;\n"+capsule[1]+".height = ("+capsule[0]+".HalfHeight*2) - ("+capsule[0]+".Radius*2);"
+    cmds.expression(s=expr, o=capsule[1], ae=1, uc="all")
+
+    cmds.setAttr(capsule[0]+".overrideEnabled", 1)
+    cmds.setAttr(capsule[0]+".overrideColor", 18)
+    cmds.setAttr(capsule[0]+".overrideShading", 0)
+
+    cmds.addAttr(capsule[0], ln="Shaded", dv=0, min=0, max=1, keyable=True)
+    cmds.connectAttr(capsule[0]+".Shaded", capsule[0]+".overrideShading")
+
+    cmds.parentConstraint("root", capsule[0], mo=True)
+    cmds.select(capsule[0], r=True)
+    capsuleDisplayLayer = cmds.createDisplayLayer(nr=True, name="Collision Capsule")
+    cmds.setAttr(capsuleDisplayLayer+".displayType", 2)
+    cmds.setAttr(capsuleDisplayLayer+".v", 0)
+    
+    for attr in ["tx", "ty", "tz", "rx", "ry", "rz", "sx", "sy", "sz"]:
+        cmds.setAttr(capsule[0]+"."+attr, lock=True, keyable=False, channelBox=False)
+
+    cmds.select(clear=True)
+addCapsule(64, 110)
+
+# ------------------------------------------------------------------
+# Add in a temporary face emote for the rabbit
+def addFaceEmote():
+   
+    # morphs
+    cmds.addAttr("head_fk_anim", ln = "face_emote", dv = 0, min = 0, max = 1, keyable = True)
+
+    cmds.connectAttr("head_fk_anim.face_emote", "blendShape2.head_geo_emote")
+    cmds.connectAttr("head_fk_anim.face_emote", "blendShape1.helmet_geo_emote")
+    cmds.connectAttr("head_fk_anim.face_emote", "blendShape3.eye_shadow_geo_emote")
+
+    cmds.select(None)
+#addFaceEmote()
+
+# ------------------------------------------------------------------
+# add ik bones Function that is called later after we open the export file.
+def addIKBones():
+    try:
+		# create the joints
+		cmds.select(clear=True)
+		ikFootRoot = cmds.joint(name = "ik_foot_root")
+		cmds.select(clear=True)
+		
+		cmds.select(clear=True)
+		ikFootLeft = cmds.joint(name = "ik_foot_l")
+		cmds.select(clear=True)
+		
+		cmds.select(clear=True)
+		ikFootRight = cmds.joint(name = "ik_foot_r")
+		cmds.select(clear=True)
+		
+		cmds.select(clear=True)
+		ikHandRoot = cmds.joint(name = "ik_hand_root")
+		cmds.select(clear=True)
+		
+		cmds.select(clear=True)
+		ikHandGun = cmds.joint(name = "ik_hand_gun")
+		cmds.select(clear=True)
+		
+		cmds.select(clear=True)
+		ikHandLeft = cmds.joint(name = "ik_hand_l")
+		cmds.select(clear=True)
+		
+		cmds.select(clear=True)
+		ikHandRight = cmds.joint(name = "ik_hand_r")
+		cmds.select(clear=True)
+		
+		# create hierarhcy
+		cmds.parent(ikFootRoot, "root")
+		cmds.parent(ikHandRoot, "root")
+		cmds.parent(ikFootLeft, ikFootRoot)
+		cmds.parent(ikFootRight, ikFootRoot)
+		cmds.parent(ikHandGun, ikHandRoot)
+		cmds.parent(ikHandLeft, ikHandGun)
+		cmds.parent(ikHandRight, ikHandGun)
+		
+		# constrain the joints
+		leftHandConstraint = cmds.parentConstraint("hand_l", ikHandLeft)[0]
+		rightHandGunConstraint = cmds.parentConstraint("hand_r", ikHandGun)[0]
+		rightHandConstraint = cmds.parentConstraint("hand_r", ikHandRight)[0]
+		leftFootConstraint = cmds.parentConstraint("foot_l", ikFootLeft)[0]
+		rightFootConstraint = cmds.parentConstraint("foot_r", ikFootRight)[0]
+	
+    except:
+		cmds.warning("Something went wrong")
+
+# ------------------------------------------------------------------
+# ADD the IK bones to the Export file.
+#save the AnimRig file before opening the export file.
+cmds.file(save = True, type = "mayaBinary", force = True)
+
+#ADD IK BONES
+#open export file
+sceneName = cmds.file(q = True, sceneName = True)
+exportFile = sceneName.partition("AnimRigs")[0] + "ExportFiles" +  sceneName.partition("AnimRigs")[2].partition(".mb")[0] + "_Export.mb"
+cmds.file(exportFile, open = True, force = True)
+
+#add ik bones
+addIKBones()
+
+# ------------------------------------------------------------------
+def addCustomAttrs():
+    attrs = cmds.listAttr("root")
+    exists = False
+    for one in attrs:
+        if one.find("DistanceCurve") != -1:
+            exists = True
+    if exists == False:
+        cmds.addAttr("root", ln="DistanceCurve", keyable=True)
+addCustomAttrs()
+
+#save the export file.
+cmds.file(save = True, type = "mayaBinary", force = True)
+
+
+
+#########################################
+# RUN THIS COMMAND DURING THE RIG POSE CREATION STEP TO GET THE SHIELD TO THE CORRECT LOCATION.  LEAVE IT COMMENTED OUT SO THAT IT DOESNT RUN DURING RIGGING.
+'''
+#cmds.xform(q=True, ws=True, rp=True)
+#cmds.xform(q=True, ro=True)
+cmds.xform("sheild_main_mover", ws=True, t=[-0.018884450479987436, 59.40691643000001, 272.98635249999995])
+cmds.xform("sheild_main_mover", ro=[-14.885476763734184, 61.20693988721698, -15.277192149145698])
+'''
+#########################################
+
+
+
+
+
+
+
+
+
+
+
+
+
+
+
+
+
+
+
+
+
+
+
+
+
+
+
+
+
+
+
+
+
+
+
+
+
+
+
+
+
+
+
+
+
+
+
+
