Branching

diff --git a/Engine/Extras/Maya_AnimationRiggingTools/ARTv1/General/ART/Projects/Orion/Scripts/Totem_CreateLODs.py b/Engine/Extras/Maya_AnimationRiggingTools/ARTv1/General/ART/Projects/Orion/Scripts/Totem_CreateLODs.py
new file mode 100644
--- /dev/null
+++ b/Engine/Extras/Maya_AnimationRiggingTools/ARTv1/General/ART/Projects/Orion/Scripts/Totem_CreateLODs.py
@@ -0,0 +1,534 @@
+import maya.cmds as cmds
+import customMayaMenu as menu
+import Tools.ART_ImportMultipleWeights as impw
+import Tools.ART_TransferWeights as transfer
+import Tools.ART_ExportLODtoFBX as export
+reload(export)
+
+################################################## TOTEM #########################################################################
+# Warn the user that this operation will save their file.
+saveResult = cmds.confirmDialog(title = "Save Warning", message="In order to continue your file must be saved.  Would you like to save it?  If yes it will be saved and after the operation is complete your file will be re-opened.", button = ["Yes", "Cancel"], defaultButton='Yes', cancelButton='Cancel', dismissString='Cancel')
+
+# If they are ok with saving the file, warn them that it will also delete all of the current files in their skin weights directory.
+if saveResult == "Yes":
+    deleteSkinFilesResult = cmds.confirmDialog(title = "SkinWeights Warning", message="In order to continue we must delete all of the files in your ..\General\ART\SkinWeights\ directory.  Would you like to continue?.", button = ["Yes", "Cancel"], defaultButton='Yes', cancelButton='Cancel', dismissString='Cancel')
+    if deleteSkinFilesResult == "Yes":
+        toolsPath = cmds.internalVar(usd = True) + "mayaTools.txt"
+        if os.path.exists(toolsPath):
+            f = open(toolsPath, 'r')
+            mayaToolsDir = f.readline()
+            f.close()
+
+        path = mayaToolsDir + "\General\ART\SkinWeights\\"
+        files = os.listdir(path)
+        #delete all weight cache files in the system folder
+        for file in files:
+            os.remove(path + file)
+            #print path + file
+
+
+
+# Totem LOD1 -------------------------------------------------------------------------------------------------------------
+def TotemLOD1():
+    if saveResult == "Yes":
+
+        cmds.file(save=True)
+
+        meshes = ["Tri_Totem_Body_LOD1", "Tri_Totem_Collar_LOD1", "Tri_Totem_Cornea_LOD1", "Tri_Totem_DogBones_LOD1", "Tri_Totem_Eyebrows_LOD1", "Tri_Totem_Eyelashes_LOD1", "Tri_Totem_Eyes_LOD1", "Tri_Totem_EyeShadow_LOD1", "Tri_Totem_Hair_LOD1", "Tri_Totem_HairBag_LOD1", "Tri_Totem_HairTie_LOD1", "Tri_Totem_Hands_LOD1", "Tri_Totem_HandTech_LOD1", "Tri_Totem_Head_LOD1", "Tri_Totem_Knees_LOD1", "Tri_Totem_Legs_LOD1", "Tri_Totem_Mouth_LOD1", "Tri_Totem_ShoulderPads_LOD1", "Tri_Totem_Weapon_LOD1", "Tri_Totem_Weapon_Bottom_FrontBase_LOD1", "Tri_Totem_Weapon_Bottom_FrontEnd_LOD1", "Tri_Totem_Weapon_Handle_LOD1", "Tri_Totem_Weapon_SupportBar_LOD1", "Tri_Totem_Weapon_TopBack_LOD1", "Tri_Totem_Weapon_TopFront_LOD1"]
+
+        jointsToDelete = []
+
+        # remove unneeded joints
+        jointsToRemove = ["glasses", "hair_l", "hair_r", "ponytail_01", "ponytail_02", "ponytail_03"]
+        jointToTransferTo = "head"
+        transfer.LOD_transferWeights(meshes, jointsToRemove, jointToTransferTo)
+        jointsToDelete.extend(jointsToRemove)
+
+        jointsToRemove = ["knee_pad_l"]
+        jointToTransferTo = "calf_l"
+        transfer.LOD_transferWeights(meshes, jointsToRemove, jointToTransferTo)
+        jointsToDelete.extend(jointsToRemove)
+
+        jointsToRemove = ["knee_pad_r"]
+        jointToTransferTo = "calf_r"
+        transfer.LOD_transferWeights(meshes, jointsToRemove, jointToTransferTo)
+        jointsToDelete.extend(jointsToRemove)
+
+        
+        # Delete the joints we no longer want.
+        for one in jointsToDelete:
+            if cmds.objExists(one):
+                cmds.delete(one)
+                
+        # Export the fbx file.
+        export.exportLOD("Totem_LOD1.fbx", "LOD1")
+        
+        # Re-open the file without saving.
+        fullPath = cmds.file(q = True, sceneName = True)
+        cmds.file(fullPath, open=True, f=True)
+
+        print "LOD1 Exported Successfully!!!!!!!!!!!!!!!"
+
+TotemLOD1()
+
+# Totem LOD2 -------------------------------------------------------------------------------------------------------------
+def TotemLOD2():
+    if saveResult == "Yes":
+
+        cmds.file(save=True)
+        
+        meshes = ["Tri_Totem_Body_LOD2", "Tri_Totem_Collar_LOD2", "Tri_Totem_Cornea_LOD2", "Tri_Totem_DogBones_LOD2", "Tri_Totem_Eyebrows_LOD2", "Tri_Totem_Eyelashes_LOD2", "Tri_Totem_Eyes_LOD2", "Tri_Totem_EyeShadow_LOD2", "Tri_Totem_Hair_LOD2", "Tri_Totem_HairBag_LOD2", "Tri_Totem_HairTie_LOD2", "Tri_Totem_Hands_LOD2", "Tri_Totem_HandTech_LOD2", "Tri_Totem_Head_LOD2", "Tri_Totem_Knees_LOD2", "Tri_Totem_Legs_LOD2", "Tri_Totem_Mouth_LOD2", "Tri_Totem_ShoulderPads_LOD2", "Tri_Totem_Weapon_LOD2", "Tri_Totem_Weapon_Bottom_FrontBase_LOD2", "Tri_Totem_Weapon_Bottom_FrontEnd_LOD2", "Tri_Totem_Weapon_Handle_LOD2", "Tri_Totem_Weapon_SupportBar_LOD2", "Tri_Totem_Weapon_TopBack_LOD2", "Tri_Totem_Weapon_TopFront_LOD2"]
+
+        # Export a skin weights file for each mesh we plan to operate on.
+        cmds.select(meshes)
+        menu.exportSelectedSkinWeights(True)
+
+        # Set the pose if needed
+        cmds.setAttr("thumb_03_l.rotateZ", 49.02926105)
+        cmds.setAttr("ring_03_l.rotateZ", 60.61060457)
+        cmds.setAttr("pinky_03_l.rotateZ", 60.61060457)
+        cmds.setAttr("middle_03_l.rotateZ", 60.61060457)
+        cmds.setAttr("index_03_l.rotateZ", 60.61060457)
+
+        cmds.setAttr("thumb_03_r.rotateZ", 49.02926105)
+        cmds.setAttr("ring_03_r.rotateZ", 60.61060457)
+        cmds.setAttr("pinky_03_r.rotateZ", 60.61060457)
+        cmds.setAttr("middle_03_r.rotateZ", 60.61060457)
+        cmds.setAttr("index_03_r.rotateZ", 60.61060457)
+
+        # Delete the history on the meshes so that they get baked down into their posed positions
+        cmds.delete(meshes, ch=True)
+
+        # Move the skeleton back to its bindpose (This script REQUIRES the bind pose to be the pose that we want to export the skeleton to.  This is the case for all known characters created with A.R.T.)
+        cmds.dagPose("root", bp=True, r=True) 
+
+        # Re-import the skinning we exported before onto the mesh so thats its bound in its posed position.
+        cmds.select(meshes)
+        impw.importMultipleSkinWeights("", True)
+
+        jointsToDelete = []
+
+         # remove unneeded joints
+        jointsToRemove = ["back_plate", "chest_l", "chest_r", "collar", "collar_clav_l", "hose_l", "l_hose_base", "collar_clav_r", "hose_r", "r_hose_base"]
+        jointToTransferTo = "spine_03"
+        transfer.LOD_transferWeights(meshes, jointsToRemove, jointToTransferTo)
+        jointsToDelete.extend(jointsToRemove)
+
+        jointsToRemove = ["glasses", "hair_l", "hair_r", "ponytail_01", "ponytail_02", "ponytail_03"]
+        jointToTransferTo = "head"
+        transfer.LOD_transferWeights(meshes, jointsToRemove, jointToTransferTo)
+        jointsToDelete.extend(jointsToRemove)
+
+        jointsToRemove = ["upperarm_twist_01_l", "shoulder_pad_l", "l_hose_end", "l_pad_end", "upperarm_twist_02_l"]
+        jointToTransferTo = "upperarm_l"
+        transfer.LOD_transferWeights(meshes, jointsToRemove, jointToTransferTo)
+        jointsToDelete.extend(jointsToRemove)
+
+        jointsToRemove = ["upperarm_twist_01_r", "shoulder_pad_r", "r_hose_end", "r_pad_end", "upperarm_twist_02_r"]
+        jointToTransferTo = "upperarm_r"
+        transfer.LOD_transferWeights(meshes, jointsToRemove, jointToTransferTo)
+        jointsToDelete.extend(jointsToRemove)
+
+        jointsToRemove = ["lowerarm_twist_01_l", "lowerarm_twist_02_l"]
+        jointToTransferTo = "lowerarm_l"
+        transfer.LOD_transferWeights(meshes, jointsToRemove, jointToTransferTo)
+        jointsToDelete.extend(jointsToRemove)
+
+        jointsToRemove = ["lowerarm_twist_01_r", "lowerarm_twist_02_r"]
+        jointToTransferTo = "lowerarm_r"
+        transfer.LOD_transferWeights(meshes, jointsToRemove, jointToTransferTo)
+        jointsToDelete.extend(jointsToRemove)
+
+        jointsToRemove = ["hand_plate_l"]
+        jointToTransferTo = "hand_l"
+        transfer.LOD_transferWeights(meshes, jointsToRemove, jointToTransferTo)
+        jointsToDelete.extend(jointsToRemove)
+
+        jointsToRemove = ["index_03_l"]
+        jointToTransferTo = "index_02_l"
+        transfer.LOD_transferWeights(meshes, jointsToRemove, jointToTransferTo)
+        jointsToDelete.extend(jointsToRemove)
+
+        jointsToRemove = ["middle_03_l"]
+        jointToTransferTo = "middle_02_l"
+        transfer.LOD_transferWeights(meshes, jointsToRemove, jointToTransferTo)
+        jointsToDelete.extend(jointsToRemove)
+
+        jointsToRemove = ["pinky_03_l"]
+        jointToTransferTo = "pinky_02_l"
+        transfer.LOD_transferWeights(meshes, jointsToRemove, jointToTransferTo)
+        jointsToDelete.extend(jointsToRemove)
+
+        jointsToRemove = ["ring_03_l"]
+        jointToTransferTo = "ring_02_l"
+        transfer.LOD_transferWeights(meshes, jointsToRemove, jointToTransferTo)
+        jointsToDelete.extend(jointsToRemove)
+
+        jointsToRemove = ["thumb_03_l"]
+        jointToTransferTo = "thumb_02_l"
+        transfer.LOD_transferWeights(meshes, jointsToRemove, jointToTransferTo)
+        jointsToDelete.extend(jointsToRemove)
+
+        jointsToRemove = ["hand_plate_r"]
+        jointToTransferTo = "hand_r"
+        transfer.LOD_transferWeights(meshes, jointsToRemove, jointToTransferTo)
+        jointsToDelete.extend(jointsToRemove)
+
+        jointsToRemove = ["index_03_r"]
+        jointToTransferTo = "index_02_r"
+        transfer.LOD_transferWeights(meshes, jointsToRemove, jointToTransferTo)
+        jointsToDelete.extend(jointsToRemove)
+
+        jointsToRemove = ["middle_03_r"]
+        jointToTransferTo = "middle_02_r"
+        transfer.LOD_transferWeights(meshes, jointsToRemove, jointToTransferTo)
+        jointsToDelete.extend(jointsToRemove)
+
+        jointsToRemove = ["pinky_03_r"]
+        jointToTransferTo = "pinky_02_r"
+        transfer.LOD_transferWeights(meshes, jointsToRemove, jointToTransferTo)
+        jointsToDelete.extend(jointsToRemove)
+
+        jointsToRemove = ["ring_03_r"]
+        jointToTransferTo = "ring_02_r"
+        transfer.LOD_transferWeights(meshes, jointsToRemove, jointToTransferTo)
+        jointsToDelete.extend(jointsToRemove)
+
+        jointsToRemove = ["thumb_03_r"]
+        jointToTransferTo = "thumb_02_r"
+        transfer.LOD_transferWeights(meshes, jointsToRemove, jointToTransferTo)
+        jointsToDelete.extend(jointsToRemove)
+
+        jointsToRemove = ["thigh_twist_02_l"]
+        jointToTransferTo = "thigh_twist_01_l"
+        transfer.LOD_transferWeights(meshes, jointsToRemove, jointToTransferTo)
+        jointsToDelete.extend(jointsToRemove)
+
+        jointsToRemove = ["thigh_twist_02_r"]
+        jointToTransferTo = "thigh_twist_01_r"
+        transfer.LOD_transferWeights(meshes, jointsToRemove, jointToTransferTo)
+        jointsToDelete.extend(jointsToRemove)
+
+        jointsToRemove = ["calf_twist_01_l", "calf_twist_02_l", "knee_pad_l"]
+        jointToTransferTo = "calf_l"
+        transfer.LOD_transferWeights(meshes, jointsToRemove, jointToTransferTo)
+        jointsToDelete.extend(jointsToRemove)
+
+        jointsToRemove = ["calf_twist_01_r", "calf_twist_02_r", "knee_pad_r"]
+        jointToTransferTo = "calf_r"
+        transfer.LOD_transferWeights(meshes, jointsToRemove, jointToTransferTo)
+        jointsToDelete.extend(jointsToRemove)
+
+        # Delete the joints we no longer want.
+        for one in jointsToDelete:
+            if cmds.objExists(one):
+                cmds.delete(one)
+                
+        # Export the fbx file.
+        export.exportLOD("Totem_LOD2.fbx", "LOD2")
+
+        # Re-open the file without saving.
+        fullPath = cmds.file(q = True, sceneName = True)
+        cmds.file(fullPath, open=True, f=True)
+
+        print "LOD2 Exported Successfully!!!!!!!!!!!!!!!"
+
+TotemLOD2()
+
+# Totem LOD3 -------------------------------------------------------------------------------------------------------------
+def TotemLOD3():
+    if saveResult == "Yes":
+
+        cmds.file(save=True)
+    
+        meshes = ["Tri_Totem_Body_LOD3", "Tri_Totem_Collar_LOD3", "Tri_Totem_Cornea_LOD3", "Tri_Totem_DogBones_LOD3", "Tri_Totem_Eyebrows_LOD3", "Tri_Totem_Eyelashes_LOD3", "Tri_Totem_Eyes_LOD3", "Tri_Totem_EyeShadow_LOD3", "Tri_Totem_Hair_LOD3", "Tri_Totem_HairBag_LOD3", "Tri_Totem_HairTie_LOD3", "Tri_Totem_Hands_LOD3", "Tri_Totem_HandTech_LOD3", "Tri_Totem_Head_LOD3", "Tri_Totem_Knees_LOD3", "Tri_Totem_Legs_LOD3", "Tri_Totem_Mouth_LOD3", "Tri_Totem_ShoulderPads_LOD3", "Tri_Totem_Weapon_LOD3", "Tri_Totem_Weapon_Bottom_FrontBase_LOD3", "Tri_Totem_Weapon_Bottom_FrontEnd_LOD3", "Tri_Totem_Weapon_Handle_LOD3", "Tri_Totem_Weapon_SupportBar_LOD3", "Tri_Totem_Weapon_TopBack_LOD3", "Tri_Totem_Weapon_TopFront_LOD3"]
+    
+        # Export a skin weights file for each mesh we plan to operate on.
+        cmds.select(meshes)
+        menu.exportSelectedSkinWeights(True)
+
+        # Set the pose if needed
+        cmds.setAttr("thumb_03_l.rotateZ", 49.02926105)
+        cmds.setAttr("thumb_02_l.rotateZ", 29.37554526)
+        cmds.setAttr("thumb_01_l.rotateZ", -1.060679943)
+        cmds.setAttr("thumb_01_l.rotateY", 16.78500021)
+        cmds.setAttr("thumb_01_l.rotateX", 5.616322924)
+        cmds.setAttr("ring_03_l.rotateZ", 60.61060457)
+        cmds.setAttr("ring_02_l.rotateZ", 60.61060457)
+        cmds.setAttr("ring_01_l.rotateZ", 60.61060457)
+        cmds.setAttr("pinky_03_l.rotateZ", 60.61060457)
+        cmds.setAttr("pinky_02_l.rotateZ", 60.61060457)
+        cmds.setAttr("pinky_01_l.rotateZ", 60.61060457)
+        cmds.setAttr("middle_03_l.rotateZ", 60.61060457)
+        cmds.setAttr("middle_02_l.rotateZ", 60.61060457)
+        cmds.setAttr("middle_01_l.rotateZ", 60.61060457)
+        cmds.setAttr("index_03_l.rotateZ", 60.61060457)
+        cmds.setAttr("index_02_l.rotateZ", 60.61060457)
+        cmds.setAttr("index_01_l.rotateZ", 60.61060457)
+
+        cmds.setAttr("thumb_03_r.rotateZ", 49.02926105)
+        cmds.setAttr("thumb_02_r.rotateZ", 29.37554526)
+        cmds.setAttr("thumb_01_r.rotateZ", -1.060679943)
+        cmds.setAttr("thumb_01_r.rotateY", 16.78500021)
+        cmds.setAttr("thumb_01_r.rotateX", 5.616322924)
+        cmds.setAttr("ring_03_r.rotateZ", 60.61060457)
+        cmds.setAttr("ring_02_r.rotateZ", 60.61060457)
+        cmds.setAttr("ring_01_r.rotateZ", 60.61060457)
+        cmds.setAttr("pinky_03_r.rotateZ", 60.61060457)
+        cmds.setAttr("pinky_02_r.rotateZ", 60.61060457)
+        cmds.setAttr("pinky_01_r.rotateZ", 60.61060457)
+        cmds.setAttr("middle_03_r.rotateZ", 60.61060457)
+        cmds.setAttr("middle_02_r.rotateZ", 60.61060457)
+        cmds.setAttr("middle_01_r.rotateZ", 60.61060457)
+        cmds.setAttr("index_03_r.rotateZ", 60.61060457)
+        cmds.setAttr("index_02_r.rotateZ", 60.61060457)
+        cmds.setAttr("index_01_r.rotateZ", 60.61060457)
+
+        # Delete the history on the meshes so that they get baked down into their posed positions
+        cmds.delete(meshes, ch=True)
+
+        # Move the skeleton back to its bindpose (This script REQUIRES the bind pose to be the pose that we want to export the skeleton to.  This is the case for all known characters created with A.R.T.)
+        cmds.dagPose("root", bp=True, r=True) 
+
+        # Re-import the skinning we exported before onto the mesh so thats its bound in its posed position.
+        cmds.select(meshes)
+        impw.importMultipleSkinWeights("", True)
+
+        jointsToDelete = []
+
+        # remove unneeded joints
+        jointsToRemove = ["back_plate", "chest_l", "chest_r", "collar", "collar_clav_l", "hose_l", "l_hose_base", "collar_clav_r", "hose_r", "r_hose_base"]
+        jointToTransferTo = "spine_03"
+        transfer.LOD_transferWeights(meshes, jointsToRemove, jointToTransferTo)
+        jointsToDelete.extend(jointsToRemove)
+
+        jointsToRemove = ["glasses", "hair_l", "hair_r", "ponytail_01", "ponytail_02", "ponytail_03"]
+        jointToTransferTo = "head"
+        transfer.LOD_transferWeights(meshes, jointsToRemove, jointToTransferTo)
+        jointsToDelete.extend(jointsToRemove)
+
+        jointsToRemove = ["upperarm_twist_01_l", "shoulder_pad_l", "l_hose_end", "l_pad_end", "upperarm_twist_02_l"]
+        jointToTransferTo = "upperarm_l"
+        transfer.LOD_transferWeights(meshes, jointsToRemove, jointToTransferTo)
+        jointsToDelete.extend(jointsToRemove)
+
+        jointsToRemove = ["upperarm_twist_01_r", "shoulder_pad_r", "r_hose_end", "r_pad_end", "upperarm_twist_02_r"]
+        jointToTransferTo = "upperarm_r"
+        transfer.LOD_transferWeights(meshes, jointsToRemove, jointToTransferTo)
+        jointsToDelete.extend(jointsToRemove)
+
+        jointsToRemove = ["lowerarm_twist_01_l", "lowerarm_twist_02_l"]
+        jointToTransferTo = "lowerarm_l"
+        transfer.LOD_transferWeights(meshes, jointsToRemove, jointToTransferTo)
+        jointsToDelete.extend(jointsToRemove)
+
+        jointsToRemove = ["lowerarm_twist_01_r", "lowerarm_twist_02_r"]
+        jointToTransferTo = "lowerarm_r"
+        transfer.LOD_transferWeights(meshes, jointsToRemove, jointToTransferTo)
+        jointsToDelete.extend(jointsToRemove)
+
+        jointsToRemove = ["hand_plate_l", "index_01_l", "index_02_l", "index_03_l", "middle_01_l", "middle_02_l", "middle_03_l", "pinky_01_l", "pinky_02_l", "pinky_03_l", "ring_01_l", "ring_02_l", "ring_03_l", "thumb_01_l", "thumb_02_l", "thumb_03_l"]
+        jointToTransferTo = "hand_l"
+        transfer.LOD_transferWeights(meshes, jointsToRemove, jointToTransferTo)
+        jointsToDelete.extend(jointsToRemove)
+
+        jointsToRemove = ["hand_plate_r", "index_01_r", "index_02_r", "index_03_r", "middle_01_r", "middle_02_r", "middle_03_r", "pinky_01_r", "pinky_02_r", "pinky_03_r", "ring_01_r", "ring_02_r", "ring_03_r", "thumb_01_r", "thumb_02_r", "thumb_03_r"]
+        jointToTransferTo = "hand_r"
+        transfer.LOD_transferWeights(meshes, jointsToRemove, jointToTransferTo)
+        jointsToDelete.extend(jointsToRemove)
+
+        jointsToRemove = ["staff_bottom", "handle", "hinge_a", "hinge_b", "staff_top", "tip_spin", "top_tip_a", "top_tip_b", "top_hinge"]
+        jointToTransferTo = "weapon"
+        transfer.LOD_transferWeights(meshes, jointsToRemove, jointToTransferTo)
+        jointsToDelete.extend(jointsToRemove)
+
+        jointsToRemove = ["thigh_twist_01_l", "thigh_twist_02_l"]
+        jointToTransferTo = "thigh_l"
+        transfer.LOD_transferWeights(meshes, jointsToRemove, jointToTransferTo)
+        jointsToDelete.extend(jointsToRemove)
+
+        jointsToRemove = ["thigh_twist_01_r", "thigh_twist_02_r"]
+        jointToTransferTo = "thigh_r"
+        transfer.LOD_transferWeights(meshes, jointsToRemove, jointToTransferTo)
+        jointsToDelete.extend(jointsToRemove)
+
+        jointsToRemove = ["calf_twist_01_l", "calf_twist_02_l", "knee_pad_l"]
+        jointToTransferTo = "calf_l"
+        transfer.LOD_transferWeights(meshes, jointsToRemove, jointToTransferTo)
+        jointsToDelete.extend(jointsToRemove)
+
+        jointsToRemove = ["calf_twist_01_r", "calf_twist_02_r", "knee_pad_r"]
+        jointToTransferTo = "calf_r"
+        transfer.LOD_transferWeights(meshes, jointsToRemove, jointToTransferTo)
+        jointsToDelete.extend(jointsToRemove)
+
+        # Delete the joints we no longer want.
+        for one in jointsToDelete:
+            if cmds.objExists(one):
+                cmds.delete(one)
+                
+        # Export the fbx file.
+        export.exportLOD("Totem_LOD3.fbx", "LOD3")
+        
+        # Re-open the file without saving.
+        fullPath = cmds.file(q = True, sceneName = True)
+        cmds.file(fullPath, open=True, f=True)
+
+        print "LOD3 Exported Successfully!!!!!!!!!!!!!!!"
+
+TotemLOD3()
+
+# Totem LOD4 -------------------------------------------------------------------------------------------------------------
+def TotemLOD4():
+    if saveResult == "Yes":
+
+        cmds.file(save=True)
+    
+        meshes = ["Tri_Totem_Body_LOD4", "Tri_Totem_Collar_LOD4", "Tri_Totem_Cornea_LOD4", "Tri_Totem_DogBones_LOD4", "Tri_Totem_Eyebrows_LOD4", "Tri_Totem_Eyelashes_LOD4", "Tri_Totem_Eyes_LOD4", "Tri_Totem_EyeShadow_LOD4", "Tri_Totem_Hair_LOD4", "Tri_Totem_HairBag_LOD4", "Tri_Totem_HairTie_LOD4", "Tri_Totem_Hands_LOD4", "Tri_Totem_HandTech_LOD4", "Tri_Totem_Head_LOD4", "Tri_Totem_Knees_LOD4", "Tri_Totem_Legs_LOD4", "Tri_Totem_Mouth_LOD4", "Tri_Totem_ShoulderPads_LOD4", "Tri_Totem_Weapon_LOD4", "Tri_Totem_Weapon_Bottom_FrontBase_LOD4", "Tri_Totem_Weapon_Bottom_FrontEnd_LOD4", "Tri_Totem_Weapon_Handle_LOD4", "Tri_Totem_Weapon_SupportBar_LOD4", "Tri_Totem_Weapon_TopBack_LOD4", "Tri_Totem_Weapon_TopFront_LOD4"]
+
+        # Export a skin weights file for each mesh we plan to operate on.
+        cmds.select(meshes)
+        menu.exportSelectedSkinWeights(True)
+
+        # Set the pose if needed
+        cmds.setAttr("thumb_03_l.rotateZ", 49.02926105)
+        cmds.setAttr("thumb_02_l.rotateZ", 29.37554526)
+        cmds.setAttr("thumb_01_l.rotateZ", -1.060679943)
+        cmds.setAttr("thumb_01_l.rotateY", 16.78500021)
+        cmds.setAttr("thumb_01_l.rotateX", 5.616322924)
+        cmds.setAttr("ring_03_l.rotateZ", 60.61060457)
+        cmds.setAttr("ring_02_l.rotateZ", 60.61060457)
+        cmds.setAttr("ring_01_l.rotateZ", 60.61060457)
+        cmds.setAttr("pinky_03_l.rotateZ", 60.61060457)
+        cmds.setAttr("pinky_02_l.rotateZ", 60.61060457)
+        cmds.setAttr("pinky_01_l.rotateZ", 60.61060457)
+        cmds.setAttr("middle_03_l.rotateZ", 60.61060457)
+        cmds.setAttr("middle_02_l.rotateZ", 60.61060457)
+        cmds.setAttr("middle_01_l.rotateZ", 60.61060457)
+        cmds.setAttr("index_03_l.rotateZ", 60.61060457)
+        cmds.setAttr("index_02_l.rotateZ", 60.61060457)
+        cmds.setAttr("index_01_l.rotateZ", 60.61060457)
+
+        cmds.setAttr("thumb_03_r.rotateZ", 49.02926105)
+        cmds.setAttr("thumb_02_r.rotateZ", 29.37554526)
+        cmds.setAttr("thumb_01_r.rotateZ", -1.060679943)
+        cmds.setAttr("thumb_01_r.rotateY", 16.78500021)
+        cmds.setAttr("thumb_01_r.rotateX", 5.616322924)
+        cmds.setAttr("ring_03_r.rotateZ", 60.61060457)
+        cmds.setAttr("ring_02_r.rotateZ", 60.61060457)
+        cmds.setAttr("ring_01_r.rotateZ", 60.61060457)
+        cmds.setAttr("pinky_03_r.rotateZ", 60.61060457)
+        cmds.setAttr("pinky_02_r.rotateZ", 60.61060457)
+        cmds.setAttr("pinky_01_r.rotateZ", 60.61060457)
+        cmds.setAttr("middle_03_r.rotateZ", 60.61060457)
+        cmds.setAttr("middle_02_r.rotateZ", 60.61060457)
+        cmds.setAttr("middle_01_r.rotateZ", 60.61060457)
+        cmds.setAttr("index_03_r.rotateZ", 60.61060457)
+        cmds.setAttr("index_02_r.rotateZ", 60.61060457)
+        cmds.setAttr("index_01_r.rotateZ", 60.61060457)
+
+        # Delete the history on the meshes so that they get baked down into their posed positions
+        cmds.delete(meshes, ch=True)
+
+        # Move the skeleton back to its bindpose (This script REQUIRES the bind pose to be the pose that we want to export the skeleton to.  This is the case for all known characters created with A.R.T.)
+        cmds.dagPose("root", bp=True, r=True) 
+
+        # Re-import the skinning we exported before onto the mesh so thats its bound in its posed position.
+        cmds.select(meshes)
+        impw.importMultipleSkinWeights("", True)
+        
+        jointsToDelete = []
+
+        # remove unneeded joints
+        jointsToRemove = ["back_plate", "chest_l", "chest_r", "collar", "collar_clav_l", "hose_l", "l_hose_base", "collar_clav_r", "hose_r", "r_hose_base"]
+        jointToTransferTo = "spine_03"
+        transfer.LOD_transferWeights(meshes, jointsToRemove, jointToTransferTo)
+        jointsToDelete.extend(jointsToRemove)
+
+        jointsToRemove = ["glasses", "hair_l", "hair_r", "ponytail_01", "ponytail_02", "ponytail_03"]
+        jointToTransferTo = "head"
+        transfer.LOD_transferWeights(meshes, jointsToRemove, jointToTransferTo)
+        jointsToDelete.extend(jointsToRemove)
+
+        jointsToRemove = ["upperarm_twist_01_l", "shoulder_pad_l", "l_hose_end", "l_pad_end", "upperarm_twist_02_l"]
+        jointToTransferTo = "upperarm_l"
+        transfer.LOD_transferWeights(meshes, jointsToRemove, jointToTransferTo)
+        jointsToDelete.extend(jointsToRemove)
+
+        jointsToRemove = ["upperarm_twist_01_r", "shoulder_pad_r", "r_hose_end", "r_pad_end", "upperarm_twist_02_r"]
+        jointToTransferTo = "upperarm_r"
+        transfer.LOD_transferWeights(meshes, jointsToRemove, jointToTransferTo)
+        jointsToDelete.extend(jointsToRemove)
+
+        jointsToRemove = ["lowerarm_twist_01_l", "lowerarm_twist_02_l"]
+        jointToTransferTo = "lowerarm_l"
+        transfer.LOD_transferWeights(meshes, jointsToRemove, jointToTransferTo)
+        jointsToDelete.extend(jointsToRemove)
+
+        jointsToRemove = ["lowerarm_twist_01_r", "lowerarm_twist_02_r"]
+        jointToTransferTo = "lowerarm_r"
+        transfer.LOD_transferWeights(meshes, jointsToRemove, jointToTransferTo)
+        jointsToDelete.extend(jointsToRemove)
+
+        jointsToRemove = ["hand_plate_l", "index_01_l", "index_02_l", "index_03_l", "middle_01_l", "middle_02_l", "middle_03_l", "pinky_01_l", "pinky_02_l", "pinky_03_l", "ring_01_l", "ring_02_l", "ring_03_l", "thumb_01_l", "thumb_02_l", "thumb_03_l"]
+        jointToTransferTo = "hand_l"
+        transfer.LOD_transferWeights(meshes, jointsToRemove, jointToTransferTo)
+        jointsToDelete.extend(jointsToRemove)
+
+        jointsToRemove = ["hand_plate_r", "index_01_r", "index_02_r", "index_03_r", "middle_01_r", "middle_02_r", "middle_03_r", "pinky_01_r", "pinky_02_r", "pinky_03_r", "ring_01_r", "ring_02_r", "ring_03_r", "thumb_01_r", "thumb_02_r", "thumb_03_r"]
+        jointToTransferTo = "hand_r"
+        transfer.LOD_transferWeights(meshes, jointsToRemove, jointToTransferTo)
+        jointsToDelete.extend(jointsToRemove)
+
+        jointsToRemove = ["staff_bottom", "handle", "hinge_a", "hinge_b", "staff_top", "tip_spin", "top_tip_a", "top_tip_b", "top_hinge"]
+        jointToTransferTo = "weapon"
+        transfer.LOD_transferWeights(meshes, jointsToRemove, jointToTransferTo)
+        jointsToDelete.extend(jointsToRemove)
+
+        jointsToRemove = ["thigh_twist_01_l", "thigh_twist_02_l"]
+        jointToTransferTo = "thigh_l"
+        transfer.LOD_transferWeights(meshes, jointsToRemove, jointToTransferTo)
+        jointsToDelete.extend(jointsToRemove)
+
+        jointsToRemove = ["thigh_twist_01_r", "thigh_twist_02_r"]
+        jointToTransferTo = "thigh_r"
+        transfer.LOD_transferWeights(meshes, jointsToRemove, jointToTransferTo)
+        jointsToDelete.extend(jointsToRemove)
+
+        jointsToRemove = ["calf_twist_01_l", "calf_twist_02_l", "knee_pad_l"]
+        jointToTransferTo = "calf_l"
+        transfer.LOD_transferWeights(meshes, jointsToRemove, jointToTransferTo)
+        jointsToDelete.extend(jointsToRemove)
+
+        jointsToRemove = ["calf_twist_01_r", "calf_twist_02_r", "knee_pad_r"]
+        jointToTransferTo = "calf_r"
+        transfer.LOD_transferWeights(meshes, jointsToRemove, jointToTransferTo)
+        jointsToDelete.extend(jointsToRemove)
+
+        jointsToRemove = ["ball_l"]
+        jointToTransferTo = "foot_l"
+        transfer.LOD_transferWeights(meshes, jointsToRemove, jointToTransferTo)
+        jointsToDelete.extend(jointsToRemove)
+
+        jointsToRemove = ["ball_r"]
+        jointToTransferTo = "foot_r"
+        transfer.LOD_transferWeights(meshes, jointsToRemove, jointToTransferTo)
+        jointsToDelete.extend(jointsToRemove)
+
+        # Delete the joints we no longer want.
+        for one in jointsToDelete:
+            if cmds.objExists(one):
+                cmds.delete(one)
+                
+        # Export the fbx file.
+        export.exportLOD("Totem_LOD4.fbx", "LOD4")
+        
+        # Re-open the file without saving.
+        fullPath = cmds.file(q = True, sceneName = True)
+        cmds.file(fullPath, open=True, f=True)
+
+        print "LOD4 Exported Successfully!!!!!!!!!!!!!!!"
+
+TotemLOD4()
+
+
diff --git a/Engine/Extras/Maya_AnimationRiggingTools/ARTv1/General/ART/Projects/Orion/Scripts/Totem_Post_Script.py b/Engine/Extras/Maya_AnimationRiggingTools/ARTv1/General/ART/Projects/Orion/Scripts/Totem_Post_Script.py
new file mode 100644
--- /dev/null
+++ b/Engine/Extras/Maya_AnimationRiggingTools/ARTv1/General/ART/Projects/Orion/Scripts/Totem_Post_Script.py
@@ -0,0 +1,889 @@
+
+#################################################################################################################
+#################################################################################################################
+#################################################################################################################
+#################################################################################################################
+def hideAndResize():
+    #hide stuff
+    cmds.setAttr("chest_r_anim_grp.v", 0, lock = True)
+    cmds.setAttr("chest_l_anim_grp.v", 0, lock = True)
+    cmds.setAttr("hose_l_anim_grp.v", 0, lock = True)
+    cmds.setAttr("hose_r_anim_grp.v", 0, lock = True)
+    cmds.setAttr("collar_clav_l_anim_grp.v", 0, lock = True)
+    cmds.setAttr("collar_clav_r_anim_grp.v", 0, lock = True)
+    cmds.setAttr("hair_r_anim_grp.v", 0, lock = True)
+    cmds.setAttr("hair_l_anim_grp.v", 0, lock = True)
+    cmds.setAttr("fk_ponytail_01_lag_grp.v", 0, lock = True)
+    cmds.setAttr("back_plate_anim_grp.v", 0, lock = True)
+    cmds.setAttr("top_tip_b_anim_grp.v", 0, lock = True)
+    cmds.setAttr("top_tip_a_anim_grp.v", 0, lock = True)
+    cmds.setAttr("hinge_a_anim_grp.v", 0, lock = True)
+    cmds.setAttr("hinge_b_anim_grp.v", 0, lock = True)
+    cmds.setAttr("l_pad_end_anim_grp.v", 0, lock = True)
+    cmds.setAttr("r_pad_end_anim_grp.v", 0, lock = True)
+    cmds.setAttr("l_hose_base_anim_grp.v", 0, lock = True)
+    cmds.setAttr("r_hose_base_anim_grp.v", 0, lock = True)
+    cmds.setAttr("l_hose_end_anim_grp.v", 0, lock = True)
+    cmds.setAttr("r_hose_end_anim_grp.v", 0, lock = True)
+
+    #resize stuff
+    cmds.select("shoulder_pad_l_anim.cv[*]")
+    cmds.scale(0.35, 1.25, 0.35)
+
+    cmds.select("shoulder_pad_r_anim.cv[*]")
+    cmds.scale(0.35, 1.25, 0.35)
+
+    cmds.select("glasses_anim.cv[*]")
+    cmds.scale(0.17, .17, 0.17)
+
+    cmds.select("knee_pad_l_anim.cv[*]")
+    cmds.scale(0.3, .3, 0.3)
+
+    cmds.select("knee_pad_r_anim.cv[*]")
+    cmds.scale(0.3, .3, 0.3)
+
+    cmds.select("weapon_l_anim.cv[*]")
+    cmds.scale(0.0, 0.5, 0.5)
+    
+    cmds.select("weapon_anim.cv[*]")
+    cmds.scale(0.4, 0.4, 0.4)
+    
+    cmds.select("hand_plate_l_anim.cv[*]")
+    cmds.scale(0.4, 0.4, 0.0)
+    cmds.move(0, 0, 1, r = True, os = True, ws = True)
+
+    cmds.select("hand_plate_r_anim.cv[*]")
+    cmds.scale(0.4, 0.4, 0.0)
+    cmds.move(0, 0, 1, r = True, os = True, ws = True)
+
+    cmds.select("top_hinge_anim.cv[*]")
+    cmds.scale(0.0, 1.0, 1.0)
+
+    cmds.select("staff_top_anim.cv[*]")
+    cmds.scale(0.4, 0.4, 0.4)
+
+    cmds.select("staff_bottom_anim.cv[*]")
+    cmds.scale(0.4, 0.4, 0.4)
+
+    cmds.select("handle_anim.cv[*]")
+    cmds.scale(0.3, 0.3, 0.3)
+
+    cmds.select("tip_spin_anim.cv[*]")
+    cmds.scale(1, 0.0, 1)
+    
+    
+    #lock stuff
+    for attr in [".tx", ".ty", ".tz", ".ry", ".rz", ".sx", ".sy", ".sz"]:
+        cmds.setAttr("top_hinge_anim" + attr, lock = True, keyable = False)
+    
+
+#################################################################################################################
+#################################################################################################################
+#################################################################################################################
+#################################################################################################################
+def weaponRig():
+    #weapon SDKs
+
+    #add attrs
+    cmds.addAttr("top_hinge_anim", ln = "tip_a_open", keyable = True, dv = 0, min = 0, max = 10)
+    cmds.addAttr("top_hinge_anim", ln = "tip_b_open", keyable = True, dv = 0, min = 0, max = 10)
+    cmds.addAttr("top_hinge_anim", ln = "tip_a_hinge", keyable = True, dv = 0, min = 0, max = 10)
+    cmds.addAttr("top_hinge_anim", ln = "tip_b_hinge", keyable = True, dv = 0, min = 0, max = 10)
+    cmds.addAttr("top_hinge_anim", ln = "reload", keyable = True, dv = 0, min = 0, max = 10)
+    cmds.addAttr("top_hinge_anim", ln = "bottom_a_open", keyable = True, dv = 0)
+    cmds.addAttr("top_hinge_anim", ln = "bottom_b_open", keyable = True, dv = 0)
+    cmds.addAttr("top_hinge_anim", ln = "pieceSlider", keyable = True, dv = 0, min = 0, max = 10)
+    cmds.addAttr("top_hinge_anim", ln = "coreScale", keyable = True, dv = 1, min = 0, max = 1)
+
+
+    #core scale
+    cmds.setAttr("top_hinge_anim.coreScale", 1)
+    cmds.setAttr("tri_weapon_corrective.core_shrink", 0)
+    cmds.setDrivenKeyframe("tri_weapon_corrective.core_shrink", cd = "top_hinge_anim.coreScale")
+
+
+    cmds.setAttr("top_hinge_anim.coreScale", 0)
+    cmds.setAttr("tri_weapon_corrective.core_shrink", 1)
+    cmds.setDrivenKeyframe("tri_weapon_corrective.core_shrink", cd = "top_hinge_anim.coreScale")
+    cmds.setAttr("top_hinge_anim.coreScale", 1)
+
+
+    #piece morph
+    cmds.setAttr("top_hinge_anim.pieceSlider", 0)
+    cmds.setAttr("tri_weapon_corrective.weapon_piece_move", 0)
+    cmds.setDrivenKeyframe("tri_weapon_corrective.weapon_piece_move", cd = "top_hinge_anim.pieceSlider")
+    
+    cmds.setAttr("top_hinge_anim.pieceSlider", 10)
+    cmds.setAttr("tri_weapon_corrective.weapon_piece_move", 1)
+    cmds.setDrivenKeyframe("tri_weapon_corrective.weapon_piece_move", cd = "top_hinge_anim.pieceSlider")
+    cmds.setAttr("top_hinge_anim.pieceSlider", 0)
+    
+    #tip a hinge
+    cmds.setAttr("top_hinge_anim.tip_a_hinge", 0)
+    cmds.setAttr("top_tip_a_anim_grp.rx", 0)
+    cmds.setDrivenKeyframe("top_tip_a_anim_grp.rx", cd = "top_hinge_anim.tip_a_hinge")
+    
+    
+    cmds.setAttr("top_hinge_anim.tip_a_hinge", 10)
+    cmds.setAttr("top_tip_a_anim_grp.rx", -45)
+
+    cmds.setDrivenKeyframe("top_tip_a_anim_grp.rx", cd = "top_hinge_anim.tip_a_hinge")    
+    cmds.setAttr("top_hinge_anim.tip_a_hinge", 0)
+    
+    
+    #tip b hinge
+    cmds.setAttr("top_hinge_anim.tip_b_hinge", 0)
+    cmds.setAttr("top_tip_b_anim_grp.rx", 0)
+    cmds.setDrivenKeyframe("top_tip_b_anim_grp.rx", cd = "top_hinge_anim.tip_b_hinge")
+    
+    
+    cmds.setAttr("top_hinge_anim.tip_b_hinge", 10)
+    cmds.setAttr("top_tip_b_anim_grp.rx", 45)
+
+    cmds.setDrivenKeyframe("top_tip_b_anim_grp.rx", cd = "top_hinge_anim.tip_b_hinge")    
+    cmds.setAttr("top_hinge_anim.tip_b_hinge", 0)
+    
+    
+    #tip a open
+    cmds.setAttr("top_hinge_anim.tip_a_open", 0)
+    cmds.setAttr("top_tip_a_anim_grp.rx", 0)
+    cmds.setAttr("top_tip_a_anim_grp.ty", 0)
+    cmds.setAttr("top_tip_a_anim_grp.tz", 0)
+    cmds.setDrivenKeyframe(["top_tip_a_anim_grp.rx", "top_tip_a_anim_grp.ty", "top_tip_a_anim_grp.tz"], cd = "top_hinge_anim.tip_a_open")
+
+    cmds.setAttr("top_hinge_anim.tip_a_open", 10)
+    cmds.setAttr("top_tip_a_anim_grp.rx", -8.388)
+    cmds.setAttr("top_tip_a_anim_grp.ty", -0.134)
+    cmds.setAttr("top_tip_a_anim_grp.tz", 5.166)
+    cmds.setDrivenKeyframe(["top_tip_a_anim_grp.rx", "top_tip_a_anim_grp.ty", "top_tip_a_anim_grp.tz"], cd = "top_hinge_anim.tip_a_open")
+    cmds.setAttr("top_hinge_anim.tip_a_open", 0)
+
+    #tip b open
+    cmds.setAttr("top_hinge_anim.tip_b_open", 0)
+    cmds.setAttr("top_tip_b_anim_grp.rx", 0)
+    cmds.setAttr("top_tip_b_anim_grp.ty", 0)
+    cmds.setAttr("top_tip_b_anim_grp.tz", 0)
+    cmds.setDrivenKeyframe(["top_tip_b_anim_grp.rx", "top_tip_b_anim_grp.ty", "top_tip_b_anim_grp.tz"], cd = "top_hinge_anim.tip_b_open")
+
+    cmds.setAttr("top_hinge_anim.tip_b_open", 10)
+    cmds.setAttr("top_tip_b_anim_grp.rx", 5)
+    cmds.setAttr("top_tip_b_anim_grp.ty", 1.923)
+    cmds.setAttr("top_tip_b_anim_grp.tz", -4.408)
+    cmds.setDrivenKeyframe(["top_tip_b_anim_grp.rx", "top_tip_b_anim_grp.ty", "top_tip_b_anim_grp.tz"], cd = "top_hinge_anim.tip_b_open")
+    cmds.setAttr("top_hinge_anim.tip_b_open", 0)
+
+
+    #reload
+    cmds.setAttr("top_hinge_anim.reload", 0)
+    cmds.setAttr("weapon_anim_grp.rx", 0)
+    cmds.setAttr("weapon_anim_grp.tx", 0)
+    cmds.setAttr("weapon_anim_grp.ty", 0)
+    cmds.setAttr("weapon_anim_grp.tz", 0)
+    cmds.setDrivenKeyframe(["weapon_anim_grp.rx", "weapon_anim_grp.tx", "weapon_anim_grp.ty", "weapon_anim_grp.tz"], cd = "top_hinge_anim.reload")
+
+    cmds.setAttr("top_hinge_anim.reload", 10)
+    cmds.setAttr("weapon_anim_grp.rx", 0)
+    cmds.setAttr("weapon_anim_grp.tx", 0)
+    cmds.setAttr("weapon_anim_grp.ty", -9.631)
+    cmds.setAttr("weapon_anim_grp.tz", 0.379)
+    cmds.setDrivenKeyframe(["weapon_anim_grp.rx", "weapon_anim_grp.tx", "weapon_anim_grp.ty", "weapon_anim_grp.tz"], cd = "top_hinge_anim.reload")
+    cmds.setAttr("top_hinge_anim.reload", 0)
+
+    #bottom open a
+    cmds.connectAttr("top_hinge_anim.bottom_a_open", "hinge_a_anim_grp.rx")
+    #bottom open b
+    cmds.connectAttr("top_hinge_anim.bottom_b_open", "hinge_b_anim_grp.rx")
+    
+    
+    #TEXT
+    shape = cmds.listRelatives("top_hinge_anim", shapes = True)[0]
+
+    text = cmds.textCurves(ch = 0, f = "Times New Roman|h-13|w400|c0", t = "W")[0]
+    cmds.setAttr(text + ".rz", 90)
+    cmds.setAttr(text + ".scale", 9,9,9, type = "double3")
+
+    cmds.select(text, hi = True)
+    textNodes = cmds.ls(sl = True)
+
+    for node in textNodes:
+        textShape = cmds.listRelatives(node, shapes = True)
+        if textShape != None:
+            cmds.parent(node, world = True)
+            cmds.makeIdentity(node, t = 1, r = 1, s = 1, apply = True)
+            
+            cmds.parent(textShape[0], "top_hinge_anim", shape = True, r = True)
+            cmds.parent(shape, shape = True, rm = True)
+            cmds.delete(node)
+            cmds.delete(text)
+        
+
+
+#################################################################################################################
+#################################################################################################################
+#################################################################################################################
+#################################################################################################################
+def setupShoulderRig():
+    
+    #Set Settings
+    cmds.setAttr("Rig_Settings.lArmMode", 0)
+    cmds.setAttr("Rig_Settings.lClavMode", 0)
+    cmds.setAttr("Rig_Settings.rArmMode", 0)
+    cmds.setAttr("Rig_Settings.rClavMode", 0)
+    
+    # LEFT UP OFF
+    cmds.setAttr("fk_clavicle_l_anim.rx", 0)
+    cmds.setAttr("fk_clavicle_l_anim.ry", 0)
+    cmds.setAttr("fk_clavicle_l_anim.rz", 0)
+    
+    cmds.setAttr("collar_clav_l_anim_grp.tx", 0)
+    cmds.setAttr("collar_clav_l_anim_grp.ty", 0)
+    cmds.setAttr("collar_clav_l_anim_grp.tz", 0)
+    cmds.setAttr("collar_clav_l_anim_grp.rx", 0)
+    cmds.setAttr("collar_clav_l_anim_grp.ry", 0)
+    cmds.setAttr("collar_clav_l_anim_grp.rz", 0)
+    
+    cmds.setAttr("collar_anim_grp.tz", 0)
+    
+    cmds.setAttr("back_plate_anim_grp.ty", 0)
+    cmds.setAttr("back_plate_anim_grp.rz", 0)
+    
+    cmds.setAttr("shoulder_pad_l_anim_grp.ry", 0)
+
+
+    driven = ["collar_clav_l_anim_grp.tx", "collar_clav_l_anim_grp.ty", "collar_clav_l_anim_grp.tz"]
+    driven.extend(["collar_clav_l_anim_grp.rx", "collar_clav_l_anim_grp.ry", "collar_clav_l_anim_grp.rz"])
+    driven.extend(["collar_anim_grp.tz", "back_plate_anim_grp.ty", "back_plate_anim_grp.rz", "shoulder_pad_l_anim_grp.ry"])
+    
+    cmds.setDrivenKeyframe(driven, cd = "angle_dimensions.lArmUpAngle", itt = "spline", ott = "spline")
+    
+    
+    #LEFT UP ON
+    cmds.setAttr("fk_clavicle_l_anim.rx", 0)
+    cmds.setAttr("fk_clavicle_l_anim.ry", -53.513)
+    cmds.setAttr("fk_clavicle_l_anim.rz", 0)
+    
+    cmds.setAttr("collar_clav_l_anim_grp.tx", 0)
+    cmds.setAttr("collar_clav_l_anim_grp.ty", 0)
+    cmds.setAttr("collar_clav_l_anim_grp.tz", 2.962)
+    cmds.setAttr("collar_clav_l_anim_grp.rx", 0)
+    cmds.setAttr("collar_clav_l_anim_grp.ry", 0)
+    cmds.setAttr("collar_clav_l_anim_grp.rz", 0)
+    
+    cmds.setAttr("collar_anim_grp.tz", 1.71)
+    
+    cmds.setAttr("back_plate_anim_grp.ty", 0)
+    cmds.setAttr("back_plate_anim_grp.rz", 0)
+    
+    cmds.setAttr("shoulder_pad_l_anim_grp.ry", 50)
+    
+    cmds.setDrivenKeyframe(driven, cd = "angle_dimensions.lArmUpAngle", itt = "spline", ott = "spline")
+    
+    cmds.selectKey("shoulder_pad_l_anim_grp_rotateY")
+    cmds.setInfinity(pri = "constant", poi = "linear")
+    
+    cmds.setAttr("fk_clavicle_l_anim.rx", 0)
+    cmds.setAttr("fk_clavicle_l_anim.ry", 0)
+    cmds.setAttr("fk_clavicle_l_anim.rz", 0)
+    
+    
+    
+    # LEFT FORWARD OFF
+    cmds.setAttr("fk_clavicle_l_anim.rx", 0)
+    cmds.setAttr("fk_clavicle_l_anim.ry", 0)
+    cmds.setAttr("fk_clavicle_l_anim.rz", 0)
+    
+    cmds.setAttr("collar_clav_l_anim_grp.tx", 0)
+    cmds.setAttr("collar_clav_l_anim_grp.ty", 0)
+    cmds.setAttr("collar_clav_l_anim_grp.tz", 0)
+    cmds.setAttr("collar_clav_l_anim_grp.rx", 0)
+    cmds.setAttr("collar_clav_l_anim_grp.ry", 0)
+    cmds.setAttr("collar_clav_l_anim_grp.rz", 0)
+    
+    cmds.setAttr("collar_anim_grp.tz", 0)
+    
+    cmds.setAttr("back_plate_anim_grp.ty", 0)
+    cmds.setAttr("back_plate_anim_grp.rz", 0)
+    
+    cmds.setDrivenKeyframe(driven, cd = "angle_dimensions.lArmForAngle", itt = "spline", ott = "spline")
+    
+    
+    # LEFT FORWARD ON
+    cmds.setAttr("fk_clavicle_l_anim.rx", 0)
+    cmds.setAttr("fk_clavicle_l_anim.ry", 0)
+    cmds.setAttr("fk_clavicle_l_anim.rz", -54.225)
+    
+    cmds.setAttr("collar_clav_l_anim_grp.tx", .78)
+    cmds.setAttr("collar_clav_l_anim_grp.ty", -1.46)
+    cmds.setAttr("collar_clav_l_anim_grp.tz", .52)
+    cmds.setAttr("collar_clav_l_anim_grp.rx", 0)
+    cmds.setAttr("collar_clav_l_anim_grp.ry", 0)
+    cmds.setAttr("collar_clav_l_anim_grp.rz", 0)
+    
+    cmds.setAttr("collar_anim_grp.tz", .34)
+    
+    cmds.setAttr("back_plate_anim_grp.ty", 0)
+    cmds.setAttr("back_plate_anim_grp.rz", 0)
+    
+    cmds.setDrivenKeyframe(driven, cd = "angle_dimensions.lArmForAngle", itt = "spline", ott = "spline")
+    cmds.setAttr("fk_clavicle_l_anim.rx", 0)
+    cmds.setAttr("fk_clavicle_l_anim.ry", 0)
+    cmds.setAttr("fk_clavicle_l_anim.rz", 0)
+    
+    # LEFT BACK OFF
+    cmds.setAttr("fk_clavicle_l_anim.rx", 0)
+    cmds.setAttr("fk_clavicle_l_anim.ry", 0)
+    cmds.setAttr("fk_clavicle_l_anim.rz", 0)
+    
+    cmds.setAttr("collar_clav_l_anim_grp.tx", 0)
+    cmds.setAttr("collar_clav_l_anim_grp.ty", 0)
+    cmds.setAttr("collar_clav_l_anim_grp.tz", 0)
+    cmds.setAttr("collar_clav_l_anim_grp.rx", 0)
+    cmds.setAttr("collar_clav_l_anim_grp.ry", 0)
+    cmds.setAttr("collar_clav_l_anim_grp.rz", 0)
+    
+    cmds.setAttr("collar_anim_grp.tz", 0)
+    
+    cmds.setAttr("back_plate_anim_grp.ty", 0)
+    cmds.setAttr("back_plate_anim_grp.rz", 0)
+    
+    cmds.setDrivenKeyframe(driven, cd = "angle_dimensions.lArmBackAngle", itt = "spline", ott = "spline")
+    
+    
+    # LEFT BACK ON
+    cmds.setAttr("fk_clavicle_l_anim.rx", 0)
+    cmds.setAttr("fk_clavicle_l_anim.ry", 0)
+    cmds.setAttr("fk_clavicle_l_anim.rz", 37.417)
+    
+    cmds.setAttr("collar_clav_l_anim_grp.tx", .637)
+    cmds.setAttr("collar_clav_l_anim_grp.ty", 5.048)
+    cmds.setAttr("collar_clav_l_anim_grp.tz", 0)
+    cmds.setAttr("collar_clav_l_anim_grp.rx", 0)
+    cmds.setAttr("collar_clav_l_anim_grp.ry", 0)
+    cmds.setAttr("collar_clav_l_anim_grp.rz", 0)
+    
+    cmds.setAttr("collar_anim_grp.tz", 0)
+    
+    cmds.setAttr("back_plate_anim_grp.ty", 2.522)
+    cmds.setAttr("back_plate_anim_grp.rz", 4.665)
+    
+    cmds.setDrivenKeyframe(driven, cd = "angle_dimensions.lArmBackAngle", itt = "spline", ott = "spline")
+    cmds.setAttr("fk_clavicle_l_anim.rx", 0)
+    cmds.setAttr("fk_clavicle_l_anim.ry", 0)
+    cmds.setAttr("fk_clavicle_l_anim.rz", 0)
+    
+    
+    # SHOULDER PAD
+    cmds.setAttr("fk_arm_l_anim.ry", 30)
+    cmds.setAttr("shoulder_pad_l_anim_grp.ry", 0)
+    cmds.setDrivenKeyframe("shoulder_pad_l_anim_grp.ry", cd = "angle_dimensions.lArmDownAngle", itt = "spline", ott = "spline")
+    
+    cmds.setAttr("fk_arm_l_anim.ry", 0)
+    cmds.setAttr("shoulder_pad_l_anim_grp.ry", 27.895)
+    cmds.setDrivenKeyframe("shoulder_pad_l_anim_grp.ry", cd = "angle_dimensions.lArmDownAngle", itt = "spline", ott = "spline")
+    
+    
+    for each in ["collar_clav_l_anim_grp", "shoulder_pad_l_anim_grp", "collar_anim_grp"]:
+        blendWeighht = cmds.listConnections(each, skipConversionNodes = True, type = "blendWeighted")
+        if blendWeighht != None:
+            curves = cmds.listConnections(blendWeighht[0], skipConversionNodes = True, type = "animCurve")
+            
+            for curve in curves:
+                cmds.selectKey(curve)
+                if each != "shoulder_pad_l_anim_grp":
+                    cmds.setInfinity(pri = "linear", poi = "linear")
+                else:
+                    cmds.setInfinity(pri = "constant", poi = "linear")
+
+    
+    
+    #SHOULDER PAD END
+    cmds.setAttr("fk_arm_l_anim.ry", 30)
+    cmds.setAttr("l_pad_end_anim_grp.ry", 0)
+    cmds.setAttr("tri_shoulder_pad_correctives.l_shoulder_pad_collapse",0)
+    cmds.setDrivenKeyframe(["l_pad_end_anim_grp.ry", "tri_shoulder_pad_correctives.l_shoulder_pad_collapse"], cd = "shoulder_pad_l_anim_grp.ry", itt = "spline", ott = "spline")
+    
+    cmds.setAttr("fk_arm_l_anim.ry", 0)
+    cmds.setAttr("l_pad_end_anim_grp.ry", -27.895)
+    cmds.setAttr("tri_shoulder_pad_correctives.l_shoulder_pad_collapse",.4)
+    cmds.setDrivenKeyframe(["l_pad_end_anim_grp.ry", "tri_shoulder_pad_correctives.l_shoulder_pad_collapse"], cd = "shoulder_pad_l_anim_grp.ry", itt = "spline", ott = "spline")
+    
+    
+    cmds.selectKey("l_pad_end_anim_grp_rotateY")
+    cmds.selectKey("tri_shoulder_pad_correctives_l_shoulder_pad_collapse", add = True)
+    cmds.setInfinity(pri = "constant", poi = "linear")
+    
+    
+    
+    
+    # RIGHT UP OFF
+    cmds.setAttr("fk_clavicle_r_anim.rx", 0)
+    cmds.setAttr("fk_clavicle_r_anim.ry", 0)
+    cmds.setAttr("fk_clavicle_r_anim.rz", 0)
+    
+    cmds.setAttr("collar_clav_r_anim_grp.tx", 0)
+    cmds.setAttr("collar_clav_r_anim_grp.ty", 0)
+    cmds.setAttr("collar_clav_r_anim_grp.tz", 0)
+    cmds.setAttr("collar_clav_r_anim_grp.rx", 0)
+    cmds.setAttr("collar_clav_r_anim_grp.ry", 0)
+    cmds.setAttr("collar_clav_r_anim_grp.rz", 0)
+    
+    cmds.setAttr("collar_anim_grp.tz", 0)
+    
+    cmds.setAttr("back_plate_anim_grp.ty", 0)
+    cmds.setAttr("back_plate_anim_grp.rz", 0)
+    
+    cmds.setAttr("shoulder_pad_r_anim_grp.ry", 0)
+
+
+    driven = ["collar_clav_r_anim_grp.tx", "collar_clav_r_anim_grp.ty", "collar_clav_r_anim_grp.tz"]
+    driven.extend(["collar_clav_r_anim_grp.rx", "collar_clav_r_anim_grp.ry", "collar_clav_r_anim_grp.rz"])
+    driven.extend(["collar_anim_grp.tz", "back_plate_anim_grp.ty", "back_plate_anim_grp.rz", "shoulder_pad_r_anim_grp.ry"])
+    
+    cmds.setDrivenKeyframe(driven, cd = "angle_dimensions.rArmUpAngle", itt = "spline", ott = "spline")
+    
+    
+    #RIGHT UP ON
+    cmds.setAttr("fk_clavicle_r_anim.rx", 0)
+    cmds.setAttr("fk_clavicle_r_anim.ry", -53.513)
+    cmds.setAttr("fk_clavicle_r_anim.rz", 0)
+    
+    cmds.setAttr("collar_clav_r_anim_grp.tx", 0)
+    cmds.setAttr("collar_clav_r_anim_grp.ty", 0)
+    cmds.setAttr("collar_clav_r_anim_grp.tz", -2.962)
+    cmds.setAttr("collar_clav_r_anim_grp.rx", 0)
+    cmds.setAttr("collar_clav_r_anim_grp.ry", 0)
+    cmds.setAttr("collar_clav_r_anim_grp.rz", 0)
+    
+    cmds.setAttr("collar_anim_grp.tz", 1.71)
+    
+    cmds.setAttr("back_plate_anim_grp.ty", 0)
+    cmds.setAttr("back_plate_anim_grp.rz", 0)
+    
+    cmds.setAttr("shoulder_pad_r_anim_grp.ry", 50)
+    
+    cmds.setDrivenKeyframe(driven, cd = "angle_dimensions.rArmUpAngle", itt = "spline", ott = "spline")
+    
+    cmds.selectKey("shoulder_pad_r_anim_grp_rotateY")
+    cmds.setInfinity(pri = "constant", poi = "linear")
+    
+    cmds.setAttr("fk_clavicle_r_anim.rx", 0)
+    cmds.setAttr("fk_clavicle_r_anim.ry", 0)
+    cmds.setAttr("fk_clavicle_r_anim.rz", 0)
+    
+    
+    
+    # RIGHT FORWARD OFF
+    cmds.setAttr("fk_clavicle_r_anim.rx", 0)
+    cmds.setAttr("fk_clavicle_r_anim.ry", 0)
+    cmds.setAttr("fk_clavicle_r_anim.rz", 0)
+    
+    cmds.setAttr("collar_clav_r_anim_grp.tx", 0)
+    cmds.setAttr("collar_clav_r_anim_grp.ty", 0)
+    cmds.setAttr("collar_clav_r_anim_grp.tz", 0)
+    cmds.setAttr("collar_clav_r_anim_grp.rx", 0)
+    cmds.setAttr("collar_clav_r_anim_grp.ry", 0)
+    cmds.setAttr("collar_clav_r_anim_grp.rz", 0)
+    
+    cmds.setAttr("collar_anim_grp.tz", 0)
+    
+    cmds.setAttr("back_plate_anim_grp.ty", 0)
+    cmds.setAttr("back_plate_anim_grp.rz", 0)
+    
+    cmds.setDrivenKeyframe(driven, cd = "angle_dimensions.rArmForAngle", itt = "spline", ott = "spline")
+    
+    
+    # RIGHT FORWARD ON
+    cmds.setAttr("fk_clavicle_r_anim.rx", 0)
+    cmds.setAttr("fk_clavicle_r_anim.ry", 0)
+    cmds.setAttr("fk_clavicle_r_anim.rz", -54.225)
+    
+    cmds.setAttr("collar_clav_r_anim_grp.tx", -.78)
+    cmds.setAttr("collar_clav_r_anim_grp.ty", 1.46)
+    cmds.setAttr("collar_clav_r_anim_grp.tz", -.52)
+    cmds.setAttr("collar_clav_r_anim_grp.rx", 0)
+    cmds.setAttr("collar_clav_r_anim_grp.ry", 0)
+    cmds.setAttr("collar_clav_r_anim_grp.rz", 0)
+    
+    cmds.setAttr("collar_anim_grp.tz", .34)
+    
+    cmds.setAttr("back_plate_anim_grp.ty", 0)
+    cmds.setAttr("back_plate_anim_grp.rz", 0)
+    
+    cmds.setDrivenKeyframe(driven, cd = "angle_dimensions.rArmForAngle", itt = "spline", ott = "spline")
+    cmds.setAttr("fk_clavicle_r_anim.rx", 0)
+    cmds.setAttr("fk_clavicle_r_anim.ry", 0)
+    cmds.setAttr("fk_clavicle_r_anim.rz", 0)
+    
+    # RIGHT BACK OFF
+    cmds.setAttr("fk_clavicle_r_anim.rx", 0)
+    cmds.setAttr("fk_clavicle_r_anim.ry", 0)
+    cmds.setAttr("fk_clavicle_r_anim.rz", 0)
+    
+    cmds.setAttr("collar_clav_r_anim_grp.tx", 0)
+    cmds.setAttr("collar_clav_r_anim_grp.ty", 0)
+    cmds.setAttr("collar_clav_r_anim_grp.tz", 0)
+    cmds.setAttr("collar_clav_r_anim_grp.rx", 0)
+    cmds.setAttr("collar_clav_r_anim_grp.ry", 0)
+    cmds.setAttr("collar_clav_r_anim_grp.rz", 0)
+    
+    cmds.setAttr("collar_anim_grp.tz", 0)
+    
+    cmds.setAttr("back_plate_anim_grp.ty", 0)
+    cmds.setAttr("back_plate_anim_grp.rz", 0)
+    
+    cmds.setDrivenKeyframe(driven, cd = "angle_dimensions.rArmBackAngle", itt = "spline", ott = "spline")
+    
+    
+    # RIGHT BACK ON
+    cmds.setAttr("fk_clavicle_r_anim.rx", 0)
+    cmds.setAttr("fk_clavicle_r_anim.ry", 0)
+    cmds.setAttr("fk_clavicle_r_anim.rz", 37.417)
+    
+    cmds.setAttr("collar_clav_r_anim_grp.tx", -.637)
+    cmds.setAttr("collar_clav_r_anim_grp.ty", -5.048)
+    cmds.setAttr("collar_clav_r_anim_grp.tz", 0)
+    cmds.setAttr("collar_clav_r_anim_grp.rx", 0)
+    cmds.setAttr("collar_clav_r_anim_grp.ry", 0)
+    cmds.setAttr("collar_clav_r_anim_grp.rz", 0)
+    
+    cmds.setAttr("collar_anim_grp.tz", 0)
+    
+    cmds.setAttr("back_plate_anim_grp.ty", 2.522)
+    cmds.setAttr("back_plate_anim_grp.rz", -4.665)
+    
+    cmds.setDrivenKeyframe(driven, cd = "angle_dimensions.rArmBackAngle", itt = "spline", ott = "spline")
+    cmds.setAttr("fk_clavicle_r_anim.rx", 0)
+    cmds.setAttr("fk_clavicle_r_anim.ry", 0)
+    cmds.setAttr("fk_clavicle_r_anim.rz", 0)
+    
+    
+    # SHOULDER PAD
+    cmds.setAttr("fk_arm_r_anim.ry", 30)
+    cmds.setAttr("shoulder_pad_r_anim_grp.ry", 0)
+    cmds.setDrivenKeyframe("shoulder_pad_r_anim_grp.ry", cd = "angle_dimensions.rArmDownAngle", itt = "spline", ott = "spline")
+    
+    cmds.setAttr("fk_arm_r_anim.ry", 0)
+    cmds.setAttr("shoulder_pad_r_anim_grp.ry", 27.895)
+    cmds.setDrivenKeyframe("shoulder_pad_r_anim_grp.ry", cd = "angle_dimensions.rArmDownAngle", itt = "spline", ott = "spline")
+    
+    
+    for each in ["collar_clav_r_anim_grp", "shoulder_pad_r_anim_grp", "collar_anim_grp"]:
+        blendWeighht = cmds.listConnections(each, skipConversionNodes = True, type = "blendWeighted")
+        if blendWeighht != None:
+            curves = cmds.listConnections(blendWeighht[0], skipConversionNodes = True, type = "animCurve")
+            
+            for curve in curves:
+                cmds.selectKey(curve)
+                if each != "shoulder_pad_r_anim_grp":
+                    cmds.setInfinity(pri = "linear", poi = "linear")
+                else:
+                    cmds.setInfinity(pri = "constant", poi = "linear")
+
+    
+    
+    #SHOULDER PAD END
+    cmds.setAttr("fk_arm_r_anim.ry", 30)
+    cmds.setAttr("r_pad_end_anim_grp.ry", 0)
+    cmds.setAttr("tri_shoulder_pad_correctives.r_shoulder_pad_collapse",0)
+    cmds.setDrivenKeyframe(["r_pad_end_anim_grp.ry", "tri_shoulder_pad_correctives.r_shoulder_pad_collapse"], cd = "shoulder_pad_r_anim_grp.ry", itt = "spline", ott = "spline")
+    
+    cmds.setAttr("fk_arm_r_anim.ry", 0)
+    cmds.setAttr("r_pad_end_anim_grp.ry", -27.895)
+    cmds.setAttr("tri_shoulder_pad_correctives.r_shoulder_pad_collapse",.4)
+    cmds.setDrivenKeyframe(["r_pad_end_anim_grp.ry", "tri_shoulder_pad_correctives.r_shoulder_pad_collapse"], cd = "shoulder_pad_r_anim_grp.ry", itt = "spline", ott = "spline")
+    
+    cmds.selectKey("r_pad_end_anim_grp_rotateY")
+    cmds.selectKey("tri_shoulder_pad_correctives_r_shoulder_pad_collapse", add = True)
+    cmds.setInfinity(pri = "constant", poi = "linear")
+    
+    
+    #Set Settings
+    cmds.setAttr("Rig_Settings.lArmMode", 1)
+    cmds.setAttr("Rig_Settings.rArmMode", 1)
+    
+#################################################################################################################
+#################################################################################################################
+#################################################################################################################
+#################################################################################################################
+def setupHoses():
+    
+    #Left Side
+    cmds.aimConstraint("l_hose_base_anim", "l_hose_end_anim", aimVector = [-1, 0, 0], upVector = [0,0,1], worldUpType = "scene")
+    cmds.aimConstraint("l_hose_end_anim", "l_hose_base_anim_grp", aimVector = [1, 0, 0], upVector = [0,0,1], worldUpType = "scene")
+    cmds.pointConstraint(["l_hose_base_anim", "l_hose_end_anim"], "hose_l_anim")
+    cmds.aimConstraint("l_hose_end_anim", "hose_l_anim", aimVector = [1, 0, 0], upVector = [0,0,1], worldUpType = "scene")
+    
+    
+    #Right Side
+    cmds.aimConstraint("r_hose_base_anim", "r_hose_end_anim", aimVector = [1, 0, 0], upVector = [0,0,-1], worldUpType = "scene")
+    cmds.aimConstraint("r_hose_end_anim", "r_hose_base_anim_grp", aimVector = [-1, 0, 0], upVector = [0,0,-1], worldUpType = "scene")
+    cmds.pointConstraint(["r_hose_base_anim", "r_hose_end_anim"], "hose_r_anim")
+    cmds.aimConstraint("r_hose_end_anim", "hose_r_anim", aimVector = [-1, 0, 0], upVector = [0,0,-1], worldUpType = "scene")
+    
+    
+    
+#################################################################################################################
+#################################################################################################################
+#################################################################################################################
+#################################################################################################################
+def dynamicsOff():
+
+        cmds.setAttr("ponytail_dyn_anim.chainStartEnvelope", 0)
+        cmds.setAttr("ponytail_dyn_anim.chainStartEnvelope", l=True)
+
+
+#################################################################################################################
+#################################################################################################################
+#################################################################################################################
+#################################################################################################################
+def addCapsule(radius, halfHeight):
+
+    trueHalfHeight = (halfHeight*2)-(radius*2)
+    capsule = cmds.polyCylinder(name="collision_capsule", r=radius, h=trueHalfHeight, sx=18, sy=1, sz=5, ax=[0,0,1], rcp=True, cuv=3, ch=True)
+    cmds.setAttr(capsule[0]+".tz", halfHeight+2.2)
+
+    cmds.setAttr(capsule[0]+".rz", 90)
+    cmds.addAttr(capsule[0], ln="Radius", dv=radius, keyable=True)
+    cmds.addAttr(capsule[0], ln="HalfHeight", dv=halfHeight, keyable=True)
+
+    expr = capsule[1]+".radius = "+capsule[0]+".Radius;\n"+capsule[1]+".height = ("+capsule[0]+".HalfHeight*2) - ("+capsule[0]+".Radius*2);"
+    cmds.expression(s=expr, o=capsule[1], ae=1, uc="all")
+
+    cmds.setAttr(capsule[0]+".overrideEnabled", 1)
+    cmds.setAttr(capsule[0]+".overrideColor", 18)
+    cmds.setAttr(capsule[0]+".overrideShading", 0)
+
+    cmds.addAttr(capsule[0], ln="Shaded", dv=0, min=0, max=1, keyable=True)
+    cmds.connectAttr(capsule[0]+".Shaded", capsule[0]+".overrideShading")
+
+    cmds.parentConstraint("root", capsule[0], mo=True)
+    cmds.select(capsule[0], r=True)
+    capsuleDisplayLayer = cmds.createDisplayLayer(nr=True, name="Collision Capsule")
+    cmds.setAttr(capsuleDisplayLayer+".displayType", 2)
+    cmds.setAttr(capsuleDisplayLayer+".v", 0)
+
+    for attr in ["tx", "ty", "tz", "rx", "ry", "rz", "sx", "sy", "sz"]:
+        cmds.setAttr(capsule[0]+"."+attr, lock=True, keyable=False, channelBox=False)
+
+    cmds.select(clear=True)
+
+
+#################################################################################################################
+#################################################################################################################
+#################################################################################################################
+#################################################################################################################
+def setupKnees():
+    
+    #L LEG
+    cmds.setAttr("ik_foot_anim_l.tz", 0)
+    cmds.setAttr("knee_pad_l_anim_grp.tx", 0)
+    cmds.setAttr("knee_pad_l_anim_grp.ty", 0)
+    cmds.setAttr("knee_pad_l_anim_grp.tz", 0)
+    cmds.setAttr("knee_pad_l_anim_grp.rx", 0)
+    cmds.setAttr("knee_pad_l_anim_grp.ry", 0)
+    cmds.setAttr("knee_pad_l_anim_grp.rz", 0)
+    
+    cmds.setDrivenKeyframe(["knee_pad_l_anim_grp.tx", "knee_pad_l_anim_grp.ty", "knee_pad_l_anim_grp.tz", "knee_pad_l_anim_grp.rx", "knee_pad_l_anim_grp.ry", "knee_pad_l_anim_grp.rz"], cd = "angle_dimensions.knee_l_bend_angle", itt = "spline", ott = "spline")
+
+    cmds.setAttr("ik_foot_anim_l.tz", 66)
+    cmds.setAttr("knee_pad_l_anim_grp.tx", 0)
+    cmds.setAttr("knee_pad_l_anim_grp.ty", 3.74)
+    cmds.setAttr("knee_pad_l_anim_grp.tz", .914)
+    cmds.setAttr("knee_pad_l_anim_grp.rx", -53.492)
+    cmds.setAttr("knee_pad_l_anim_grp.ry", 0)
+    cmds.setAttr("knee_pad_l_anim_grp.rz", 0)
+    
+    cmds.setDrivenKeyframe(["knee_pad_l_anim_grp.tx", "knee_pad_l_anim_grp.ty", "knee_pad_l_anim_grp.tz", "knee_pad_l_anim_grp.rx", "knee_pad_l_anim_grp.ry", "knee_pad_l_anim_grp.rz"], cd = "angle_dimensions.knee_l_bend_angle", itt = "spline", ott = "spline")
+
+    blendWeighht = cmds.listConnections("knee_pad_l_anim_grp", skipConversionNodes = True, type = "blendWeighted")
+    if blendWeighht != None:
+        curves = cmds.listConnections(blendWeighht[0], skipConversionNodes = True, type = "animCurve")
+        
+        for curve in curves:
+            cmds.selectKey(curve)
+            cmds.setInfinity(pri = "constant", poi = "linear")
+    
+    cmds.setAttr("ik_foot_anim_l.tz", 0)
+                
+                
+    
+    #R LEG
+    cmds.setAttr("ik_foot_anim_r.tz", 0)
+    cmds.setAttr("knee_pad_r_anim_grp.tx", 0)
+    cmds.setAttr("knee_pad_r_anim_grp.ty", 0)
+    cmds.setAttr("knee_pad_r_anim_grp.tz", 0)
+    cmds.setAttr("knee_pad_r_anim_grp.rx", 0)
+    cmds.setAttr("knee_pad_r_anim_grp.ry", 0)
+    cmds.setAttr("knee_pad_r_anim_grp.rz", 0)
+    
+    cmds.setDrivenKeyframe(["knee_pad_r_anim_grp.tx", "knee_pad_r_anim_grp.ty", "knee_pad_r_anim_grp.tz", "knee_pad_r_anim_grp.rx", "knee_pad_r_anim_grp.ry", "knee_pad_r_anim_grp.rz"], cd = "angle_dimensions.knee_r_bend_angle", itt = "spline", ott = "spline")
+
+    cmds.setAttr("ik_foot_anim_r.tz", 66)
+    cmds.setAttr("knee_pad_r_anim_grp.tx", 0)
+    cmds.setAttr("knee_pad_r_anim_grp.ty", -3.74)
+    cmds.setAttr("knee_pad_r_anim_grp.tz", -.914)
+    cmds.setAttr("knee_pad_r_anim_grp.rx", -53.492)
+    cmds.setAttr("knee_pad_r_anim_grp.ry", 0)
+    cmds.setAttr("knee_pad_r_anim_grp.rz", 0)
+    
+    cmds.setDrivenKeyframe(["knee_pad_r_anim_grp.tx", "knee_pad_r_anim_grp.ty", "knee_pad_r_anim_grp.tz", "knee_pad_r_anim_grp.rx", "knee_pad_r_anim_grp.ry", "knee_pad_r_anim_grp.rz"], cd = "angle_dimensions.knee_r_bend_angle", itt = "spline", ott = "spline")
+
+    blendWeighht = cmds.listConnections("knee_pad_r_anim_grp", skipConversionNodes = True, type = "blendWeighted")
+    if blendWeighht != None:
+        curves = cmds.listConnections(blendWeighht[0], skipConversionNodes = True, type = "animCurve")
+        
+        for curve in curves:
+            cmds.selectKey(curve)
+            cmds.setInfinity(pri = "constant", poi = "linear")
+    cmds.setAttr("ik_foot_anim_r.tz", 0)
+    
+#################################################################################################################
+#################################################################################################################
+#################################################################################################################
+#################################################################################################################
+def createGameCam():
+    gameCam = cmds.camera(centerOfInterest=5, focalLength=18, lensSqueezeRatio=1, cameraScale=1, horizontalFilmAperture=1.41732, horizontalFilmOffset=0, verticalFilmAperture=0.94488, verticalFilmOffset=0, filmFit="Overscan", overscan=1.05, motionBlur=0, shutterAngle=144, nearClipPlane=1.0, farClipPlane=10000, orthographic=0, orthographicWidth=30, panZoomEnabled=0, horizontalPan=0, verticalPan=0, zoom=1, displayFilmGate=True)[0]
+
+    gameCam = cmds.rename(gameCam, "GameCam")
+    cmds.camera(gameCam, e=True, horizontalFilmAperture=1.6724376, focalLength=21.2399857011)
+    cmds.setAttr(gameCam+".displayGateMaskOpacity", 1)
+    cmds.setAttr(gameCam+".displayGateMaskColor", 0, 0, 0, type="double3")
+    cmds.setAttr(gameCam+".r", 90, 0, 180, type="double3")
+
+    # Create the control
+    control = cmds.circle(name=gameCam+"_Ctrl", ch=True, o=True, r=15)[0]
+    cmds.makeIdentity(control, a=True, t=True, r=True, s=True)
+
+    # Create groups for the sdk and the constraints
+    const_grp = cmds.group(control, name=control+"_grp")
+    if cmds.objExists("rig_grp"):
+        cmds.parent(const_grp, "rig_grp")
+
+    # Parent Constrain the cam to the control
+    cmds.parentConstraint(control, gameCam, mo=True)
+    cmds.connectAttr(control+".visibility", gameCam+".visibility", f=True)
+    cmds.setAttr(const_grp+".v", 0)
+    cmds.setAttr(const_grp+".t", 0, 500, 200, type="double3")
+    
+#################################################################################################################
+#################################################################################################################
+#################################################################################################################
+#################################################################################################################
+def addCustomAttrs():
+    attrs = cmds.listAttr("root")
+    exists = False
+
+    for each in attrs:
+        if each.find("DistanceCurve") != -1:
+            exists = True
+    if exists == False:
+        cmds.addAttr("root", ln="DistanceCurve", keyable=True)
+        
+        
+#################################################################################################################
+#################################################################################################################
+#################################################################################################################
+#################################################################################################################
+
+
+def addIKtoExport():
+    
+    # ------------------------------------------------------------------
+    # ADD the IK bones to the Export file.
+    #save the AnimRig file before opening the export file.
+    cmds.file(save = True, type = "mayaBinary", force = True)
+
+    #ADD IK BONES
+    #open export file
+    sceneName = cmds.file(q = True, sceneName = True)
+    exportFile = sceneName.partition("AnimRigs")[0] + "ExportFiles" +  sceneName.partition("AnimRigs")[2].partition(".mb")[0] + "_Export.mb"
+    cmds.file(exportFile, open = True, force = True)
+
+    #add ik bones
+    try:
+		# create the joints
+		cmds.select(clear = True)
+		ikFootRoot = cmds.joint(name = "ik_foot_root")
+		cmds.select(clear = True)
+		
+		cmds.select(clear = True)
+		ikFootLeft = cmds.joint(name = "ik_foot_l")
+		cmds.select(clear = True)
+		
+		cmds.select(clear = True)
+		ikFootRight = cmds.joint(name = "ik_foot_r")
+		cmds.select(clear = True)
+		
+		cmds.select(clear = True)
+		ikHandRoot = cmds.joint(name = "ik_hand_root")
+		cmds.select(clear = True)
+		
+		cmds.select(clear = True)
+		ikHandGun = cmds.joint(name = "ik_hand_gun")
+		cmds.select(clear = True)
+		
+		cmds.select(clear = True)
+		ikHandLeft = cmds.joint(name = "ik_hand_l")
+		cmds.select(clear = True)
+		
+		cmds.select(clear = True)
+		ikHandRight = cmds.joint(name = "ik_hand_r")
+		cmds.select(clear = True)
+		
+		# create hierarhcy
+		cmds.parent(ikFootRoot, "root")
+		cmds.parent(ikHandRoot, "root")
+		cmds.parent(ikFootLeft, ikFootRoot)
+		cmds.parent(ikFootRight, ikFootRoot)
+		cmds.parent(ikHandGun, ikHandRoot)
+		cmds.parent(ikHandLeft, ikHandGun)
+		cmds.parent(ikHandRight, ikHandGun)
+		
+		# constrain the joints
+		leftHandConstraint = cmds.parentConstraint("hand_l", ikHandLeft)[0]
+		rightHandGunConstraint = cmds.parentConstraint("hand_r", ikHandGun)[0]
+		rightHandConstraint = cmds.parentConstraint("hand_r", ikHandRight)[0]
+		leftFootConstraint = cmds.parentConstraint("foot_l", ikFootLeft)[0]
+		rightFootConstraint = cmds.parentConstraint("foot_r", ikFootRight)[0]
+	
+    except:
+		cmds.warning("Something went wrong")
+
+    #save the export file.
+    cmds.file(save = True, type = "mayaBinary", force = True)
+
+
+    #delete stuff from export file
+    cmds.delete("Totem_Geo")
+    cmds.select(["correctives", "rig_pose_locs"], hi = True)
+    selection = cmds.ls(sl = True)
+    for each in selection:
+        cmds.lockNode(each, lock = False)
+    cmds.select(clear = True)
+
+    cmds.delete("correctives")
+    cmds.delete("rig_pose_locs")
+
+    #save the export file.
+    cmds.file(save = True, type = "mayaBinary", force = True)
+
+
+#################################################################################################################
+#################################################################################################################
+#################################################################################################################
+#################################################################################################################
+#run functions
+hideAndResize()
+weaponRig()
+setupShoulderRig()
+setupHoses()
+setupKnees()
+
+#common
+dynamicsOff()
+addCapsule(42, 86)
+createGameCam()
+addCustomAttrs()
+addIKtoExport()
+
+
diff --git a/Engine/Extras/Maya_AnimationRiggingTools/ARTv1/General/ART/Projects/Orion/Scripts/Totem_Pre_Script.py b/Engine/Extras/Maya_AnimationRiggingTools/ARTv1/General/ART/Projects/Orion/Scripts/Totem_Pre_Script.py
new file mode 100644
--- /dev/null
+++ b/Engine/Extras/Maya_AnimationRiggingTools/ARTv1/General/ART/Projects/Orion/Scripts/Totem_Pre_Script.py
@@ -0,0 +1,87 @@
+import maya.cmds as cmds
+
+if cmds.objExists("ik_foot_root"):
+	cmds.delete("ik_foot_root")
+	
+if cmds.objExists("ik_hand_root"):
+	cmds.delete("ik_hand_root")
+	
+
+def hookupShoulerLocs():
+        
+    cmds.parentConstraint("upperarm_l", "upperarm_l_loc", mo = True)
+    cmds.parentConstraint("clavicle_l", "clavicle_l_loc", mo = True)
+    cmds.parentConstraint("spine_03", "upperarm_l_up_loc", mo = True)
+    cmds.parentConstraint("spine_03", "upperarm_l_forward_loc", mo = True)
+    cmds.parentConstraint("spine_03", "upperarm_l_backward_loc", mo = True)
+    
+    cmds.parentConstraint("lowerarm_l", "lowerarm_l_loc")
+    cmds.parentConstraint("clavicle_l", "lowerarm_l_down_loc", mo = True)
+    
+    cmds.parentConstraint("lowerarm_r", "lowerarm_r_loc")
+    cmds.parentConstraint("clavicle_r", "lowerarm_r_down_loc", mo = True)
+    
+    cmds.parentConstraint("upperarm_r", "upperarm_r_loc", mo = True)
+    cmds.parentConstraint("clavicle_r", "clavicle_r_loc", mo = True)
+    cmds.parentConstraint("spine_03", "upperarm_r_up_loc", mo = True)
+    cmds.parentConstraint("spine_03", "upperarm_r_forward_loc", mo = True)
+    cmds.parentConstraint("spine_03", "upperarm_r_backward_loc", mo = True)
+    
+    
+    
+def hookupLegCorrectives():
+        
+    cmds.parentConstraint("calf_l", "calf_l_loc", mo = True)
+    cmds.parentConstraint("thigh_l", "thigh_l_loc", mo = True)
+    cmds.parentConstraint("pelvis", "calf_l_forward_loc", mo = True)
+    cmds.parentConstraint("pelvis", "calf_l_backward_loc", mo = True)
+    cmds.parentConstraint("foot_l", "foot_l_loc", mo = True)
+    
+    cmds.parentConstraint("calf_r", "calf_r_loc", mo = True)
+    cmds.parentConstraint("thigh_r", "thigh_r_loc", mo = True)
+    cmds.parentConstraint("pelvis", "calf_r_forward_loc", mo = True)
+    cmds.parentConstraint("pelvis", "calf_r_backward_loc", mo = True)
+    cmds.parentConstraint("foot_r", "foot_r_loc", mo = True)
+    
+    cmds.setAttr("thigh_l.rotateZ", 0)
+    cmds.setAttr("tri_leg_correctives.l_leg_forward", 0)
+    cmds.setDrivenKeyframe("tri_leg_correctives.l_leg_forward", cd = "angle_dimensions.lForAngle")
+    
+    cmds.setAttr("thigh_l.rotateZ", 57)
+    cmds.setAttr("tri_leg_correctives.l_leg_forward", 1)
+    cmds.setDrivenKeyframe("tri_leg_correctives.l_leg_forward", cd = "angle_dimensions.lForAngle")
+    cmds.setAttr("thigh_l.rotateZ", 0)
+    
+    
+    cmds.setAttr("thigh_l.rotateZ", 0)
+    cmds.setAttr("tri_leg_correctives.left_butt_scrunch_corrective", 0)
+    cmds.setDrivenKeyframe("tri_leg_correctives.left_butt_scrunch_corrective", cd = "angle_dimensions.lBacAngle")
+    
+    cmds.setAttr("thigh_l.rotateZ", -35)
+    cmds.setAttr("tri_leg_correctives.left_butt_scrunch_corrective", 1)
+    cmds.setDrivenKeyframe("tri_leg_correctives.left_butt_scrunch_corrective", cd = "angle_dimensions.lBacAngle")
+    cmds.setAttr("thigh_l.rotateZ", 0)
+    
+    
+    
+    cmds.setAttr("thigh_r.rotateZ", 0)
+    cmds.setAttr("tri_leg_correctives.r_leg_forward", 0)
+    cmds.setDrivenKeyframe("tri_leg_correctives.r_leg_forward", cd = "angle_dimensions.rForAngle")
+    
+    cmds.setAttr("thigh_r.rotateZ", 57)
+    cmds.setAttr("tri_leg_correctives.r_leg_forward", 1)
+    cmds.setDrivenKeyframe("tri_leg_correctives.r_leg_forward", cd = "angle_dimensions.rForAngle")
+    cmds.setAttr("thigh_r.rotateZ", 0)
+    
+    
+    cmds.setAttr("thigh_r.rotateZ", 0)
+    cmds.setAttr("tri_leg_correctives.right_butt_scrunch_corrective", 0)
+    cmds.setDrivenKeyframe("tri_leg_correctives.right_butt_scrunch_corrective", cd = "angle_dimensions.rBacAngle")
+    
+    cmds.setAttr("thigh_r.rotateZ", -35)
+    cmds.setAttr("tri_leg_correctives.right_butt_scrunch_corrective", 1)
+    cmds.setDrivenKeyframe("tri_leg_correctives.right_butt_scrunch_corrective", cd = "angle_dimensions.rBacAngle")
+    cmds.setAttr("thigh_r.rotateZ", 0)
+    
+#hookupLegCorrectives()
+hookupShoulerLocs()
\ No newline at end of file
