Branching

diff --git a/Engine/Extras/Maya_AnimationRiggingTools/ARTv1/General/ART/Projects/Orion/Scripts/TwinBlast_CreateLODs.py b/Engine/Extras/Maya_AnimationRiggingTools/ARTv1/General/ART/Projects/Orion/Scripts/TwinBlast_CreateLODs.py
new file mode 100644
--- /dev/null
+++ b/Engine/Extras/Maya_AnimationRiggingTools/ARTv1/General/ART/Projects/Orion/Scripts/TwinBlast_CreateLODs.py
@@ -0,0 +1,568 @@
+import maya.cmds as cmds
+import customMayaMenu as menu
+import Tools.ART_ImportMultipleWeights as impw
+import Tools.ART_TransferWeights as transfer
+import Tools.ART_ExportLODtoFBX as export
+reload(export)
+
+################################################## TWINBLAST #########################################################################
+# Warn the user that this operation will save their file.
+saveResult = cmds.confirmDialog(title = "Save Warning", message="In order to continue your file must be saved.  Would you like to save it?  If yes it will be saved and after the operation is complete your file will be re-opened.", button = ["Yes", "Cancel"], defaultButton='Yes', cancelButton='Cancel', dismissString='Cancel')
+
+# If they are ok with saving the file, warn them that it will also delete all of the current files in their skin weights directory.
+if saveResult == "Yes":
+    deleteSkinFilesResult = cmds.confirmDialog(title = "SkinWeights Warning", message="In order to continue we must delete all of the files in your ..\General\ART\SkinWeights\ directory.  Would you like to continue?.", button = ["Yes", "Cancel"], defaultButton='Yes', cancelButton='Cancel', dismissString='Cancel')
+    if deleteSkinFilesResult == "Yes":
+        toolsPath = cmds.internalVar(usd = True) + "mayaTools.txt"
+        if os.path.exists(toolsPath):
+            f = open(toolsPath, 'r')
+            mayaToolsDir = f.readline()
+            f.close()
+
+        path = mayaToolsDir + "\General\ART\SkinWeights\\"
+        files = os.listdir(path)
+        #delete all weight cache files in the system folder
+        for file in files:
+            os.remove(path + file)
+            #print path + file
+
+
+
+# TwinBlast LOD1 -------------------------------------------------------------------------------------------------------------
+def TwinBlastLOD1():
+    if saveResult == "Yes":
+
+        cmds.file(save=True)
+
+        meshes = ["boots_geo_LOD1", "eye_wetness_geo_LOD1", "eyeballs_geo_LOD1", "eyelash_geo_LOD1", "eyesbrow_geo_LOD1", "eyeshadow_geo_LOD1", "fingers_01_geo_LOD1", "fingers_02_geo_LOD1", "forearm_cores_geo_LOD1", "forearm_gears_geo_LOD1", "forearm_grenade_gears_geo_LOD1", "grenade_back_geo_LOD1", "grenade_flare_01_geo_LOD1", "grenade_flare_02_geo_LOD1", "grenade_flare_03_geo_LOD1", "grenade_flare_04_geo_LOD1", "grenade_flare_05_geo_LOD1", "grenade_flare_06_geo_LOD1", "grenade_front_geo_LOD1", "hair_geo_LOD1", "hands_geo_LOD1", "hip_pouches_geo_LOD1", "jacket_geo_LOD1", "lower_forearm_plate_A_geo_LOD1", "lower_forearm_plate_B_geo_LOD1", "lower_forearm_plate_C_geo_LOD1", "lower_forearm_plate_D_geo_LOD1", "lower_teeth_geo_LOD1", "necklace_geo_LOD1", "pants_geo_LOD1", "pistol_barrels_l_geo_LOD1", "pistol_barrels_r_geo_LOD1", "pistol_clip_l_geo_LOD1", "pistol_clip_r_geo_LOD1", "pistol_foresrock_bott_l_geo_LOD1", "pistol_forestock_bott_r_geo_LOD1", "pistol_forestock_top_l_geo_LOD1", "pistol_forestock_top_r_geo_LOD1", "pistol_grip_l_geo_LOD1", "pistol_grip_r_geo_LOD1", "pistol_mag_cylinder_l_geo_LOD1", "pistol_mag_cylinder_r_geo_LOD1", "pistol_mag_l_geo_LOD1", "pistol_mag_r_geo_LOD1", "pistol_stock_l_geo_LOD1", "pistol_stock_r_geo_LOD1", "pistol_trigger_l_geo_LOD1", "pistol_trigger_r_geo_LOD1", "tear_duct_geo_LOD1", "tongue_geo_LOD1", "torso_geo_LOD1", "upper_forearm_plate_A_geo_LOD1", "upper_forearm_plate_B_geo_LOD1", "upper_forearm_plate_C_geo_LOD1", "upper_forearm_plate_D_geo_LOD1", "upper_teeth_geo_LOD1", "wrist_cuffs_geo_LOD1"]
+
+        jointsToDelete = []
+
+        # remove unneeded joints
+        jointsToRemove = ["collar_back", "collar_in_l", "collar_in_r", "collar_l", "collar_out_flap_bott_l", "collar_out_flap_top_l", "collar_out_flap_bott_r", "collar_out_flap_top_r", "collar_r", "dogtag"]
+        jointToTransferTo = "spine_03"
+        transfer.LOD_transferWeights(meshes, jointsToRemove, jointToTransferTo)
+        jointsToDelete.extend(jointsToRemove)
+
+        jointsToRemove = ["weapon_forestock_bott_l", "weapon_forestock_top_l", "weapon_stock_l"]
+        jointToTransferTo = "weapon_l"
+        transfer.LOD_transferWeights(meshes, jointsToRemove, jointToTransferTo)
+        jointsToDelete.extend(jointsToRemove)
+
+        jointsToRemove = ["weapon_forestock_bott_r", "weapon_forestock_top_r", "weapon_stock_r"]
+        jointToTransferTo = "weapon_r"
+        transfer.LOD_transferWeights(meshes, jointsToRemove, jointToTransferTo)
+        jointsToDelete.extend(jointsToRemove)
+
+        jointsToRemove = ["grenade_fin_01", "grenade_fin_02", "grenade_fin_03", "grenade_fin_04", "grenade_fin_05", "grenade_fin_06", "grenade_front"]
+        jointToTransferTo = "grenade"
+        transfer.LOD_transferWeights(meshes, jointsToRemove, jointToTransferTo)
+        jointsToDelete.extend(jointsToRemove)
+
+        jointsToRemove = ["eyeball_l", "eyeball_r", "eyebrow_in_l", "eyebrow_in_r", "eyebrow_mid", "eyebrow_out_l", "eyebrow_out_r", "eyelid_bott_l", "eyelid_bott_r", "eyelid_top_l", "eyelid_top_r"]
+        jointToTransferTo = "head"
+        transfer.LOD_transferWeights(meshes, jointsToRemove, jointToTransferTo)
+        jointsToDelete.extend(jointsToRemove)
+
+        # Delete the joints we no longer want.
+        for one in jointsToDelete:
+            if cmds.objExists(one):
+                cmds.delete(one)
+                
+        # Export the fbx file.
+        export.exportLOD("TwinBlast_LOD1.fbx", "LOD1")
+        
+        # Re-open the file without saving.
+        fullPath = cmds.file(q = True, sceneName = True)
+        cmds.file(fullPath, open=True, f=True)
+
+        print "LOD1 Exported Successfully!!!!!!!!!!!!!!!"
+
+TwinBlastLOD1()
+
+# TwinBlast LOD2 -------------------------------------------------------------------------------------------------------------
+def TwinBlastLOD2():
+    if saveResult == "Yes":
+
+        cmds.file(save=True)
+        
+        meshes = ["boots_geo_LOD2", "eye_wetness_geo_LOD2", "eyeballs_geo_LOD2", "eyelash_geo_LOD2", "eyesbrow_geo_LOD2", "eyeshadow_geo_LOD2", "fingers_01_geo_LOD2", "fingers_02_geo_LOD2", "forearm_cores_geo_LOD2", "forearm_gears_geo_LOD2", "forearm_grenade_gears_geo_LOD2", "grenade_back_geo_LOD2", "grenade_flare_01_geo_LOD2", "grenade_flare_02_geo_LOD2", "grenade_flare_03_geo_LOD2", "grenade_flare_04_geo_LOD2", "grenade_flare_05_geo_LOD2", "grenade_flare_06_geo_LOD2", "grenade_front_geo_LOD2", "hair_geo_LOD2", "hands_geo_LOD2", "hip_pouches_geo_LOD2", "jacket_geo_LOD2", "lower_forearm_plate_A_geo_LOD2", "lower_forearm_plate_B_geo_LOD2", "lower_forearm_plate_C_geo_LOD2", "lower_forearm_plate_D_geo_LOD2", "lower_teeth_geo_LOD2", "necklace_geo_LOD2", "pants_geo_LOD2", "pistol_barrels_l_geo_LOD2", "pistol_barrels_r_geo_LOD2", "pistol_clip_l_geo_LOD2", "pistol_clip_r_geo_LOD2", "pistol_foresrock_bott_l_geo_LOD2", "pistol_forestock_bott_r_geo_LOD2", "pistol_forestock_top_l_geo_LOD2", "pistol_forestock_top_r_geo_LOD2", "pistol_grip_l_geo_LOD2", "pistol_grip_r_geo_LOD2", "pistol_mag_cylinder_l_geo_LOD2", "pistol_mag_cylinder_r_geo_LOD2", "pistol_mag_l_geo_LOD2", "pistol_mag_r_geo_LOD2", "pistol_stock_l_geo_LOD2", "pistol_stock_r_geo_LOD2", "pistol_trigger_l_geo_LOD2", "pistol_trigger_r_geo_LOD2", "tear_duct_geo_LOD2", "tongue_geo_LOD2", "torso_geo_LOD2", "upper_forearm_plate_A_geo_LOD2", "upper_forearm_plate_B_geo_LOD2", "upper_forearm_plate_C_geo_LOD2", "upper_forearm_plate_D_geo_LOD2", "upper_teeth_geo_LOD2", "wrist_cuffs_geo_LOD2"]
+
+        # Export a skin weights file for each mesh we plan to operate on.
+        cmds.select(meshes)
+        menu.exportSelectedSkinWeights(True)
+
+        # Set the pose if needed
+        cmds.setAttr("thumb_03_l.rotateZ", 49.02926105)
+        cmds.setAttr("ring_03_l.rotateZ", 60.61060457)
+        cmds.setAttr("pinky_03_l.rotateZ", 60.61060457)
+        cmds.setAttr("middle_03_l.rotateZ", 60.61060457)
+        cmds.setAttr("index_03_l.rotateZ", 60.61060457)
+
+        cmds.setAttr("thumb_03_r.rotateZ", 49.02926105)
+        cmds.setAttr("ring_03_r.rotateZ", 60.61060457)
+        cmds.setAttr("pinky_03_r.rotateZ", 60.61060457)
+        cmds.setAttr("middle_03_r.rotateZ", 60.61060457)
+        cmds.setAttr("index_03_r.rotateZ", 60.61060457)
+
+        # Delete the history on the meshes so that they get baked down into their posed positions
+        cmds.delete(meshes, ch=True)
+
+        # Move the skeleton back to its bindpose (This script REQUIRES the bind pose to be the pose that we want to export the skeleton to.  This is the case for all known characters created with A.R.T.)
+        cmds.dagPose("root", bp=True, r=True) 
+
+        # Re-import the skinning we exported before onto the mesh so thats its bound in its posed position.
+        cmds.select(meshes)
+        impw.importMultipleSkinWeights("", True)
+
+        jointsToDelete = []
+
+         # remove unneeded joints
+        jointsToRemove = ["collar_back", "collar_in_l", "collar_in_r", "collar_l", "collar_out_l", "collar_out_flap_bott_l", "collar_out_flap_top_l", "collar_out_r", "collar_out_flap_bott_r", "collar_out_flap_top_r", "collar_r", "dogtag"]
+        jointToTransferTo = "spine_03"
+        transfer.LOD_transferWeights(meshes, jointsToRemove, jointToTransferTo)
+        jointsToDelete.extend(jointsToRemove)
+
+        jointsToRemove = ["upperarm_twist_02_l"]
+        jointToTransferTo = "upperarm_twist_01_l"
+        transfer.LOD_transferWeights(meshes, jointsToRemove, jointToTransferTo)
+        jointsToDelete.extend(jointsToRemove)
+
+        jointsToRemove = ["upperarm_twist_02_r"]
+        jointToTransferTo = "upperarm_twist_01_r"
+        transfer.LOD_transferWeights(meshes, jointsToRemove, jointToTransferTo)
+        jointsToDelete.extend(jointsToRemove)
+
+        jointsToRemove = ["forearm_gears_l", "low_lowarm_plate_A_l", "low_lowarm_plate_B_l", "low_lowarm_plate_C_l", "low_lowarm_plate_D_l", "lowerarm_twist_01_l", "lowerarm_twist_02_l", "lowerarm_twist_03_l", "up_lowarm_plate_A_l", "up_lowarm_plate_B_l", "up_lowarm_plate_C_l", "up_lowarm_plate_D_l"]
+        jointToTransferTo = "lowerarm_l"
+        transfer.LOD_transferWeights(meshes, jointsToRemove, jointToTransferTo)
+        jointsToDelete.extend(jointsToRemove)
+
+        jointsToRemove = ["forearm_gears_r", "low_lowarm_plate_A_r", "low_lowarm_plate_B_r", "low_lowarm_plate_C_r", "low_lowarm_plate_D_r", "lowerarm_twist_01_r", "lowerarm_twist_02_r", "lowerarm_twist_03_r", "up_lowarm_plate_A_r", "up_lowarm_plate_B_r", "up_lowarm_plate_C_r", "up_lowarm_plate_D_r"]
+        jointToTransferTo = "lowerarm_r"
+        transfer.LOD_transferWeights(meshes, jointsToRemove, jointToTransferTo)
+        jointsToDelete.extend(jointsToRemove)
+
+        jointsToRemove = ["weapon_forestock_bott_l", "weapon_forestock_top_l", "weapon_stock_l"]
+        jointToTransferTo = "weapon_l"
+        transfer.LOD_transferWeights(meshes, jointsToRemove, jointToTransferTo)
+        jointsToDelete.extend(jointsToRemove)
+
+        jointsToRemove = ["weapon_forestock_bott_r", "weapon_forestock_top_r", "weapon_stock_r"]
+        jointToTransferTo = "weapon_r"
+        transfer.LOD_transferWeights(meshes, jointsToRemove, jointToTransferTo)
+        jointsToDelete.extend(jointsToRemove)
+
+        jointsToRemove = ["grenade_fin_01", "grenade_fin_02", "grenade_fin_03", "grenade_fin_04", "grenade_fin_05", "grenade_fin_06", "grenade_front"]
+        jointToTransferTo = "grenade"
+        transfer.LOD_transferWeights(meshes, jointsToRemove, jointToTransferTo)
+        jointsToDelete.extend(jointsToRemove)
+
+        jointsToRemove = ["eyeball_l", "eyeball_r", "eyebrow_in_l", "eyebrow_in_r", "eyebrow_mid", "eyebrow_out_l", "eyebrow_out_r", "eyelid_bott_l", "eyelid_bott_r", "eyelid_top_l", "eyelid_top_r"]
+        jointToTransferTo = "head"
+        transfer.LOD_transferWeights(meshes, jointsToRemove, jointToTransferTo)
+        jointsToDelete.extend(jointsToRemove)
+
+        jointsToRemove = ["thigh_twist_02_l"]
+        jointToTransferTo = "thigh_twist_01_l"
+        transfer.LOD_transferWeights(meshes, jointsToRemove, jointToTransferTo)
+        jointsToDelete.extend(jointsToRemove)
+
+        jointsToRemove = ["thigh_twist_02_r"]
+        jointToTransferTo = "thigh_twist_01_r"
+        transfer.LOD_transferWeights(meshes, jointsToRemove, jointToTransferTo)
+        jointsToDelete.extend(jointsToRemove)
+
+        jointsToRemove = ["calf_twist_01_l"]
+        jointToTransferTo = "calf_l"
+        transfer.LOD_transferWeights(meshes, jointsToRemove, jointToTransferTo)
+        jointsToDelete.extend(jointsToRemove)
+
+        jointsToRemove = ["calf_twist_01_r"]
+        jointToTransferTo = "calf_r"
+        transfer.LOD_transferWeights(meshes, jointsToRemove, jointToTransferTo)
+        jointsToDelete.extend(jointsToRemove)
+
+        jointsToRemove = ["index_03_l"]
+        jointToTransferTo = "index_02_l"
+        transfer.LOD_transferWeights(meshes, jointsToRemove, jointToTransferTo)
+        jointsToDelete.extend(jointsToRemove)
+
+        jointsToRemove = ["middle_03_l"]
+        jointToTransferTo = "middle_02_l"
+        transfer.LOD_transferWeights(meshes, jointsToRemove, jointToTransferTo)
+        jointsToDelete.extend(jointsToRemove)
+
+        jointsToRemove = ["pinky_03_l"]
+        jointToTransferTo = "pinky_02_l"
+        transfer.LOD_transferWeights(meshes, jointsToRemove, jointToTransferTo)
+        jointsToDelete.extend(jointsToRemove)
+
+        jointsToRemove = ["ring_03_l"]
+        jointToTransferTo = "ring_02_l"
+        transfer.LOD_transferWeights(meshes, jointsToRemove, jointToTransferTo)
+        jointsToDelete.extend(jointsToRemove)
+
+        jointsToRemove = ["thumb_03_l"]
+        jointToTransferTo = "thumb_02_l"
+        transfer.LOD_transferWeights(meshes, jointsToRemove, jointToTransferTo)
+        jointsToDelete.extend(jointsToRemove)
+
+        jointsToRemove = ["index_03_r"]
+        jointToTransferTo = "index_02_r"
+        transfer.LOD_transferWeights(meshes, jointsToRemove, jointToTransferTo)
+        jointsToDelete.extend(jointsToRemove)
+
+        jointsToRemove = ["middle_03_r"]
+        jointToTransferTo = "middle_02_r"
+        transfer.LOD_transferWeights(meshes, jointsToRemove, jointToTransferTo)
+        jointsToDelete.extend(jointsToRemove)
+
+        jointsToRemove = ["pinky_03_r"]
+        jointToTransferTo = "pinky_02_r"
+        transfer.LOD_transferWeights(meshes, jointsToRemove, jointToTransferTo)
+        jointsToDelete.extend(jointsToRemove)
+
+        jointsToRemove = ["ring_03_r"]
+        jointToTransferTo = "ring_02_r"
+        transfer.LOD_transferWeights(meshes, jointsToRemove, jointToTransferTo)
+        jointsToDelete.extend(jointsToRemove)
+
+        jointsToRemove = ["thumb_03_r"]
+        jointToTransferTo = "thumb_02_r"
+        transfer.LOD_transferWeights(meshes, jointsToRemove, jointToTransferTo)
+        jointsToDelete.extend(jointsToRemove)
+
+        # Delete the joints we no longer want.
+        for one in jointsToDelete:
+            if cmds.objExists(one):
+                cmds.delete(one)
+                
+        # Export the fbx file.
+        export.exportLOD("TwinBlast_LOD2.fbx", "LOD2")
+
+        # Re-open the file without saving.
+        fullPath = cmds.file(q = True, sceneName = True)
+        cmds.file(fullPath, open=True, f=True)
+
+        print "LOD2 Exported Successfully!!!!!!!!!!!!!!!"
+
+TwinBlastLOD2()
+
+# TwinBlast LOD3 -------------------------------------------------------------------------------------------------------------
+def TwinBlastLOD3():
+    if saveResult == "Yes":
+
+        cmds.file(save=True)
+    
+        meshes = ["boots_geo_LOD3", "eye_wetness_geo_LOD3", "eyeballs_geo_LOD3", "eyelash_geo_LOD3", "eyesbrow_geo_LOD3", "eyeshadow_geo_LOD3", "fingers_01_geo_LOD3", "fingers_02_geo_LOD3", "forearm_cores_geo_LOD3", "forearm_gears_geo_LOD3", "forearm_grenade_gears_geo_LOD3", "grenade_back_geo_LOD3", "grenade_flare_01_geo_LOD3", "grenade_flare_02_geo_LOD3", "grenade_flare_03_geo_LOD3", "grenade_flare_04_geo_LOD3", "grenade_flare_05_geo_LOD3", "grenade_flare_06_geo_LOD3", "grenade_front_geo_LOD3", "hair_geo_LOD3", "hands_geo_LOD3", "hip_pouches_geo_LOD3", "jacket_geo_LOD3", "lower_forearm_plate_A_geo_LOD3", "lower_forearm_plate_B_geo_LOD3", "lower_forearm_plate_C_geo_LOD3", "lower_forearm_plate_D_geo_LOD3", "lower_teeth_geo_LOD3", "necklace_geo_LOD3", "pants_geo_LOD3", "pistol_barrels_l_geo_LOD3", "pistol_barrels_r_geo_LOD3", "pistol_clip_l_geo_LOD3", "pistol_clip_r_geo_LOD3", "pistol_foresrock_bott_l_geo_LOD3", "pistol_forestock_bott_r_geo_LOD3", "pistol_forestock_top_l_geo_LOD3", "pistol_forestock_top_r_geo_LOD3", "pistol_grip_l_geo_LOD3", "pistol_grip_r_geo_LOD3", "pistol_mag_cylinder_l_geo_LOD3", "pistol_mag_cylinder_r_geo_LOD3", "pistol_mag_l_geo_LOD3", "pistol_mag_r_geo_LOD3", "pistol_stock_l_geo_LOD3", "pistol_stock_r_geo_LOD3", "pistol_trigger_l_geo_LOD3", "pistol_trigger_r_geo_LOD3", "tear_duct_geo_LOD3", "tongue_geo_LOD3", "torso_geo_LOD3", "upper_forearm_plate_A_geo_LOD3", "upper_forearm_plate_B_geo_LOD3", "upper_forearm_plate_C_geo_LOD3", "upper_forearm_plate_D_geo_LOD3", "upper_teeth_geo_LOD3", "wrist_cuffs_geo_LOD3"]
+    
+        # Export a skin weights file for each mesh we plan to operate on.
+        cmds.select(meshes)
+        menu.exportSelectedSkinWeights(True)
+
+        # Set the pose if needed
+        cmds.setAttr("thumb_03_l.rotateZ", 49.02926105)
+        cmds.setAttr("thumb_02_l.rotateZ", 29.37554526)
+        cmds.setAttr("thumb_01_l.rotateZ", -1.060679943)
+        cmds.setAttr("thumb_01_l.rotateY", 16.78500021)
+        cmds.setAttr("thumb_01_l.rotateX", 5.616322924)
+        cmds.setAttr("ring_03_l.rotateZ", 60.61060457)
+        cmds.setAttr("ring_02_l.rotateZ", 60.61060457)
+        cmds.setAttr("ring_01_l.rotateZ", 60.61060457)
+        cmds.setAttr("pinky_03_l.rotateZ", 60.61060457)
+        cmds.setAttr("pinky_02_l.rotateZ", 60.61060457)
+        cmds.setAttr("pinky_01_l.rotateZ", 60.61060457)
+        cmds.setAttr("middle_03_l.rotateZ", 60.61060457)
+        cmds.setAttr("middle_02_l.rotateZ", 60.61060457)
+        cmds.setAttr("middle_01_l.rotateZ", 60.61060457)
+        cmds.setAttr("index_03_l.rotateZ", 60.61060457)
+        cmds.setAttr("index_02_l.rotateZ", 29.68388)
+        cmds.setAttr("index_01_l.rotateZ", -15.38361)
+
+        cmds.setAttr("thumb_03_r.rotateZ", 49.02926105)
+        cmds.setAttr("thumb_02_r.rotateZ", 29.37554526)
+        cmds.setAttr("thumb_01_r.rotateZ", -1.060679943)
+        cmds.setAttr("thumb_01_r.rotateY", 16.78500021)
+        cmds.setAttr("thumb_01_r.rotateX", 5.616322924)
+        cmds.setAttr("ring_03_r.rotateZ", 60.61060457)
+        cmds.setAttr("ring_02_r.rotateZ", 60.61060457)
+        cmds.setAttr("ring_01_r.rotateZ", 60.61060457)
+        cmds.setAttr("pinky_03_r.rotateZ", 60.61060457)
+        cmds.setAttr("pinky_02_r.rotateZ", 60.61060457)
+        cmds.setAttr("pinky_01_r.rotateZ", 60.61060457)
+        cmds.setAttr("middle_03_r.rotateZ", 60.61060457)
+        cmds.setAttr("middle_02_r.rotateZ", 60.61060457)
+        cmds.setAttr("middle_01_r.rotateZ", 60.61060457)
+        cmds.setAttr("index_03_r.rotateZ", 60.61060457)
+        cmds.setAttr("index_02_r.rotateZ", 29.68388)
+        cmds.setAttr("index_01_r.rotateZ", -15.38361)
+
+        # Delete the history on the meshes so that they get baked down into their posed positions
+        cmds.delete(meshes, ch=True)
+
+        # Move the skeleton back to its bindpose (This script REQUIRES the bind pose to be the pose that we want to export the skeleton to.  This is the case for all known characters created with A.R.T.)
+        cmds.dagPose("root", bp=True, r=True) 
+
+        # Re-import the skinning we exported before onto the mesh so thats its bound in its posed position.
+        cmds.select(meshes)
+        impw.importMultipleSkinWeights("", True)
+
+        jointsToDelete = []
+
+        # remove unneeded joints
+        jointsToRemove = ["collar_back", "collar_in_l", "collar_in_r", "collar_l", "collar_out_l", "collar_out_flap_bott_l", "collar_out_flap_top_l", "collar_out_r", "collar_out_flap_bott_r", "collar_out_flap_top_r", "collar_r", "dogtag"]
+        jointToTransferTo = "spine_03"
+        transfer.LOD_transferWeights(meshes, jointsToRemove, jointToTransferTo)
+        jointsToDelete.extend(jointsToRemove)
+
+        jointsToRemove = ["upperarm_twist_01_l", "upperarm_twist_02_l"]
+        jointToTransferTo = "upperarm_l"
+        transfer.LOD_transferWeights(meshes, jointsToRemove, jointToTransferTo)
+        jointsToDelete.extend(jointsToRemove)
+
+        jointsToRemove = ["upperarm_twist_01_r", "upperarm_twist_02_r"]
+        jointToTransferTo = "upperarm_r"
+        transfer.LOD_transferWeights(meshes, jointsToRemove, jointToTransferTo)
+        jointsToDelete.extend(jointsToRemove)
+
+        jointsToRemove = ["forearm_gears_l", "low_lowarm_plate_A_l", "low_lowarm_plate_B_l", "low_lowarm_plate_C_l", "low_lowarm_plate_D_l", "lowerarm_twist_01_l", "lowerarm_twist_02_l", "lowerarm_twist_03_l", "up_lowarm_plate_A_l", "up_lowarm_plate_B_l", "up_lowarm_plate_C_l", "up_lowarm_plate_D_l"]
+        jointToTransferTo = "lowerarm_l"
+        transfer.LOD_transferWeights(meshes, jointsToRemove, jointToTransferTo)
+        jointsToDelete.extend(jointsToRemove)
+
+        jointsToRemove = ["forearm_gears_r", "low_lowarm_plate_A_r", "low_lowarm_plate_B_r", "low_lowarm_plate_C_r", "low_lowarm_plate_D_r", "lowerarm_twist_01_r", "lowerarm_twist_02_r", "lowerarm_twist_03_r", "up_lowarm_plate_A_r", "up_lowarm_plate_B_r", "up_lowarm_plate_C_r", "up_lowarm_plate_D_r"]
+        jointToTransferTo = "lowerarm_r"
+        transfer.LOD_transferWeights(meshes, jointsToRemove, jointToTransferTo)
+        jointsToDelete.extend(jointsToRemove)
+
+        jointsToRemove = ["index_01_l", "index_02_l", "index_03_l", "middle_01_l", "middle_02_l", "middle_03_l", "pinky_01_l", "pinky_02_l", "pinky_03_l", "q_weapon_l", "ring_01_l", "ring_02_l", "ring_03_l", "thumb_01_l", "thumb_02_l", "thumb_03_l"]
+        jointToTransferTo = "hand_l"
+        transfer.LOD_transferWeights(meshes, jointsToRemove, jointToTransferTo)
+        jointsToDelete.extend(jointsToRemove)
+
+        jointsToRemove = ["index_01_r", "index_02_r", "index_03_r", "middle_01_r", "middle_02_r", "middle_03_r", "pinky_01_r", "pinky_02_r", "pinky_03_r", "q_weapon_r", "ring_01_r", "ring_02_r", "ring_03_r", "thumb_01_r", "thumb_02_r", "thumb_03_r"]
+        jointToTransferTo = "hand_r"
+        transfer.LOD_transferWeights(meshes, jointsToRemove, jointToTransferTo)
+        jointsToDelete.extend(jointsToRemove)
+
+        jointsToRemove = ["weapon_forestock_bott_l", "weapon_forestock_top_l", "weapon_stock_l"]
+        jointToTransferTo = "weapon_l"
+        transfer.LOD_transferWeights(meshes, jointsToRemove, jointToTransferTo)
+        jointsToDelete.extend(jointsToRemove)
+
+        jointsToRemove = ["weapon_forestock_bott_r", "weapon_forestock_top_r", "weapon_stock_r"]
+        jointToTransferTo = "weapon_r"
+        transfer.LOD_transferWeights(meshes, jointsToRemove, jointToTransferTo)
+        jointsToDelete.extend(jointsToRemove)
+
+        jointsToRemove = ["grenade_fin_01", "grenade_fin_02", "grenade_fin_03", "grenade_fin_04", "grenade_fin_05", "grenade_fin_06", "grenade_front"]
+        jointToTransferTo = "grenade"
+        transfer.LOD_transferWeights(meshes, jointsToRemove, jointToTransferTo)
+        jointsToDelete.extend(jointsToRemove)
+
+        jointsToRemove = ["eyeball_l", "eyeball_r", "eyebrow_in_l", "eyebrow_in_r", "eyebrow_mid", "eyebrow_out_l", "eyebrow_out_r", "eyelid_bott_l", "eyelid_bott_r", "eyelid_top_l", "eyelid_top_r", "jaw"]
+        jointToTransferTo = "head"
+        transfer.LOD_transferWeights(meshes, jointsToRemove, jointToTransferTo)
+        jointsToDelete.extend(jointsToRemove)
+
+        jointsToRemove = ["thigh_twist_01_l", "thigh_twist_02_l"]
+        jointToTransferTo = "thigh_l"
+        transfer.LOD_transferWeights(meshes, jointsToRemove, jointToTransferTo)
+        jointsToDelete.extend(jointsToRemove)
+
+        jointsToRemove = ["thigh_twist_01_r", "thigh_twist_02_r"]
+        jointToTransferTo = "thigh_r"
+        transfer.LOD_transferWeights(meshes, jointsToRemove, jointToTransferTo)
+        jointsToDelete.extend(jointsToRemove)
+
+        jointsToRemove = ["calf_twist_01_l"]
+        jointToTransferTo = "calf_l"
+        transfer.LOD_transferWeights(meshes, jointsToRemove, jointToTransferTo)
+        jointsToDelete.extend(jointsToRemove)
+
+        jointsToRemove = ["calf_twist_01_r"]
+        jointToTransferTo = "calf_r"
+        transfer.LOD_transferWeights(meshes, jointsToRemove, jointToTransferTo)
+        jointsToDelete.extend(jointsToRemove)
+
+        # Delete the joints we no longer want.
+        for one in jointsToDelete:
+            if cmds.objExists(one):
+                cmds.delete(one)
+                
+        # Export the fbx file.
+        export.exportLOD("TwinBlast_LOD3.fbx", "LOD3")
+        
+        # Re-open the file without saving.
+        fullPath = cmds.file(q = True, sceneName = True)
+        cmds.file(fullPath, open=True, f=True)
+
+        print "LOD3 Exported Successfully!!!!!!!!!!!!!!!"
+
+TwinBlastLOD3()
+
+# TwinBlast LOD4 -------------------------------------------------------------------------------------------------------------
+def TwinBlastLOD4():
+    if saveResult == "Yes":
+
+        cmds.file(save=True)
+    
+        meshes = ["boots_geo_LOD4", "eye_wetness_geo_LOD4", "eyeballs_geo_LOD4", "eyelash_geo_LOD4", "eyesbrow_geo_LOD4", "eyeshadow_geo_LOD4", "fingers_01_geo_LOD4", "fingers_02_geo_LOD4", "forearm_cores_geo_LOD4", "forearm_gears_geo_LOD4", "forearm_grenade_gears_geo_LOD4", "grenade_back_geo_LOD4", "grenade_flare_01_geo_LOD4", "grenade_flare_02_geo_LOD4", "grenade_flare_03_geo_LOD4", "grenade_flare_04_geo_LOD4", "grenade_flare_05_geo_LOD4", "grenade_flare_06_geo_LOD4", "grenade_front_geo_LOD4", "hair_geo_LOD4", "hands_geo_LOD4", "hip_pouches_geo_LOD4", "jacket_geo_LOD4", "lower_forearm_plate_A_geo_LOD4", "lower_forearm_plate_B_geo_LOD4", "lower_forearm_plate_C_geo_LOD4", "lower_forearm_plate_D_geo_LOD4", "lower_teeth_geo_LOD4", "necklace_geo_LOD4", "pants_geo_LOD4", "pistol_barrels_l_geo_LOD4", "pistol_barrels_r_geo_LOD4", "pistol_clip_l_geo_LOD4", "pistol_clip_r_geo_LOD4", "pistol_foresrock_bott_l_geo_LOD4", "pistol_forestock_bott_r_geo_LOD4", "pistol_forestock_top_l_geo_LOD4", "pistol_forestock_top_r_geo_LOD4", "pistol_grip_l_geo_LOD4", "pistol_grip_r_geo_LOD4", "pistol_mag_cylinder_l_geo_LOD4", "pistol_mag_cylinder_r_geo_LOD4", "pistol_mag_l_geo_LOD4", "pistol_mag_r_geo_LOD4", "pistol_stock_l_geo_LOD4", "pistol_stock_r_geo_LOD4", "pistol_trigger_l_geo_LOD4", "pistol_trigger_r_geo_LOD4", "tear_duct_geo_LOD4", "tongue_geo_LOD4", "torso_geo_LOD4", "upper_forearm_plate_A_geo_LOD4", "upper_forearm_plate_B_geo_LOD4", "upper_forearm_plate_C_geo_LOD4", "upper_forearm_plate_D_geo_LOD4", "upper_teeth_geo_LOD4", "wrist_cuffs_geo_LOD4"]
+
+        # Export a skin weights file for each mesh we plan to operate on.
+        cmds.select(meshes)
+        menu.exportSelectedSkinWeights(True)
+
+        # Set the pose if needed
+        cmds.setAttr("thumb_03_l.rotateZ", 49.02926105)
+        cmds.setAttr("thumb_02_l.rotateZ", 29.37554526)
+        cmds.setAttr("thumb_01_l.rotateZ", -1.060679943)
+        cmds.setAttr("thumb_01_l.rotateY", 16.78500021)
+        cmds.setAttr("thumb_01_l.rotateX", 5.616322924)
+        cmds.setAttr("ring_03_l.rotateZ", 60.61060457)
+        cmds.setAttr("ring_02_l.rotateZ", 60.61060457)
+        cmds.setAttr("ring_01_l.rotateZ", 60.61060457)
+        cmds.setAttr("pinky_03_l.rotateZ", 60.61060457)
+        cmds.setAttr("pinky_02_l.rotateZ", 60.61060457)
+        cmds.setAttr("pinky_01_l.rotateZ", 60.61060457)
+        cmds.setAttr("middle_03_l.rotateZ", 60.61060457)
+        cmds.setAttr("middle_02_l.rotateZ", 60.61060457)
+        cmds.setAttr("middle_01_l.rotateZ", 60.61060457)
+        cmds.setAttr("index_03_l.rotateZ", 60.61060457)
+        cmds.setAttr("index_02_l.rotateZ", 29.68388)
+        cmds.setAttr("index_01_l.rotateZ", -15.38361)
+
+        cmds.setAttr("thumb_03_r.rotateZ", 49.02926105)
+        cmds.setAttr("thumb_02_r.rotateZ", 29.37554526)
+        cmds.setAttr("thumb_01_r.rotateZ", -1.060679943)
+        cmds.setAttr("thumb_01_r.rotateY", 16.78500021)
+        cmds.setAttr("thumb_01_r.rotateX", 5.616322924)
+        cmds.setAttr("ring_03_r.rotateZ", 60.61060457)
+        cmds.setAttr("ring_02_r.rotateZ", 60.61060457)
+        cmds.setAttr("ring_01_r.rotateZ", 60.61060457)
+        cmds.setAttr("pinky_03_r.rotateZ", 60.61060457)
+        cmds.setAttr("pinky_02_r.rotateZ", 60.61060457)
+        cmds.setAttr("pinky_01_r.rotateZ", 60.61060457)
+        cmds.setAttr("middle_03_r.rotateZ", 60.61060457)
+        cmds.setAttr("middle_02_r.rotateZ", 60.61060457)
+        cmds.setAttr("middle_01_r.rotateZ", 60.61060457)
+        cmds.setAttr("index_03_r.rotateZ", 60.61060457)
+        cmds.setAttr("index_02_r.rotateZ", 29.68388)
+        cmds.setAttr("index_01_r.rotateZ", -15.38361)
+
+        # Delete the history on the meshes so that they get baked down into their posed positions
+        cmds.delete(meshes, ch=True)
+
+        # Move the skeleton back to its bindpose (This script REQUIRES the bind pose to be the pose that we want to export the skeleton to.  This is the case for all known characters created with A.R.T.)
+        cmds.dagPose("root", bp=True, r=True) 
+
+        # Re-import the skinning we exported before onto the mesh so thats its bound in its posed position.
+        cmds.select(meshes)
+        impw.importMultipleSkinWeights("", True)
+        
+        jointsToDelete = []
+
+        # remove unneeded joints
+        jointsToRemove = ["collar_back", "collar_in_l", "collar_in_r", "collar_l", "collar_out_l", "collar_out_flap_bott_l", "collar_out_flap_top_l", "collar_out_r", "collar_out_flap_bott_r", "collar_out_flap_top_r", "collar_r", "dogtag"]
+        jointToTransferTo = "spine_03"
+        transfer.LOD_transferWeights(meshes, jointsToRemove, jointToTransferTo)
+        jointsToDelete.extend(jointsToRemove)
+
+        jointsToRemove = ["upperarm_twist_01_l", "upperarm_twist_02_l"]
+        jointToTransferTo = "upperarm_l"
+        transfer.LOD_transferWeights(meshes, jointsToRemove, jointToTransferTo)
+        jointsToDelete.extend(jointsToRemove)
+
+        jointsToRemove = ["upperarm_twist_01_r", "upperarm_twist_02_r"]
+        jointToTransferTo = "upperarm_r"
+        transfer.LOD_transferWeights(meshes, jointsToRemove, jointToTransferTo)
+        jointsToDelete.extend(jointsToRemove)
+
+        jointsToRemove = ["forearm_gears_l", "low_lowarm_plate_A_l", "low_lowarm_plate_B_l", "low_lowarm_plate_C_l", "low_lowarm_plate_D_l", "lowerarm_twist_01_l", "lowerarm_twist_02_l", "lowerarm_twist_03_l", "up_lowarm_plate_A_l", "up_lowarm_plate_B_l", "up_lowarm_plate_C_l", "up_lowarm_plate_D_l"]
+        jointToTransferTo = "lowerarm_l"
+        transfer.LOD_transferWeights(meshes, jointsToRemove, jointToTransferTo)
+        jointsToDelete.extend(jointsToRemove)
+
+        jointsToRemove = ["forearm_gears_r", "low_lowarm_plate_A_r", "low_lowarm_plate_B_r", "low_lowarm_plate_C_r", "low_lowarm_plate_D_r", "lowerarm_twist_01_r", "lowerarm_twist_02_r", "lowerarm_twist_03_r", "up_lowarm_plate_A_r", "up_lowarm_plate_B_r", "up_lowarm_plate_C_r", "up_lowarm_plate_D_r"]
+        jointToTransferTo = "lowerarm_r"
+        transfer.LOD_transferWeights(meshes, jointsToRemove, jointToTransferTo)
+        jointsToDelete.extend(jointsToRemove)
+
+        jointsToRemove = ["index_01_l", "index_02_l", "index_03_l", "middle_01_l", "middle_02_l", "middle_03_l", "pinky_01_l", "pinky_02_l", "pinky_03_l", "q_weapon_l", "ring_01_l", "ring_02_l", "ring_03_l", "thumb_01_l", "thumb_02_l", "thumb_03_l"]
+        jointToTransferTo = "hand_l"
+        transfer.LOD_transferWeights(meshes, jointsToRemove, jointToTransferTo)
+        jointsToDelete.extend(jointsToRemove)
+
+        jointsToRemove = ["index_01_r", "index_02_r", "index_03_r", "middle_01_r", "middle_02_r", "middle_03_r", "pinky_01_r", "pinky_02_r", "pinky_03_r", "q_weapon_r", "ring_01_r", "ring_02_r", "ring_03_r", "thumb_01_r", "thumb_02_r", "thumb_03_r"]
+        jointToTransferTo = "hand_r"
+        transfer.LOD_transferWeights(meshes, jointsToRemove, jointToTransferTo)
+        jointsToDelete.extend(jointsToRemove)
+
+        jointsToRemove = ["weapon_forestock_bott_l", "weapon_forestock_top_l", "weapon_stock_l"]
+        jointToTransferTo = "weapon_l"
+        transfer.LOD_transferWeights(meshes, jointsToRemove, jointToTransferTo)
+        jointsToDelete.extend(jointsToRemove)
+
+        jointsToRemove = ["weapon_forestock_bott_r", "weapon_forestock_top_r", "weapon_stock_r"]
+        jointToTransferTo = "weapon_r"
+        transfer.LOD_transferWeights(meshes, jointsToRemove, jointToTransferTo)
+        jointsToDelete.extend(jointsToRemove)
+
+        jointsToRemove = ["grenade", "grenade_fin_01", "grenade_fin_02", "grenade_fin_03", "grenade_fin_04", "grenade_fin_05", "grenade_fin_06", "grenade_front"]
+        jointToTransferTo = "lowerarm_r"
+        transfer.LOD_transferWeights(meshes, jointsToRemove, jointToTransferTo)
+        jointsToDelete.extend(jointsToRemove)
+
+        jointsToRemove = ["eyeball_l", "eyeball_r", "eyebrow_in_l", "eyebrow_in_r", "eyebrow_mid", "eyebrow_out_l", "eyebrow_out_r", "eyelid_bott_l", "eyelid_bott_r", "eyelid_top_l", "eyelid_top_r", "jaw"]
+        jointToTransferTo = "head"
+        transfer.LOD_transferWeights(meshes, jointsToRemove, jointToTransferTo)
+        jointsToDelete.extend(jointsToRemove)
+
+        jointsToRemove = ["thigh_twist_01_l", "thigh_twist_02_l"]
+        jointToTransferTo = "thigh_l"
+        transfer.LOD_transferWeights(meshes, jointsToRemove, jointToTransferTo)
+        jointsToDelete.extend(jointsToRemove)
+
+        jointsToRemove = ["thigh_twist_01_r", "thigh_twist_02_r"]
+        jointToTransferTo = "thigh_r"
+        transfer.LOD_transferWeights(meshes, jointsToRemove, jointToTransferTo)
+        jointsToDelete.extend(jointsToRemove)
+
+        jointsToRemove = ["calf_twist_01_l"]
+        jointToTransferTo = "calf_l"
+        transfer.LOD_transferWeights(meshes, jointsToRemove, jointToTransferTo)
+        jointsToDelete.extend(jointsToRemove)
+
+        jointsToRemove = ["calf_twist_01_r"]
+        jointToTransferTo = "calf_r"
+        transfer.LOD_transferWeights(meshes, jointsToRemove, jointToTransferTo)
+        jointsToDelete.extend(jointsToRemove)
+
+        jointsToRemove = ["ball_l"]
+        jointToTransferTo = "foot_l"
+        transfer.LOD_transferWeights(meshes, jointsToRemove, jointToTransferTo)
+        jointsToDelete.extend(jointsToRemove)
+
+        jointsToRemove = ["ball_r"]
+        jointToTransferTo = "foot_r"
+        transfer.LOD_transferWeights(meshes, jointsToRemove, jointToTransferTo)
+        jointsToDelete.extend(jointsToRemove)
+
+        # Delete the joints we no longer want.
+        for one in jointsToDelete:
+            if cmds.objExists(one):
+                cmds.delete(one)
+                
+        # Export the fbx file.
+        export.exportLOD("TwinBlast_LOD4.fbx", "LOD4")
+        
+        # Re-open the file without saving.
+        fullPath = cmds.file(q = True, sceneName = True)
+        cmds.file(fullPath, open=True, f=True)
+
+        print "LOD4 Exported Successfully!!!!!!!!!!!!!!!"
+
+TwinBlastLOD4()
+
+
diff --git a/Engine/Extras/Maya_AnimationRiggingTools/ARTv1/General/ART/Projects/Orion/Scripts/TwinBlast_Post_Script.py b/Engine/Extras/Maya_AnimationRiggingTools/ARTv1/General/ART/Projects/Orion/Scripts/TwinBlast_Post_Script.py
new file mode 100644
--- /dev/null
+++ b/Engine/Extras/Maya_AnimationRiggingTools/ARTv1/General/ART/Projects/Orion/Scripts/TwinBlast_Post_Script.py
@@ -0,0 +1,254 @@
+import maya.cmds as cmds
+# NOTE: If you mirror his movers, lock the weapon controls first so that they dont mirror because they are set up to not.  And putting them back is a pain.
+# NOTE: Use the "GUN_locator___USE_FOR_RIG_POSE" in the scene during the Create Rig Pose phase to put the weapon movers adn the grenade movers back to the correct location.
+
+# ------------------------------------------------------------------
+# scale all of the controls to be more visible with the mesh.
+def scaleControls():
+    sides = ["_l", "_r"]
+    for side in sides:
+        # feet
+        cmds.select(["ik_foot_anim"+side+".cv[0:17]"])
+        cmds.scale(1, 1, 1, r = True, ocp = False)
+        cmds.move(0, 0, -4.78  , r = True)
+
+        cmds.select(["q_weapon"+side+"_anim.cv[0:32]"])
+        cmds.scale(0.4, 0.4, 0.4, r = True, ocp = False)
+        if side == "_l":
+            cmds.move(0, 20, 0, r = True)
+        else:
+            cmds.move(0, 20, 0, r = True)            
+
+        cmds.select(["ult_weapon"+side+"_anim.cv[0:32]"])
+        cmds.scale(0.5, 0.5, 0.5, r = True, ocp = False)
+        if side == "_l":
+            cmds.move(0, 25, 0, r = True)
+        else:
+            cmds.move(0, 25, 0, r = True)            
+
+        cmds.select(["eyelid_bott"+side+"_anim.cv[0:32]"])
+        cmds.scale(0.1, 0.1, 0.01, r = True, ocp = False)
+        if side is "_l":
+            cmds.move(0, -5, 0, r = True, os=True)
+        else:
+            cmds.move(0, 5, 0, r = True, os=True)
+
+        cmds.select(["eyelid_top"+side+"_anim.cv[0:32]"])
+        cmds.scale(0.1, 0.1, 0.01, r = True, ocp = False)
+        if side is "_l":
+            cmds.move(0, -5, 0, r = True, os=True)
+        else:
+            cmds.move(0, 5, 0, r = True, os=True)
+
+        cmds.select(["eyeball"+side+"_anim.cv[0:32]"])
+        cmds.scale(0.24, 0.24, 0.24, r = True, ocp = False)
+        cmds.move(0, 0, 0, r = True, os=True)
+
+        cmds.select(["forearm_gears"+side+"_anim.cv[0:32]"])
+        cmds.scale(0, 1, 1, r = True, ocp = False)
+
+    # Other
+    objects = ["weapon_forestock_bott_l_anim", "weapon_forestock_bott_r_anim", "weapon_forestock_top_l_anim", "weapon_forestock_top_r_anim", "weapon_l_anim", "weapon_r_anim", "weapon_stock_l_anim", "weapon_stock_r_anim", "low_lowarm_plate_A_l_anim", "low_lowarm_plate_B_l_anim", "low_lowarm_plate_C_l_anim", "low_lowarm_plate_D_l_anim", "up_lowarm_plate_A_l_anim", "up_lowarm_plate_B_l_anim", "up_lowarm_plate_C_l_anim", "up_lowarm_plate_D_l_anim", "grenade_anim", "grenade_front_anim", "grenade_fin_01_anim", "grenade_fin_02_anim", "grenade_fin_03_anim", "grenade_fin_04_anim", "grenade_fin_05_anim", "grenade_fin_06_anim", "collar_in_l_anim", "collar_out_l_anim", "collar_out_flap_bott_l_anim", "collar_out_flap_top_l_anim", "low_lowarm_plate_A_r_anim", "low_lowarm_plate_B_r_anim", "low_lowarm_plate_C_r_anim", "low_lowarm_plate_D_r_anim", "up_lowarm_plate_A_r_anim", "up_lowarm_plate_B_r_anim", "up_lowarm_plate_C_r_anim", "up_lowarm_plate_D_r_anim", "collar_in_r_anim", "collar_out_r_anim", "collar_out_flap_bott_r_anim", "collar_out_flap_top_r_anim", "dogtag_anim", "eyebrow_in_l_anim", "eyebrow_in_r_anim", "eyebrow_out_l_anim", "eyebrow_out_r_anim", "eyebrow_mid_anim", "collar_back_anim", "collar_l_anim", "collar_r_anim"]
+    for one in objects:
+        cmds.select([one+".cv[*]"])
+        cmds.scale(0.227057, 0.227057, 0.227057, r = True, ocp = True)
+        cmds.move(0, 0, 0, r = True, os=True, wd=True)
+
+    cmds.select(["jaw_anim.cv[0:32]"])
+    cmds.scale(0.2, 0.2, 0.2, r = True, ocp = True)
+    cmds.move(0, -9, -8, r = True)
+
+    cmds.select(None)
+
+scaleControls()
+
+# ------------------------------------------------------------------
+def rigWeapons():
+    cmds.move(0, 4.4, 24.558962, "ult_weapon_l_anim", r=True, os=True, wd=True)
+    cmds.scale(0.05, 0.05, 0.05, "ult_weapon_l_anim", r=True)
+
+    cmds.move(0, -4.4, -24.558962, "ult_weapon_r_anim", r=True, os=True, wd=True)
+    cmds.scale(0.05, 0.05, 0.05, "ult_weapon_r_anim", r=True)
+
+    cmds.move(0, 4.4, 24.558962, "q_weapon_l_anim", r=True, os=True, wd=True)
+    cmds.scale(0.05, 0.05, 0.05, "q_weapon_l_anim", r=True)
+
+    cmds.move(0, -4.4, -24.558962, "q_weapon_r_anim", r=True, os=True, wd=True)
+    cmds.scale(0.05, 0.05, 0.05, "q_weapon_r_anim", r=True)
+
+    cmds.xform("weapon_l_anim_grp", t=[82.96438708999413, 10.607007776552306, 147.1500662787879], ro=[37.910151258451656, -65.27865699485243, 49.67785824895247])
+    cmds.xform("weapon_r_anim_grp", t=[-82.96438708999413, 10.607007776552306, 147.1500662787879], ro=[37.910151258451656, 65.27865699485243, -49.67785824895247])
+    cmds.xform("grenade_front_anim", t=[0, 0.4509, 0], ro=[0, 0, 0])
+    cmds.xform("grenade_fin_01_anim", t=[0, 0, 0], ro=[-180, 0, 0])
+    cmds.xform("grenade_fin_02_anim", t=[0, 0, 0], ro=[-180, 0, 0])
+    cmds.xform("grenade_fin_03_anim", t=[0, 0, 0], ro=[-180, 0, 0])
+    cmds.xform("grenade_fin_04_anim", t=[0, 0, 0], ro=[-180, 0, 0])
+    cmds.xform("grenade_fin_05_anim", t=[0, 0, 0], ro=[-180, 0, 0])
+    cmds.xform("grenade_fin_06_anim", t=[0, 0, 0], ro=[-180, 0, 0])
+    cmds.xform("grenade_anim_grp", t=[-59.36194438801096, 12.504397019491371, 150.4395268893561], ro=[0.0, 0.0, -90])
+    
+    cmds.setAttr("Rig_Settings.lUpperarmTwistAmount", 0.9)
+    cmds.setAttr("Rig_Settings.lUpperarmTwist2Amount", 0.3)
+    cmds.setAttr("Rig_Settings.rUpperarmTwistAmount", 0.9)
+    cmds.setAttr("Rig_Settings.rUpperarmTwist2Amount", 0.3)
+
+    cmds.setAttr("Rig_Settings.lForearmTwistAmount", 1)
+    cmds.setAttr("Rig_Settings.rForearmTwistAmount", 1)
+    
+rigWeapons()
+
+# ------------------------------------------------------------------
+def invertEyelids():
+    sides = ["_l", "_r"]
+    for side in sides:
+        cmds.delete("driver_eyelid_bott"+side+"_parentConstraint1")
+        cmds.rotate(0, -180, 0, "eyelid_bott"+side+"_anim_grp", r=True, os=True)
+        cmds.parentConstraint("eyelid_bott"+side+"_anim", "driver_eyelid_bott"+side, mo=True)
+invertEyelids()
+
+def invertRightEyeball():
+    cmds.delete("driver_eyeball_r_parentConstraint1")
+    cmds.rotate(0, -180, -180, "eyeball_r_anim_grp", r=True, os=True)
+    cmds.parentConstraint("eyeball_r_anim", "driver_eyeball_r", mo=True)
+invertRightEyeball()
+
+# ------------------------------------------------------------------
+def createGameCam():
+    gameCam = cmds.camera(centerOfInterest=5, focalLength=18, lensSqueezeRatio=1, cameraScale=1, horizontalFilmAperture=1.41732, horizontalFilmOffset=0, verticalFilmAperture=0.94488, verticalFilmOffset=0, filmFit="Overscan", overscan=1.05, motionBlur=0, shutterAngle=144, nearClipPlane=1.0, farClipPlane=10000, orthographic=0, orthographicWidth=30, panZoomEnabled=0, horizontalPan=0, verticalPan=0, zoom=1, displayFilmGate=True)[0]
+    print gameCam
+    gameCam = cmds.rename(gameCam, "GameCam")
+    cmds.camera(gameCam, e=True, horizontalFilmAperture=1.6724376, focalLength=21.2399857011)
+    cmds.setAttr(gameCam+".displayGateMaskOpacity", 1)
+    cmds.setAttr(gameCam+".displayGateMaskColor", 0, 0, 0, type="double3")
+    cmds.setAttr(gameCam+".r", 90, 0, 180, type="double3")
+    
+    # Create the control
+    control = cmds.circle(name=gameCam+"_Ctrl", ch=True, o=True, r=15)[0]
+    cmds.makeIdentity(control, a=True, t=True, r=True, s=True)
+    # Create groups for the sdk and the constraints
+    const_grp = cmds.group(control, name=control+"_grp")
+    if cmds.objExists("rig_grp"):
+        cmds.parent(const_grp, "rig_grp")
+    # Parent Constrain the cam to the control
+    cmds.parentConstraint(control, gameCam, mo=True)
+    
+    cmds.connectAttr(control+".visibility", gameCam+".visibility", f=True)
+    cmds.setAttr(const_grp+".v", 0)
+    
+    cmds.setAttr(const_grp+".t", 0, 500, 200, type="double3")
+createGameCam()
+# ------------------------------------------------------------------
+
+##################################################################################
+##################################################################################
+
+def addCapsule(radius, halfHeight):
+    trueHalfHeight = (halfHeight*2)-(radius*2)
+    capsule = cmds.polyCylinder(name="collision_capsule", r=radius, h=trueHalfHeight, sx=18, sy=1, sz=5, ax=[0,0,1], rcp=True, cuv=3, ch=True)
+    cmds.setAttr(capsule[0]+".tz", halfHeight+2.2)
+    cmds.setAttr(capsule[0]+".rz", 90)
+
+    cmds.addAttr(capsule[0], ln="Radius", dv=radius, keyable=True)
+    cmds.addAttr(capsule[0], ln="HalfHeight", dv=halfHeight, keyable=True)
+
+    expr = capsule[1]+".radius = "+capsule[0]+".Radius;\n"+capsule[1]+".height = ("+capsule[0]+".HalfHeight*2) - ("+capsule[0]+".Radius*2);"
+    cmds.expression(s=expr, o=capsule[1], ae=1, uc="all")
+
+    cmds.setAttr(capsule[0]+".overrideEnabled", 1)
+    cmds.setAttr(capsule[0]+".overrideColor", 18)
+    cmds.setAttr(capsule[0]+".overrideShading", 0)
+
+    cmds.addAttr(capsule[0], ln="Shaded", dv=0, min=0, max=1, keyable=True)
+    cmds.connectAttr(capsule[0]+".Shaded", capsule[0]+".overrideShading")
+
+    cmds.parentConstraint("root", capsule[0], mo=True)
+    cmds.select(capsule[0], r=True)
+    capsuleDisplayLayer = cmds.createDisplayLayer(nr=True, name="Collision Capsule")
+    cmds.setAttr(capsuleDisplayLayer+".displayType", 2)
+    cmds.setAttr(capsuleDisplayLayer+".v", 0)
+    
+    for attr in ["tx", "ty", "tz", "rx", "ry", "rz", "sx", "sy", "sz"]:
+        cmds.setAttr(capsule[0]+"."+attr, lock=True, keyable=False, channelBox=False)
+
+    cmds.select(clear=True)
+addCapsule(42, 86)
+
+# ------------------------------------------------------------------
+# add ik bones Function that is called later after we open the export file.
+def addIKBones():
+    try:
+		# create the joints
+		cmds.select(clear = True)
+		ikFootRoot = cmds.joint(name = "ik_foot_root")
+		cmds.select(clear = True)
+		
+		cmds.select(clear = True)
+		ikFootLeft = cmds.joint(name = "ik_foot_l")
+		cmds.select(clear = True)
+		
+		cmds.select(clear = True)
+		ikFootRight = cmds.joint(name = "ik_foot_r")
+		cmds.select(clear = True)
+		
+		cmds.select(clear = True)
+		ikHandRoot = cmds.joint(name = "ik_hand_root")
+		cmds.select(clear = True)
+		
+		cmds.select(clear = True)
+		ikHandGun = cmds.joint(name = "ik_hand_gun")
+		cmds.select(clear = True)
+		
+		cmds.select(clear = True)
+		ikHandLeft = cmds.joint(name = "ik_hand_l")
+		cmds.select(clear = True)
+		
+		cmds.select(clear = True)
+		ikHandRight = cmds.joint(name = "ik_hand_r")
+		cmds.select(clear = True)
+		
+		# create hierarhcy
+		cmds.parent(ikFootRoot, "root")
+		cmds.parent(ikHandRoot, "root")
+		cmds.parent(ikFootLeft, ikFootRoot)
+		cmds.parent(ikFootRight, ikFootRoot)
+		cmds.parent(ikHandGun, ikHandRoot)
+		cmds.parent(ikHandLeft, ikHandGun)
+		cmds.parent(ikHandRight, ikHandGun)
+		
+		# constrain the joints
+		leftHandConstraint = cmds.parentConstraint("hand_l", ikHandLeft)[0]
+		rightHandGunConstraint = cmds.parentConstraint("hand_r", ikHandGun)[0]
+		rightHandConstraint = cmds.parentConstraint("hand_r", ikHandRight)[0]
+		leftFootConstraint = cmds.parentConstraint("foot_l", ikFootLeft)[0]
+		rightFootConstraint = cmds.parentConstraint("foot_r", ikFootRight)[0]
+	
+    except:
+		cmds.warning("Something went wrong")
+
+# ------------------------------------------------------------------
+# ADD the IK bones to the Export file.
+#save the AnimRig file before opening the export file.
+cmds.file(save = True, type = "mayaBinary", force = True)
+
+#ADD IK BONES
+#open export file
+sceneName = cmds.file(q = True, sceneName = True)
+exportFile = sceneName.partition("AnimRigs")[0] + "ExportFiles" +  sceneName.partition("AnimRigs")[2].partition(".mb")[0] + "_Export.mb"
+cmds.file(exportFile, open = True, force = True)
+
+#add ik bones
+addIKBones()
+
+# ------------------------------------------------------------------
+def addCustomAttrs():
+    attrs = cmds.listAttr("root")
+    exists = False
+    for one in attrs:
+        if one.find("DistanceCurve") != -1:
+            exists = True
+    if exists == False:
+        cmds.addAttr("root", ln="DistanceCurve", keyable=True)
+addCustomAttrs()
+
+#save the export file.
+cmds.file(save = True, type = "mayaBinary", force = True)
\ No newline at end of file
diff --git a/Engine/Extras/Maya_AnimationRiggingTools/ARTv1/General/ART/Projects/Orion_Cine/Scripts/postScripts/Gadget_Post_Script.py b/Engine/Extras/Maya_AnimationRiggingTools/ARTv1/General/ART/Projects/Orion_Cine/Scripts/postScripts/Gadget_Post_Script.py
new file mode 100644
--- /dev/null
+++ b/Engine/Extras/Maya_AnimationRiggingTools/ARTv1/General/ART/Projects/Orion_Cine/Scripts/postScripts/Gadget_Post_Script.py
@@ -0,0 +1,460 @@
+import maya.cmds as cmds
+import Tools.ART_CreateHydraulicRigs as hyd
+reload(hyd)
+import Tools.ART_addSpaceSwitching as space
+reload(space)
+
+# ------------------------------------------------------------------
+# scale all of the controls to be more visible with the mesh.
+def scaleControls():
+    sides = ["_l", "_r"]
+    for side in sides:
+        # feet
+        cmds.select(["ik_foot_anim"+side+".cv[0:17]"])
+        cmds.scale(1.0, 1.0, 1.0, r = True, ocp = False)
+        cmds.move(0, 0, -8.4, r = True)
+
+        cmds.select(["fk_helmet_straps"+side+"_01_anim.cv[0:7]"])
+        cmds.scale(0.204714, 0.204714, 0.204714, r = True, ocp = False)
+        cmds.move(0, 0, 0, r = True)
+
+        cmds.select(["fk_helmet_straps"+side+"_02_anim.cv[0:7]"])
+        cmds.scale(0.204714, 0.204714, 0.204714, r = True, ocp = False)
+        cmds.move(0, 0, 0, r = True)
+
+    # Other
+
+    cmds.select(["backpack_gas_01_anim.cv[0:32]"])
+    cmds.scale(0.391111, 0.391111, 0.391111, r = True, ocp = True)
+    cmds.move(-7.1, 0, 0, r = True, os=True, wd=True)
+
+    cmds.select(["backpack_gas_02_anim.cv[0:32]"])
+    cmds.scale(0.391111, 0.391111, 0.391111, r = True, ocp = True)
+    cmds.move(-7.1, 0, 0, r = True, os=True, wd=True)
+
+    cmds.select(["thigh_pouch_sml_anim.cv[0:32]"])
+    cmds.scale(0.391111, 0.391111, 0.391111, r = True, ocp = True)
+    cmds.move(0, -6.6, 0, r = True, os=True, wd=True)
+
+    cmds.select(["thigh_pouch_big_anim.cv[0:32]"])
+    cmds.scale(0.391111, 0.391111, 0.391111, r = True, ocp = True)
+    cmds.move(0, -6.6, 0, r = True, os=True, wd=True)
+
+    cmds.select(["backpack_base_anim.cv[0:32]"])
+    cmds.scale(0.306684, 1, 0.306684, r = True, ocp = True)
+    cmds.move(0, 24.6, 0, r = True, os=True, wd=True)
+
+    cmds.select(["shield_bott_anim.cv[0:32]"])
+    cmds.scale(0.301673, 0.301673, 0.301673, r = True, ocp = True)
+    cmds.move(7.937765, 12.982729, -6.212037, r = True, os=True, wd=True)
+
+    cmds.select(["shield_top_anim.cv[0:32]"])
+    cmds.scale(0.301673, 0.301673, 0.301673, r = True, ocp = True)
+    cmds.move(7.788894, 16.519037, 10.864125, r = True, os=True, wd=True)
+
+    cmds.select(["spring_anim.cv[0:32]"])
+    cmds.scale(1.0, 1.0, 0.001, r = True, ocp = True)
+
+    cmds.select(["canister_main_anim.cv[0:32]"])
+    cmds.scale(1.252256, 1.252256, 1.252256, r = True, ocp = True)
+
+    cmds.select(["robo_uparm_pistonrod_anim.cv[0:32]"])
+    cmds.scale(0.1, 0.447259, 0.447259, r = True, ocp = True)
+
+    cmds.select(["robo_uparm_pistonrod_01_anim.cv[0:32]"])
+    cmds.scale(0.1, 0.35, 0.35, r = True, ocp = True)
+
+    cmds.select(["robo_lowarm_pistonrod_01_anim.cv[0:32]"])
+    cmds.scale(0.1, 0.45, 0.45, r = True, ocp = True)
+
+    cmds.select(["robo_lowarm_pistonrod_02_anim.cv[0:32]"])
+    cmds.scale(0.1, 0.677192, 0.677192, r = True, ocp = True)
+
+    cmds.select(["fk_robo_arm_01_anim.cv[0:7]"])
+    cmds.scale(1.772257, 1.772257, 1.772257, r = True, ocp = False)
+
+    cmds.select(["fk_robo_arm_02_anim.cv[0:7]"])
+    cmds.scale(0.699603, 10.699603, 0.699603, r = True, ocp = False)
+
+    cmds.select(["fk_robo_arm_03_anim.cv[0:7]"])
+    cmds.scale(1.772257, 1.772257, 1.772257, r = True, ocp = False)
+
+    cmds.select(["fk_robo_arm_04_anim.cv[0:7]"])
+    cmds.scale(0.699603, 0.699603, 0.699603, r = True, ocp = False)
+
+    cmds.select(["fk_robo_arm_05_anim.cv[0:7]"])
+    cmds.scale(1.772257, 1.772257, 1.772257, r = True, ocp = False)
+
+    cmds.select(["fk_robo_arm_06_anim.cv[0:7]"])
+    cmds.scale(0.699603, 0.699603, 0.699603, r = True, ocp = False)
+
+    cmds.select(["fk_robo_arm_07_anim.cv[0:7]"])
+    cmds.scale(1.772257, 1.772257, 1.772257, r = True, ocp = False)
+
+    cmds.select(["fk_robo_thumb_01_anim.cv[0:7]"])
+    cmds.scale(1.0, 1.0, 1.0, r = True, ocp = True)
+
+    cmds.select(["fk_robo_thumb_02_anim.cv[0:7]"])
+    cmds.scale(1.0, 1.0, 1.0, r = True, ocp = True)
+
+    cmds.select(["fk_robo_pointer_01_anim.cv[0:7]"])
+    cmds.scale(1.0, 1.0, 1.0, r = True, ocp = True)
+
+    cmds.select(["fk_robo_pointer_02_anim.cv[0:7]"])
+    cmds.scale(1.0, 1.0, 1.0, r = True, ocp = True)
+
+    cmds.select(["fk_robo_pinky_01_anim.cv[0:7]"])
+    cmds.scale(1.0, 1.0, 1.0, r = True, ocp = True)
+
+    cmds.select(["fk_robo_pinky_02_anim.cv[0:7]"])
+    cmds.scale(1.0, 1.0, 1.0, r = True, ocp = True)
+
+    cmds.select(["arm_wire_01_anim.cv[0:32]"])
+    cmds.scale(0.5, 0.5, 0.5, r = True, ocp = False)
+
+    cmds.select(["arm_wire_02_anim.cv[0:32]"])
+    cmds.scale(0.5, 0.5, 0.5, r = True, ocp = False)
+
+    cmds.select(None)
+
+scaleControls()
+
+# ------------------------------------------------------------------
+def buildRoboArmRigFK():
+    import Tools.ART_CreateHydraulicRigs as hyd
+    reload(hyd)
+
+    # Delete all of the dynamics for the arms as we dont need them.
+    cmds.delete("rig_dyn_robo_arm_01_HairControls","rig_dyn_robo_pinky_01_HairControls","rig_dyn_robo_pointer_01_HairControls","rig_dyn_robo_thumb_01_HairControls")
+
+    # Lock the attrs on the piston controls for the robo arm so they cant accidentally be animated.
+    for one in ["fk_robo_arm_06_anim", "fk_robo_arm_04_anim", "fk_robo_arm_02_anim"]:
+        cmds.setAttr(one+".tx", lock=True)
+        cmds.setAttr(one+".tz", lock=True)
+        cmds.setAttr(one+".rx", lock=True)
+        cmds.setAttr(one+".rz", lock=True)
+
+    for one in ["fk_robo_arm_03_anim", "fk_robo_arm_05_anim"]:
+        cmds.setAttr(one+".tx", lock=True)
+        cmds.setAttr(one+".ty", lock=True)
+        cmds.setAttr(one+".tz", lock=True)
+        cmds.setAttr(one+".ry", lock=True)
+        cmds.setAttr(one+".rz", lock=True)
+
+    cmds.select("robo_lowarm_pistonrod_01_anim", "robo_lowarm_pistonrod_02_anim", "driver_robo_arm_05")
+    hyd.CreateHydraulicRig(-1)
+
+    cmds.select("robo_uparm_pistonrod_01_anim", "robo_uparm_pistonrod_anim", "driver_robo_arm_02")
+    hyd.CreateHydraulicRig(-1)
+
+    # Create a slide mechanism on the forearm part of the lowerarm piston.
+    cmds.addAttr("robo_lowarm_pistonrod_02_anim", ln = "Slide", dv = 0, min=0, max=10, keyable = True)
+    cmds.setDrivenKeyframe("robo_lowarm_pistonrod_02_anim_grp.translateX", cd="robo_lowarm_pistonrod_02_anim.Slide", itt="spline", ott="spline")
+    cmds.setDrivenKeyframe("robo_lowarm_pistonrod_02_anim_grp.translateZ", cd="robo_lowarm_pistonrod_02_anim.Slide", itt="spline", ott="spline")
+
+    cmds.setAttr("robo_lowarm_pistonrod_02_anim.Slide", 10)
+    cmds.setAttr("robo_lowarm_pistonrod_02_anim_grp.tx", -6.662443)
+    cmds.setAttr("robo_lowarm_pistonrod_02_anim_grp.tz", -5.510393)
+
+    cmds.setDrivenKeyframe("robo_lowarm_pistonrod_02_anim_grp.translateX", cd="robo_lowarm_pistonrod_02_anim.Slide", itt="spline", ott="spline")
+    cmds.setDrivenKeyframe("robo_lowarm_pistonrod_02_anim_grp.translateZ", cd="robo_lowarm_pistonrod_02_anim.Slide", itt="spline", ott="spline")
+    cmds.setAttr("robo_lowarm_pistonrod_02_anim.Slide", 0)
+    
+    cmds.setAttr("robo_lowarm_pistonrod_02_anim_grp.ry", -28)
+buildRoboArmRigFK()
+
+
+# ------------------------------------------------------------------
+def buildRoboArmRigIK():
+    import Modules.ART_rigUtils as utils
+    reload(utils)
+    import Tools.ART_addSpaceSwitching as space
+    reload(space)
+
+    # create an empty group for the Roboarm rig 
+    RoboArmGrp = cmds.group(em=True, name="RoboArm_grp")
+
+    # Create an IK arm control
+    ik_robowrist_anim = utils.createControl("sphere", 15, "ik_robowrist_anim")
+    constraint = cmds.parentConstraint("driver_robo_arm_07", ik_robowrist_anim)[0]
+    cmds.delete(constraint)
+    cmds.makeIdentity(ik_robowrist_anim, r = 1, s = 1, apply =True)
+    cmds.setAttr(ik_robowrist_anim + ".sx", 1.8)
+    cmds.setAttr(ik_robowrist_anim + ".sy", 1.8)
+    cmds.setAttr(ik_robowrist_anim + ".sz", 0.3)
+    cmds.makeIdentity(ik_robowrist_anim, r = 1, s = 1, apply =True)
+
+    ik_robowrist_anim_grp = cmds.group(empty=True, name=ik_robowrist_anim + "_grp")
+    constraint = cmds.pointConstraint("driver_robo_arm_07", ik_robowrist_anim_grp)[0]
+    cmds.delete(constraint)
+
+    ik_robowrist_animSpaceSwitchFollow = cmds.duplicate(ik_robowrist_anim_grp, po=True, n=ik_robowrist_anim+"_space_switcher_follow")[0]
+    ik_robowrist_animSpaceSwitch = cmds.duplicate(ik_robowrist_anim_grp, po=True, n=ik_robowrist_anim+"_space_switcher")[0]
+
+    cmds.parent(ik_robowrist_animSpaceSwitchFollow, RoboArmGrp)
+    cmds.parent(ik_robowrist_animSpaceSwitch, ik_robowrist_animSpaceSwitchFollow)
+    cmds.parent(ik_robowrist_anim_grp, ik_robowrist_animSpaceSwitch)
+
+    cmds.parent(ik_robowrist_anim, ik_robowrist_anim_grp)
+    cmds.makeIdentity(ik_robowrist_anim_grp, t=1, r=1, s=1, apply=True)
+
+    # create a pole vector control for the elbow
+    ik_roboelbow_anim = utils.createControl("sphere", 6, "ik_roboelbow_anim")
+    constraint = cmds.parentConstraint("driver_robo_arm_03", ik_roboelbow_anim)[0]
+    cmds.delete(constraint)
+    cmds.makeIdentity(ik_roboelbow_anim, t=1, r=1, s=1, apply=True)
+    cmds.setAttr(ik_roboelbow_anim + ".tz", -30)
+    cmds.makeIdentity(ik_roboelbow_anim, t=1, r=1, s=1, apply=True)	    
+    #cmds.parent(ik_roboelbowx_anim, RoboArmGrp)
+
+    ik_roboelbow_anim_grp = cmds.group(empty=True, name=ik_roboelbow_anim + "_grp")
+    constraint = cmds.pointConstraint(ik_roboelbow_anim, ik_roboelbow_anim_grp)[0]
+    cmds.delete(constraint)
+
+    ik_roboelbow_animSpaceSwitchFollow = cmds.duplicate(ik_roboelbow_anim_grp, po=True, n=ik_roboelbow_anim+"_space_switcher_follow")[0]
+    ik_roboelbow_animSpaceSwitch = cmds.duplicate(ik_roboelbow_anim_grp, po=True, n=ik_roboelbow_anim+"_space_switcher")[0]
+
+    cmds.parent(ik_roboelbow_animSpaceSwitchFollow, RoboArmGrp)
+    cmds.parent(ik_roboelbow_animSpaceSwitch, ik_roboelbow_animSpaceSwitchFollow)
+    cmds.parent(ik_roboelbow_anim_grp, ik_roboelbow_animSpaceSwitch)
+
+    cmds.parent(ik_roboelbow_anim, ik_roboelbow_anim_grp)
+    cmds.makeIdentity(ik_roboelbow_anim_grp, t=1, r=1, s=1, apply=True)
+
+    # Duplicate the driver Robo Arm to make a new IK robo arm
+    roboarm_01_joint = cmds.duplicate("driver_robo_arm_01", name="rig_robo_arm_01_ik", po=True)[0]
+    cmds.parent(roboarm_01_joint, RoboArmGrp)
+    roboarm_03_joint = cmds.duplicate("driver_robo_arm_03", name="rig_robo_arm_03_ik", po=True)[0]
+    cmds.parent(roboarm_03_joint, roboarm_01_joint)
+    roboarm_05_joint = cmds.duplicate("driver_robo_arm_05", name="rig_robo_arm_05_ik", po=True)[0]
+    cmds.parent(roboarm_05_joint, roboarm_03_joint)
+    roboarm_07_joint = cmds.duplicate("driver_robo_arm_07", name="rig_robo_arm_07_ik", po=True)[0]
+    cmds.parent(roboarm_07_joint, roboarm_05_joint)
+
+    # create an ikHandle for the arm and parent it under the IK arm control
+    roboarmIKHandle = cmds.ikHandle(name="roboarmIKHandle", sol="ikRPsolver", sj=roboarm_01_joint, ee=roboarm_07_joint)[0]
+    cmds.parent(roboarmIKHandle, ik_robowrist_anim)
+
+    # create a pole vector constraint for the ikhandle to the pv control
+    cmds.poleVectorConstraint(ik_roboelbow_anim, roboarmIKHandle)
+
+    # Orient Constrain the last joint in the Roboarm IK chain to the IK control for hand control.
+    cmds.orientConstraint(ik_robowrist_anim, roboarm_07_joint, mo=True)
+    
+    # Parent constrain the first joint of the IK chain to the driver skeleton.
+    cmds.parentConstraint("driver_backpack_base", roboarm_01_joint, mo=True)
+
+    # parent the roboarm group under the ctrl_rig node
+    cmds.parent(RoboArmGrp, "ctrl_rig")
+
+    # create new parent constraints with MO on for each of the original driver joints to be following the new ik joints
+    for one in ["01", "03", "05", "07"]:
+        cmds.parentConstraint("rig_robo_arm_"+one+"_ik", "driver_robo_arm_"+one, mo=True)
+        
+    Rig_SettingsConnection = cmds.listConnections("Rig_Settings.robo_arm_ik", p=True)
+    for one in Rig_SettingsConnection:
+        cmds.disconnectAttr("Rig_Settings.robo_arm_ik", one)
+
+    cmds.connectAttr("Rig_Settings.robo_arm_ik", "driver_robo_arm_01_parentConstraint1.rig_robo_arm_01_ikW3")
+    cmds.connectAttr("Rig_Settings.robo_arm_ik", "driver_robo_arm_03_parentConstraint1.rig_robo_arm_03_ikW3")
+    cmds.connectAttr("Rig_Settings.robo_arm_ik", "driver_robo_arm_05_parentConstraint1.rig_robo_arm_05_ikW3")
+    cmds.connectAttr("Rig_Settings.robo_arm_ik", "driver_robo_arm_07_parentConstraint1.rig_robo_arm_07_ikW3")
+    cmds.connectAttr("Rig_Settings.robo_arm_ik", "RoboArm_grp.visibility")
+
+    space.addSpaceSwitching(ik_robowrist_anim, "master_anim", False)
+    space.addSpaceSwitching(ik_robowrist_anim, "body_anim", False)
+    space.addSpaceSwitching(ik_robowrist_anim, "backpack_base_anim", False)
+
+    space.addSpaceSwitching(ik_roboelbow_anim, "master_anim", False)
+    space.addSpaceSwitching(ik_roboelbow_anim, "body_anim", False)
+    space.addSpaceSwitching(ik_roboelbow_anim, "backpack_base_anim", False)
+
+    cmds.setAttr("Rig_Settings.robo_arm_ik", 1)
+    cmds.setAttr("Rig_Settings.robo_arm_fk", 0)
+
+buildRoboArmRigIK()
+
+# ------------------------------------------------------------------
+def buildBackpackRig():
+    # Move the spring into position.
+    cmds.setAttr("spring_anim.ty", 0.058)
+    cmds.setAttr("spring_anim.tz", -28.7)
+    cmds.setAttr("spring_anim.sx", 0.749)
+    cmds.setAttr("spring_anim.sy", 0.749)
+    cmds.setAttr("spring_anim.sz", 0.081)
+    
+    # Move the canister into position
+    cmds.setAttr("canister_main_anim.ty", 1)
+    cmds.setAttr("canister_main_anim.tz", -49.4)
+    cmds.setAttr("canister_main_anim.sx", 0.585)
+    cmds.setAttr("canister_main_anim.sy", 0.585)
+    cmds.setAttr("canister_main_anim.sz", 0.314)
+    
+    # Rig the gas cans to be staying between the pelvis and the backpack base.
+    gascan_02_const = cmds.parentConstraint("driver_pelvis", "backpack_gas_02_parent_grp", mo=True)[0]
+    cmds.setAttr(gascan_02_const+".interpType", 2)
+    gascan_01_const = cmds.parentConstraint("driver_backpack_gas_02", "backpack_gas_01_parent_grp", mo=True)[0]
+    cmds.setAttr(gascan_01_const+".interpType", 2)
+
+buildBackpackRig()
+
+# ------------------------------------------------------------------
+def dynamicsOff():
+    for one in ["thigh_pouch_big_anim", "thigh_pouch_sml_anim", "helmet_straps_l_dyn_anim", "helmet_straps_r_dyn_anim"]:
+        cmds.setAttr(one+".chainStartEnvelope", 0)
+        cmds.setAttr(one+".chainStartEnvelope", l=True)
+dynamicsOff()
+
+# ------------------------------------------------------------------
+def helmetStraps():
+    const = cmds.parentConstraint("driver_C_jaw", "helmet_straps_l_ik_base_anim_grp", mo=True)[0]
+    cmds.setAttr(const+".interpType", 2)
+    const = cmds.parentConstraint("driver_C_jaw", "helmet_straps_r_ik_base_anim_grp", mo=True)[0]
+    cmds.setAttr(const+".interpType", 2)
+    const = cmds.parentConstraint("driver_C_jaw", "fk_helmet_straps_l_01_grp", mo=True)[0]
+    cmds.setAttr(const+".interpType", 2)
+    const = cmds.parentConstraint("driver_C_jaw", "fk_helmet_straps_r_01_grp", mo=True)[0]
+    cmds.setAttr(const+".interpType", 2)
+helmetStraps()
+
+
+##################################################################################
+##################################################################################
+
+def addCapsule(radius, halfHeight):
+    trueHalfHeight = (halfHeight*2)-(radius*2)
+    capsule = cmds.polyCylinder(name="collision_capsule", r=radius, h=trueHalfHeight, sx=18, sy=1, sz=5, ax=[0,0,1], rcp=True, cuv=3, ch=True)
+    cmds.setAttr(capsule[0]+".tz", halfHeight+2.2)
+    cmds.setAttr(capsule[0]+".rz", 90)
+
+    cmds.addAttr(capsule[0], ln="Radius", dv=radius, keyable=True)
+    cmds.addAttr(capsule[0], ln="HalfHeight", dv=halfHeight, keyable=True)
+
+    expr = capsule[1]+".radius = "+capsule[0]+".Radius;\n"+capsule[1]+".height = ("+capsule[0]+".HalfHeight*2) - ("+capsule[0]+".Radius*2);"
+    cmds.expression(s=expr, o=capsule[1], ae=1, uc="all")
+
+    cmds.setAttr(capsule[0]+".overrideEnabled", 1)
+    cmds.setAttr(capsule[0]+".overrideColor", 18)
+    cmds.setAttr(capsule[0]+".overrideShading", 0)
+
+    cmds.addAttr(capsule[0], ln="Shaded", dv=0, min=0, max=1, keyable=True)
+    cmds.connectAttr(capsule[0]+".Shaded", capsule[0]+".overrideShading")
+
+    cmds.parentConstraint("root", capsule[0], mo=True)
+    cmds.select(capsule[0], r=True)
+    capsuleDisplayLayer = cmds.createDisplayLayer(nr=True, name="Collision Capsule")
+    cmds.setAttr(capsuleDisplayLayer+".displayType", 2)
+    cmds.setAttr(capsuleDisplayLayer+".v", 0)
+    
+    for attr in ["tx", "ty", "tz", "rx", "ry", "rz", "sx", "sy", "sz"]:
+        cmds.setAttr(capsule[0]+"."+attr, lock=True, keyable=False, channelBox=False)
+
+    cmds.select(clear=True)
+addCapsule(42, 85)
+
+# ------------------------------------------------------------------
+def createGameCam():
+    gameCam = cmds.camera(centerOfInterest=5, focalLength=18, lensSqueezeRatio=1, cameraScale=1, horizontalFilmAperture=1.41732, horizontalFilmOffset=0, verticalFilmAperture=0.94488, verticalFilmOffset=0, filmFit="Overscan", overscan=1.05, motionBlur=0, shutterAngle=144, nearClipPlane=1.0, farClipPlane=10000, orthographic=0, orthographicWidth=30, panZoomEnabled=0, horizontalPan=0, verticalPan=0, zoom=1, displayFilmGate=True)[0]
+    print gameCam
+    gameCam = cmds.rename(gameCam, "GameCam")
+    cmds.camera(gameCam, e=True, horizontalFilmAperture=1.6724376, focalLength=21.2399857011)
+    cmds.setAttr(gameCam+".displayGateMaskOpacity", 1)
+    cmds.setAttr(gameCam+".displayGateMaskColor", 0, 0, 0, type="double3")
+    cmds.setAttr(gameCam+".r", 90, 0, 180, type="double3")
+    
+    # Create the control
+    control = cmds.circle(name=gameCam+"_Ctrl", ch=True, o=True, r=15)[0]
+    cmds.makeIdentity(control, a=True, t=True, r=True, s=True)
+    # Create groups for the sdk and the constraints
+    const_grp = cmds.group(control, name=control+"_grp")
+    if cmds.objExists("rig_grp"):
+        cmds.parent(const_grp, "rig_grp")
+    # Parent Constrain the cam to the control
+    cmds.parentConstraint(control, gameCam, mo=True)
+    
+    cmds.connectAttr(control+".visibility", gameCam+".visibility", f=True)
+    cmds.setAttr(const_grp+".v", 0)
+    
+    cmds.setAttr(const_grp+".t", 0, 500, 200, type="double3")
+createGameCam()
+# ------------------------------------------------------------------
+
+# ------------------------------------------------------------------
+# add ik bones Function that is called later after we open the export file.
+def addIKBones():
+    try:
+		# create the joints
+		cmds.select(clear = True)
+		ikFootRoot = cmds.joint(name = "ik_foot_root")
+		cmds.select(clear = True)
+		
+		cmds.select(clear = True)
+		ikFootLeft = cmds.joint(name = "ik_foot_l")
+		cmds.select(clear = True)
+		
+		cmds.select(clear = True)
+		ikFootRight = cmds.joint(name = "ik_foot_r")
+		cmds.select(clear = True)
+		
+		cmds.select(clear = True)
+		ikHandRoot = cmds.joint(name = "ik_hand_root")
+		cmds.select(clear = True)
+		
+		cmds.select(clear = True)
+		ikHandGun = cmds.joint(name = "ik_hand_gun")
+		cmds.select(clear = True)
+		
+		cmds.select(clear = True)
+		ikHandLeft = cmds.joint(name = "ik_hand_l")
+		cmds.select(clear = True)
+		
+		cmds.select(clear = True)
+		ikHandRight = cmds.joint(name = "ik_hand_r")
+		cmds.select(clear = True)
+		
+		# create hierarhcy
+		cmds.parent(ikFootRoot, "root")
+		cmds.parent(ikHandRoot, "root")
+		cmds.parent(ikFootLeft, ikFootRoot)
+		cmds.parent(ikFootRight, ikFootRoot)
+		cmds.parent(ikHandGun, ikHandRoot)
+		cmds.parent(ikHandLeft, ikHandGun)
+		cmds.parent(ikHandRight, ikHandGun)
+		
+		# constrain the joints
+		leftHandConstraint = cmds.parentConstraint("hand_l", ikHandLeft)[0]
+		rightHandGunConstraint = cmds.parentConstraint("hand_r", ikHandGun)[0]
+		rightHandConstraint = cmds.parentConstraint("hand_r", ikHandRight)[0]
+		leftFootConstraint = cmds.parentConstraint("foot_l", ikFootLeft)[0]
+		rightFootConstraint = cmds.parentConstraint("foot_r", ikFootRight)[0]
+	
+    except:
+		cmds.warning("Something went wrong")
+
+# ------------------------------------------------------------------
+# ADD the IK bones to the Export file.
+#save the AnimRig file before opening the export file.
+cmds.file(save = True, type = "mayaBinary", force = True)
+
+#ADD IK BONES
+#open export file
+sceneName = cmds.file(q = True, sceneName = True)
+exportFile = sceneName.partition("AnimRigs")[0] + "ExportFiles" +  sceneName.partition("AnimRigs")[2].partition(".mb")[0] + "_Export.mb"
+cmds.file(exportFile, open = True, force = True)
+
+#add ik bones
+addIKBones()
+
+# ------------------------------------------------------------------
+def addCustomAttrs():
+    attrs = cmds.listAttr("root")
+    exists = False
+    for one in attrs:
+        if one.find("DistanceCurve") != -1:
+            exists = True
+    if exists == False:
+        cmds.addAttr("root", ln="DistanceCurve", keyable=True)
+addCustomAttrs()
+
+#save the export file.
+cmds.file(save = True, type = "mayaBinary", force = True)
\ No newline at end of file
diff --git a/Engine/Extras/Maya_AnimationRiggingTools/ARTv1/General/ART/Projects/Orion_Cine/Scripts/postScripts/Price_Post_Script.py b/Engine/Extras/Maya_AnimationRiggingTools/ARTv1/General/ART/Projects/Orion_Cine/Scripts/postScripts/Price_Post_Script.py
new file mode 100644
--- /dev/null
+++ b/Engine/Extras/Maya_AnimationRiggingTools/ARTv1/General/ART/Projects/Orion_Cine/Scripts/postScripts/Price_Post_Script.py
@@ -0,0 +1,515 @@
+import maya.cmds as cmds
+import Tools.ART_addSpaceSwitching as space
+reload(space)
+
+#NOTE This rig has custom changes to be made during the RigPose placement phase.  See the commented out script at the bottom of this file.
+
+
+# ------------------------------------------------------------------
+# scale all of the controls to be more visible with the mesh.
+def scaleControls():
+    objects = ["bot_uparm_pistonC_top_anim", "bot_uparm_pistonC_bot_anim", "bot_loarm_pistonB_top_anim", "bot_loarm_pistonB_bot_anim", "bot_loarm_pistonC_top_anim", "bot_loarm_pistonC_bot_anim", "bot_loarm_pistonD_top_anim", "bot_loarm_pistonD_bot_anim", "bot_loarm_anim", "bot_loarm_twist_anim", "bot_hand_anim", "gun_scope_anim", "gun_stock_base_anim", "gun_stock_bott_anim", "gun_barrel_inner_anim", "gun_body_panel_l_anim", "gun_body_panel_r_anim", "gun_stock_base_in_bott_anim", "gun_stock_base_out_bott_anim", "gun_stock_base_in_top_anim", "gun_stock_base_out_top_anim", "gun_stock_in_bott_anim", "gun_stock_in_top_anim", "trap_handle_anim", "sniper_can_arm_01_anim"]
+    for one in objects:
+        cmds.select([one+".cv[0:32]"])
+        cmds.scale(0.227057, 0.227057, 0.227057, r = True, ocp = True)
+        cmds.move(0, -5, 0, r = True, os=True, wd=True)
+        
+    objects = ["bot_uparm_anim", "gun_anim", "gun_stock_panel_r_anim", "gun_stock_panel_l_anim", "gun_barrel_anim", "shotgun_base_anim", "shotgun_clip_anim", "shotgun_anim", "ankle_plate_in_l_anim", "ankle_plate_in_r_anim",  "ankle_plate_out_l_anim", "ankle_plate_out_r_anim", "sniper_can_arm_02_anim", "sniper_can_arm_03_anim", "arm_hose_bott_anim", "arm_hose_top_anim", "mask_anim", "helmet_sight_anim", "arm_pouch_l_anim", "pack_hose_out_anim", "pack_hose_in_anim", "heel_plate_l_anim", "heel_plate_r_anim", "wrist_cover_anim"]
+    for one in objects:
+        cmds.select([one+".cv[0:32]"])
+        cmds.scale(0.2, 0.2, 0.2, r = True, ocp = True)
+        
+    for side in ["_l", "_r"]:
+        #ik feet
+        cmds.select("ik_foot_anim"+side + ".cv[0:32]")
+        cmds.scale(1, 1, 1, r=True)
+        cmds.move(0, 0, -4.695816, r = True, os=True, wd=True)
+            
+    cmds.select(["gun_anim.cv[0:32]"])
+    cmds.move(8, 0, 0, r = True, os=True, wd=True)
+
+    cmds.select(["gun_scope_anim.cv[0:32]"])
+    cmds.move(0, 10, 0, r = True, os=True, wd=True)
+
+    cmds.select(["gun_stock_panel_r_anim.cv[0:32]"])
+    cmds.move(2.5, 0, 0, r = True, os=True, wd=True)
+
+    cmds.select(["gun_stock_panel_l_anim.cv[0:32]"])
+    cmds.move(-2.5, 0, 0, r = True, os=True, wd=True)
+
+    cmds.select(["gun_stock_bott_anim.cv[0:32]"])
+    cmds.move(0, -1, 0, r = True, os=True, wd=True)
+
+    cmds.select(["gun_barrel_anim.cv[0:32]"])
+    cmds.move(0, 0, -5, r = True, os=True, wd=True)
+
+    cmds.select(["gun_stock_base_out_top_anim.cv[0:32]"])
+    cmds.move(0, 10, 0, r = True, os=True, wd=True)
+
+    cmds.select(["gun_stock_base_in_top_anim.cv[0:32]"])
+    cmds.move(0, 10, 0, r = True, os=True, wd=True)
+
+    cmds.select(["gun_stock_in_top_anim.cv[0:32]"])
+    cmds.move(0, 15, 0, r = True, os=True, wd=True)
+
+    cmds.select(["bot_hand_anim.cv[0:32]"])
+    cmds.move(0, 0, 6, r = True, os=True, wd=True)
+
+    cmds.select(["bot_loarm_anim.cv[0:32]"])
+    cmds.move(0, 10, 0, r = True, os=True, wd=True)
+
+    cmds.select(["bot_uparm_anim.cv[0:32]"])
+    cmds.move(0, 5, -8, r = True, os=True, wd=True)
+
+    cmds.select(["shotgun_base_anim.cv[0:32]"])
+    cmds.move(7, 0, 0, r = True, os=True, wd=True)
+
+    cmds.select(["shotgun_anim.cv[0:32]"])
+    cmds.move(0, 0, 3, r = True, os=True, wd=True)
+
+    cmds.select(["shotgun_clip_anim.cv[0:32]"])
+    cmds.move(-5, 0, 0, r = True, os=True, wd=True)
+
+    cmds.select(["sniper_can_arm_01_anim.cv[0:32]"])
+    cmds.move(0, 7, 9, r = True, os=True, wd=True)
+
+    cmds.select(["sniper_can_arm_02_anim.cv[0:32]"])
+    cmds.move(0, 5, 5, r = True, os=True, wd=True)
+
+    cmds.select(["sniper_can_arm_03_anim.cv[0:32]"])
+    cmds.scale(1.0, 6.3, 6.3, r = True, ocp = True)
+    cmds.move(0, 5, 0, r = True, os=True, wd=True)
+
+    cmds.select(None)
+scaleControls()
+
+# ------------------------------------------------------------------
+def createPistons():
+    import Tools.ART_CreateHydraulicRigs as hyd
+    reload(hyd)
+    
+    # pistons
+    cmds.select("bot_uparm_pistonC_top_anim", "bot_uparm_pistonC_bot_anim", "driver_upperarm_r")
+    hyd.CreateHydraulicRig(-1)
+    cmds.select("bot_loarm_pistonB_top_anim", "bot_loarm_pistonB_bot_anim", "driver_lowerarm_r")
+    hyd.CreateHydraulicRig(-1)
+    cmds.select("bot_loarm_pistonC_top_anim", "bot_loarm_pistonC_bot_anim", "driver_lowerarm_r")
+    hyd.CreateHydraulicRig(-1)
+    cmds.select("bot_loarm_pistonD_top_anim", "bot_loarm_pistonD_bot_anim", "driver_lowerarm_r")
+    hyd.CreateHydraulicRig(-1)
+
+    # manually changing the up object for a few of the constraints that were created above so that they rotate at teh correct angles.
+    cmds.aimConstraint("bot_loarm_pistonB_bot_anim_sub_aimConstraint1", e=True, mo=True, wuo="driver_bot_hand")
+    cmds.aimConstraint("bot_loarm_pistonC_bot_anim_sub_aimConstraint1", e=True, mo=True, u=(0,-1,0), wuo="bot_loarm_pistonD_bot_anim")
+    cmds.aimConstraint("bot_loarm_pistonD_bot_anim_sub_aimConstraint1", e=True, mo=True, u=(0,-1,0), wuo="bot_loarm_pistonC_bot_anim")
+    cmds.aimConstraint("bot_uparm_pistonC_bot_anim_sub_aimConstraint1", e=True, mo=True, wuo="driver_lowerarm_r")
+
+    # bot arm twist
+    cmds.orientConstraint("driver_lowerarm_twist_01_r", "bot_loarm_twist_anim_space_switcher", mo = True)
+createPistons()
+
+# ------------------------------------------------------------------
+def anklePlatesRig():
+    for side in ["_l", "_r"]:
+        sidegrp = cmds.group(em=True, n="ankle_plates_grp"+side)
+        const = cmds.parentConstraint("driver_calf_twist_01"+side, sidegrp)[0]
+        cmds.delete(const)
+        cmds.makeIdentity(sidegrp, a=True, t=True, r=True, s=True)
+        cmds.parent(sidegrp, "ctrl_rig")
+        cmds.parentConstraint("driver_calf_twist_01"+side, sidegrp, mo=True)
+        
+        for dir in ["_in", "_out"]:
+            target = cmds.spaceLocator(n="ankle_plate"+dir+"_target"+side)[0]
+            const = cmds.parentConstraint("ankle_plate"+dir+side+"_anim", target)[0]
+            cmds.delete(const)
+            if side == "_l":
+                cmds.move(-10, 0, 0, target, r=True, os=True, wd=True)
+            else:
+                cmds.move(10, 0, 0, target, r=True, os=True, wd=True)
+            
+            up = cmds.spaceLocator(n="ankle_plate"+dir+"_up"+side)[0]
+            const = cmds.parentConstraint("ankle_plate"+dir+side+"_anim", up)[0]
+            cmds.delete(const)
+            if side == "_l":
+                cmds.move(0, 0, -5, up, r=True, os=True, wd=True)
+            else:
+                cmds.move(0, 0, 5, up, r=True, os=True, wd=True)
+
+            
+            grp = cmds.group(em=True, n="ankle_plate"+dir+"_grp"+side)
+            const = cmds.parentConstraint("driver_foot"+side, grp)[0]
+            cmds.delete(const)
+            cmds.makeIdentity(grp, a=True, t=True, r=True, s=True)
+            
+            cmds.parent(target, up, grp)
+            cmds.parent(grp, sidegrp)
+            
+            driver_foot_loc = cmds.xform("driver_foot"+side, q=True, ws=True, rp=True)
+            cmds.xform("ankle_plate"+dir+side+"_anim_grp", ws=True, rp=driver_foot_loc)
+
+            cmds.parentConstraint("driver_foot"+side, grp, mo=True, sr=["x"])
+          
+            cmds.aimConstraint(target, "ankle_plate"+dir+side+"_anim_grp", mo=True, aimVector=(-1, 0, 0), upVector=(0, 0, -1), worldUpType="object", worldUpVector=(0, 0, -1), worldUpObject=str(up))
+            
+  
+            cmds.parentConstraint("driver_foot"+side, "ankle_plate"+dir+side+"_anim_grp", mo=True, sr=["x", "y", "z"])
+            
+            cmds.setAttr("ankle_plate"+dir+"_grp"+side+".rotateOrder", 4)
+            
+            if side == "_r":
+                cmds.setAttr("ankle_plate"+dir+"_grp"+side+"_parentConstraint1.target[0].targetOffsetRotateX", -180)
+                cmds.setAttr("ankle_plate"+dir+"_grp"+side+"_parentConstraint1.target[0].targetOffsetRotateY", -90)
+            
+            cmds.setAttr(grp+".v", 0)
+            
+anklePlatesRig()
+
+# ------------------------------------------------------------------
+def heelRig():
+    for side in ["_l", "_r"]:
+        cmds.setDrivenKeyframe(["heel_plate"+side+"_anim_grp.rx"], cd = "driver_foot"+side+".rz", itt = "spline", ott = "spline")
+        cmds.setAttr("ik_foot_anim"+side+".rx", 40)
+        cmds.setAttr("heel_plate"+side+"_anim_grp.rx", -17.471)
+        cmds.setDrivenKeyframe(["heel_plate"+side+"_anim_grp.rx"], cd = "driver_foot"+side+".rz", itt = "spline", ott = "spline")
+        cmds.setAttr("ik_foot_anim"+side+".rx", 80)
+        cmds.setAttr("heel_plate"+side+"_anim_grp.rx", -59.01)
+        cmds.setDrivenKeyframe(["heel_plate"+side+"_anim_grp.rx"], cd = "driver_foot"+side+".rz", itt = "spline", ott = "spline")
+        cmds.setAttr("ik_foot_anim"+side+".rx", 0)
+heelRig()
+
+# ------------------------------------------------------------------
+def trapHandleRig():
+    cmds.parentConstraint("driver_spine_02", "driver_spine_03", "trap_handle_parent_grp", mo=True)
+    cmds.setAttr("trap_handle_parent_grp_parentConstraint1.driver_spine_01W0", 0.12)
+    cmds.setAttr("trap_handle_parent_grp_parentConstraint1.driver_spine_02W1", 0.18)
+    cmds.setAttr("trap_handle_parent_grp_parentConstraint1.driver_spine_03W2", 0.7)
+
+trapHandleRig()
+
+# ------------------------------------------------------------------
+def wristPlatesRig():
+    cmds.rotate(-10.523881, 0, 0, "wrist_cover_anim_grp")
+    
+    target = cmds.spaceLocator(n="wrist_cover_target")[0]
+    const = cmds.parentConstraint("driver_bot_loarm_pistonB_bot", target)[0]
+    cmds.delete(const)
+    
+    up = cmds.spaceLocator(n="wrist_cover_up")[0]
+    const = cmds.parentConstraint("driver_wrist_cover", up)[0]
+    cmds.delete(const)
+    cmds.move(0, 5, 0, up, r=True, os=True, wd=True)
+    
+    grp = cmds.group(em=True, n="wrist_cover_grp")
+    const = cmds.parentConstraint("driver_hand_r", grp)[0]
+    cmds.delete(const)
+    cmds.makeIdentity(grp, a=True, t=True, r=True, s=True)
+    
+    cmds.parent(target, up, grp)
+    cmds.parent(grp, "ctrl_rig")
+    
+    prntConst = cmds.parentConstraint("driver_bot_loarm_pistonB_bot", "driver_hand_r", grp, mo=True)[0]
+    #cmds.pointConstraint("driver_bot_loarm_pistonB_bot", "driver_hand_r", grp, mo=True)
+    #ornConst = cmds.orientConstraint("driver_bot_loarm_pistonB_bot", "driver_hand_r", grp, mo=True)[0]
+    cmds.setAttr(prntConst+".interpType", 2)
+
+    cmds.aimConstraint(target, "wrist_cover_anim_grp", mo=True, aimVector=(-1, 0, 0), upVector=(0, 1, 0), worldUpType="object", worldUpVector=(0, 1, 0), worldUpObject=str(up))
+    cmds.parentConstraint("driver_bot_loarm_pistonB_bot", "driver_hand_r", "wrist_cover_anim_grp", mo=True, sr=["x", "y", "z"])
+    #cmds.pointConstraint("driver_bot_loarm_pistonB_bot", "driver_hand_r", "wrist_cover_anim_grp", mo=True)
+    
+    cmds.setAttr(grp+".v", 0)
+wristPlatesRig()
+
+# ------------------------------------------------------------------
+def shotgunRig():
+    import Tools.ART_addSpaceSwitching as space
+    reload(space)
+
+    cmds.addAttr("shotgun_anim", ln = "shotgun_rotate", dv = 0, min = 0, max = 10, keyable = True)
+    cmds.addAttr("shotgun_anim", ln = "shotgun_rail", dv = 0, min = 0, max = 10, keyable = True)
+
+    shotgun_rail_grp = cmds.group("shotgun_anim", n="shotgun_anim_rail_grp")
+    cmds.xform(shotgun_rail_grp, ws=True, piv=[16.745853637042327, 26.165296554565593, 153.62670823479255])
+
+    cmds.setDrivenKeyframe(["shotgun_anim_rail_grp.rz"], cd = "shotgun_anim.shotgun_rotate", itt = "auto", ott = "auto")
+    cmds.setAttr("shotgun_anim.shotgun_rotate", 10)
+    cmds.setAttr("shotgun_anim_rail_grp.rz", -90)
+    cmds.setDrivenKeyframe(["shotgun_anim_rail_grp.rz"], cd = "shotgun_anim.shotgun_rotate", itt = "auto", ott = "auto")
+    cmds.setAttr("shotgun_anim.shotgun_rotate", 0)
+
+    cmds.setDrivenKeyframe(["shotgun_base_anim_grp.tx", "shotgun_base_anim_grp.ty", "shotgun_base_anim_grp.tz", "shotgun_base_anim_grp.rx"], cd = "shotgun_anim.shotgun_rail", itt = "auto", ott = "auto")
+    cmds.setAttr("shotgun_anim.shotgun_rail", 5)
+    cmds.setAttr("shotgun_base_anim_grp.tx", 3.14)
+    cmds.setAttr("shotgun_base_anim_grp.ty", -3.774)
+    cmds.setAttr("shotgun_base_anim_grp.tz", 9.402)
+    cmds.setAttr("shotgun_base_anim_grp.rx", 13.9)
+    cmds.setDrivenKeyframe(["shotgun_base_anim_grp.tx", "shotgun_base_anim_grp.ty", "shotgun_base_anim_grp.tz", "shotgun_base_anim_grp.rx"], cd = "shotgun_anim.shotgun_rail", itt = "auto", ott = "auto")
+    cmds.setAttr("shotgun_anim.shotgun_rail", 10)
+    cmds.setAttr("shotgun_base_anim_grp.tx", 4.999)
+    cmds.setAttr("shotgun_base_anim_grp.ty", -8.562)
+    cmds.setAttr("shotgun_base_anim_grp.tz", 12.609)
+    cmds.setAttr("shotgun_base_anim_grp.rx", 72.809)
+    cmds.setDrivenKeyframe(["shotgun_base_anim_grp.tx", "shotgun_base_anim_grp.ty", "shotgun_base_anim_grp.tz", "shotgun_base_anim_grp.rx"], cd = "shotgun_anim.shotgun_rail", itt = "auto", ott = "auto")
+    cmds.setAttr("shotgun_anim.shotgun_rail", 0)
+
+    space.addSpaceSwitching("shotgun_anim", "driver_hand_l", False)
+shotgunRig()
+
+# ------------------------------------------------------------------
+def shoulderPadRig():
+    cmds.setAttr("Rig_Settings.rArmMode", 0)
+
+    shoulderPad_grp = cmds.group("bot_uparm_anim_grp", n="bot_uparm_anim_pivot")
+    cmds.xform(shoulderPad_grp, ws=True, piv=[16.745853637042327, 26.165296554565593, 153.62670823479255])
+
+    driver_upperarm_r_loc = cmds.xform("driver_upperarm_r", q=True, ws=True, rp=True)
+    print driver_upperarm_r_loc
+    cmds.xform(shoulderPad_grp, ws=True, rp=driver_upperarm_r_loc)
+
+    cmds.setAttr("fk_arm_r_anim.ry", 0)
+    cmds.setAttr("bot_uparm_anim_pivot.ry", 19.153)
+    cmds.setAttr("bot_uparm_anim_grp.tx", 2.209)
+    cmds.setAttr("bot_uparm_anim_grp.tz", -2.583)
+    cmds.setAttr("bot_uparm_anim_grp.ry", -38.873)
+    cmds.setDrivenKeyframe(["bot_uparm_anim_pivot.ry", "bot_uparm_anim_grp.tx", "bot_uparm_anim_grp.tz", "bot_uparm_anim_grp.ry"], cd = "driver_upperarm_r.ry", itt = "spline", ott = "spline")
+    cmds.setAttr("fk_arm_r_anim.ry", 45)
+    cmds.setAttr("bot_uparm_anim_pivot.ry", 0)
+    cmds.setAttr("bot_uparm_anim_grp.tx", 0)
+    cmds.setAttr("bot_uparm_anim_grp.tz", 0)
+    cmds.setAttr("bot_uparm_anim_grp.ry", 0)
+    cmds.setDrivenKeyframe(["bot_uparm_anim_pivot.ry", "bot_uparm_anim_grp.tx", "bot_uparm_anim_grp.tz", "bot_uparm_anim_grp.ry"], cd = "driver_upperarm_r.ry", itt = "spline", ott = "spline")
+    cmds.setAttr("fk_arm_r_anim.ry", -20)
+    cmds.setAttr("bot_uparm_anim_pivot.ry", 27.813)
+    cmds.setAttr("bot_uparm_anim_grp.tx", 6.98)
+    cmds.setAttr("bot_uparm_anim_grp.tz", -1.586)
+    cmds.setAttr("bot_uparm_anim_grp.ry", -58.887)
+    cmds.setDrivenKeyframe(["bot_uparm_anim_pivot.ry", "bot_uparm_anim_grp.tx", "bot_uparm_anim_grp.tz", "bot_uparm_anim_grp.ry"], cd = "driver_upperarm_r.ry", itt = "spline", ott = "spline")
+    cmds.setAttr("fk_arm_r_anim.ry", 0)
+    
+    cmds.setAttr("Rig_Settings.rArmMode", 1)
+
+shoulderPadRig()
+
+# ------------------------------------------------------------------
+def sniperCanisterRig():
+    import Tools.ART_addSpaceSwitching as space
+    reload(space)
+
+    space.addSpaceSwitching("sniper_can_arm_03_anim", "driver_hand_l", False)
+    space.addSpaceSwitching("sniper_can_arm_03_anim", "driver_gun", False)
+sniperCanisterRig()
+
+# ------------------------------------------------------------------
+def trapRig():
+    import Tools.ART_addSpaceSwitching as space
+    reload(space)
+
+    space.addSpaceSwitching("trap_handle_anim", "driver_hand_l", False)
+trapRig()
+
+# ------------------------------------------------------------------
+def leftHandRig():
+    import Tools.ART_addSpaceSwitching as space
+    reload(space)
+
+    space.addSpaceSwitching("ik_wrist_l_anim", "driver_gun_stock_bott", False)
+    
+    cmds.setAttr("Rig_Settings.lForearmTwistAmount", 0.9)
+    cmds.setAttr("Rig_Settings.lForearmTwist2Amount", 0.6)
+    
+leftHandRig()
+
+# ------------------------------------------------------------------
+def maskRig():
+    import Tools.ART_addSpaceSwitching as space
+    reload(space)
+
+    space.addSpaceSwitching("mask_anim", "driver_hand_l", False)
+    space.addSpaceSwitching("mask_anim", "driver_hand_r", False)
+maskRig()
+
+# ------------------------------------------------------------------
+def dynamicsOff():
+    for one in ["hip_pack_a_anim", "hip_pack_b_anim", "hip_pack_c_anim", "hip_pack_d_anim"]:
+        cmds.setAttr(one+".chainStartEnvelope", 0)
+        cmds.setAttr(one+".chainStartEnvelope", l=True)
+dynamicsOff()
+
+# ------------------------------------------------------------------
+def createGameCam():
+    gameCam = cmds.camera(centerOfInterest=5, focalLength=18, lensSqueezeRatio=1, cameraScale=1, horizontalFilmAperture=1.41732, horizontalFilmOffset=0, verticalFilmAperture=0.94488, verticalFilmOffset=0, filmFit="Overscan", overscan=1.05, motionBlur=0, shutterAngle=144, nearClipPlane=1.0, farClipPlane=10000, orthographic=0, orthographicWidth=30, panZoomEnabled=0, horizontalPan=0, verticalPan=0, zoom=1, displayFilmGate=True)[0]
+    print gameCam
+    gameCam = cmds.rename(gameCam, "GameCam")
+    cmds.camera(gameCam, e=True, horizontalFilmAperture=1.6724376, focalLength=21.2399857011)
+    cmds.setAttr(gameCam+".displayGateMaskOpacity", 1)
+    cmds.setAttr(gameCam+".displayGateMaskColor", 0, 0, 0, type="double3")
+    cmds.setAttr(gameCam+".r", 90, 0, 180, type="double3")
+    
+    # Create the control
+    control = cmds.circle(name=gameCam+"_Ctrl", ch=True, o=True, r=15)[0]
+    cmds.makeIdentity(control, a=True, t=True, r=True, s=True)
+    # Create groups for the sdk and the constraints
+    const_grp = cmds.group(control, name=control+"_grp")
+    if cmds.objExists("rig_grp"):
+        cmds.parent(const_grp, "rig_grp")
+    # Parent Constrain the cam to the control
+    cmds.parentConstraint(control, gameCam, mo=True)
+    
+    cmds.connectAttr(control+".visibility", gameCam+".visibility", f=True)
+    cmds.setAttr(const_grp+".v", 0)
+    
+    cmds.setAttr(const_grp+".t", 0, 500, 200, type="double3")
+createGameCam()
+
+# ------------------------------------------------------------------
+def addCapsule(radius, halfHeight):
+    trueHalfHeight = (halfHeight*2)-(radius*2)
+    capsule = cmds.polyCylinder(name="collision_capsule", r=radius, h=trueHalfHeight, sx=18, sy=1, sz=5, ax=[0,0,1], rcp=True, cuv=3, ch=True)
+    cmds.setAttr(capsule[0]+".tz", halfHeight+2.2)
+    cmds.setAttr(capsule[0]+".rz", 90)
+
+    cmds.addAttr(capsule[0], ln="Radius", dv=radius, keyable=True)
+    cmds.addAttr(capsule[0], ln="HalfHeight", dv=halfHeight, keyable=True)
+
+    expr = capsule[1]+".radius = "+capsule[0]+".Radius;\n"+capsule[1]+".height = ("+capsule[0]+".HalfHeight*2) - ("+capsule[0]+".Radius*2);"
+    cmds.expression(s=expr, o=capsule[1], ae=1, uc="all")
+
+    cmds.setAttr(capsule[0]+".overrideEnabled", 1)
+    cmds.setAttr(capsule[0]+".overrideColor", 18)
+    cmds.setAttr(capsule[0]+".overrideShading", 0)
+
+    cmds.addAttr(capsule[0], ln="Shaded", dv=0, min=0, max=1, keyable=True)
+    cmds.connectAttr(capsule[0]+".Shaded", capsule[0]+".overrideShading")
+
+    cmds.parentConstraint("root", capsule[0], mo=True)
+    cmds.select(capsule[0], r=True)
+    capsuleDisplayLayer = cmds.createDisplayLayer(nr=True, name="Collision Capsule")
+    cmds.setAttr(capsuleDisplayLayer+".displayType", 2)
+    cmds.setAttr(capsuleDisplayLayer+".v", 0)
+    
+    for attr in ["tx", "ty", "tz", "rx", "ry", "rz", "sx", "sy", "sz"]:
+        cmds.setAttr(capsule[0]+"."+attr, lock=True, keyable=False, channelBox=False)
+
+    cmds.select(clear=True)
+addCapsule(42, 86)
+
+# ------------------------------------------------------------------
+# ------------------------------------------------------------------
+# add ik bones Function that is called later after we open the export file.
+def addIKBones():
+    try:
+		# create the joints
+		cmds.select(clear = True)
+		ikFootRoot = cmds.joint(name = "ik_foot_root")
+		cmds.select(clear = True)
+		
+		cmds.select(clear = True)
+		ikFootLeft = cmds.joint(name = "ik_foot_l")
+		cmds.select(clear = True)
+		
+		cmds.select(clear = True)
+		ikFootRight = cmds.joint(name = "ik_foot_r")
+		cmds.select(clear = True)
+		
+		cmds.select(clear = True)
+		ikHandRoot = cmds.joint(name = "ik_hand_root")
+		cmds.select(clear = True)
+		
+		cmds.select(clear = True)
+		ikHandGun = cmds.joint(name = "ik_hand_gun")
+		cmds.select(clear = True)
+		
+		cmds.select(clear = True)
+		ikHandLeft = cmds.joint(name = "ik_hand_l")
+		cmds.select(clear = True)
+		
+		cmds.select(clear = True)
+		ikHandRight = cmds.joint(name = "ik_hand_r")
+		cmds.select(clear = True)
+		
+		# create hierarhcy
+		cmds.parent(ikFootRoot, "root")
+		cmds.parent(ikHandRoot, "root")
+		cmds.parent(ikFootLeft, ikFootRoot)
+		cmds.parent(ikFootRight, ikFootRoot)
+		cmds.parent(ikHandGun, ikHandRoot)
+		cmds.parent(ikHandLeft, ikHandGun)
+		cmds.parent(ikHandRight, ikHandGun)
+		
+		# constrain the joints
+		leftHandConstraint = cmds.parentConstraint("hand_l", ikHandLeft)[0]
+		rightHandGunConstraint = cmds.parentConstraint("hand_r", ikHandGun)[0]
+		rightHandConstraint = cmds.parentConstraint("hand_r", ikHandRight)[0]
+		leftFootConstraint = cmds.parentConstraint("foot_l", ikFootLeft)[0]
+		rightFootConstraint = cmds.parentConstraint("foot_r", ikFootRight)[0]
+	
+    except:
+		cmds.warning("Something went wrong")
+
+# ------------------------------------------------------------------
+# ADD the IK bones to the Export file.
+#save the AnimRig file before opening the export file.
+cmds.file(save = True, type = "mayaBinary", force = True)
+
+#ADD IK BONES
+#open export file
+sceneName = cmds.file(q = True, sceneName = True)
+exportFile = sceneName.partition("AnimRigs")[0] + "ExportFiles" +  sceneName.partition("AnimRigs")[2].partition(".mb")[0] + "_Export.mb"
+cmds.file(exportFile, open = True, force = True)
+
+#add ik bones
+addIKBones()
+
+# ------------------------------------------------------------------
+def addCustomAttrs():
+    attrs = cmds.listAttr("root")
+    exists = False
+    for one in attrs:
+        if one.find("DistanceCurve") != -1:
+            exists = True
+    if exists == False:
+        cmds.addAttr("root", ln="DistanceCurve", keyable=True)
+addCustomAttrs()
+
+#save the export file.
+cmds.file(save = True, type = "mayaBinary", force = True)
+
+######################################################################################
+######################################################################################
+'''
+# This script is only used for aligning the pistons to each other when placing the joint movers.   It is also to be used while creating the Rig Pose.  Don't use it in the post script running.
+# for each piston end, place the mover at the center of the piston disk (visually), and place the UP locator along the +Z axis of the top part of the piston. (+X should be the down axis, -Y should point up and out of the piston.)
+# This script will aim both piston parts at the locator and then rotate (at the end) the bott one to align with the angle it should use if its different.
+def alignPistons():
+    const = cmds.aimConstraint("bot_uparm_pistonC_top_mover", "bot_uparm_pistonC_bot_mover", weight=1, mo=False, aimVector=(-1, 0, 0), upVector=(0, 0, 1), worldUpType="object", worldUpVector=(0, 0, 1), worldUpObject="bot_uparm_pistonC_top_mover_UP")
+    cmds.delete(const)
+    const = cmds.aimConstraint("bot_uparm_pistonC_bot_mover", "bot_uparm_pistonC_top_mover", weight=1, mo=False, aimVector=(1, 0, 0), upVector=(0, 0, 1), worldUpType="object", worldUpVector=(0, 0, 1), worldUpObject="bot_uparm_pistonC_top_mover_UP")
+    cmds.delete(const)        
+    cmds.rotate(77.41, 0, 0, "bot_uparm_pistonC_bot_mover", r=True, os=True)
+
+    const = cmds.aimConstraint("bot_loarm_pistonB_bot_mover", "bot_loarm_pistonB_top_mover", weight=1, mo=False, aimVector=(1, 0, 0), upVector=(0, 0, 1), worldUpType="object", worldUpVector=(0, 0, 1), worldUpObject="bot_loarm_pistonB_top_mover_UP")
+    cmds.delete(const)
+    const = cmds.aimConstraint("bot_loarm_pistonB_top_mover", "bot_loarm_pistonB_bot_mover", weight=1, mo=False, aimVector=(-1, 0, 0), upVector=(0, 0, 1), worldUpType="object", worldUpVector=(0, 0, 1), worldUpObject="bot_loarm_pistonB_top_mover_UP")
+    cmds.delete(const)
+    cmds.rotate(15.8, 0, 0, "bot_loarm_pistonB_bot_mover", r=True, os=True)
+
+    const = cmds.aimConstraint("bot_loarm_pistonC_bot_mover", "bot_loarm_pistonC_top_mover", weight=1, mo=False, aimVector=(-1, 0, 0), upVector=(0, 0, 1), worldUpType="object", worldUpVector=(0, 0, 1), worldUpObject="bot_loarm_pistonC_bot_mover_UP")
+    cmds.delete(const)
+    const = cmds.aimConstraint("bot_loarm_pistonC_top_mover", "bot_loarm_pistonC_bot_mover", weight=1, mo=False, aimVector=(1, 0, 0), upVector=(0, 0, 1), worldUpType="object", worldUpVector=(0, 0, 1), worldUpObject="bot_loarm_pistonC_bot_mover_UP")
+    cmds.delete(const)
+    cmds.rotate(0, 0, 0, "bot_loarm_pistonC_top_mover", r=True, os=True)
+            
+    const = cmds.aimConstraint("bot_loarm_pistonD_bot_mover", "bot_loarm_pistonD_top_mover", weight=1, mo=False, aimVector=(-1, 0, 0), upVector=(0, 0, 1), worldUpType="object", worldUpVector=(0, 0, 1), worldUpObject="bot_loarm_pistonD_bot_mover_UP")
+    cmds.delete(const)
+    const = cmds.aimConstraint("bot_loarm_pistonD_top_mover", "bot_loarm_pistonD_bot_mover", weight=1, mo=False, aimVector=(1, 0, 0), upVector=(0, 0, 1), worldUpType="object", worldUpVector=(0, 0, 1), worldUpObject="bot_loarm_pistonD_bot_mover_UP")
+    cmds.delete(const)
+    cmds.rotate(-84.5, 0, 0, "bot_loarm_pistonD_top_mover", r=True, os=True)
+alignPistons()
+
+
+# Then make sure that the UP locators are parent constrained to the uparm and loarm or loarm twist movers. (the ones that are leafs, not the actual skeleton ones.) Use the following script to do this.
+# REMINDER: This must be done BEFORE you move on to Deformation Setup phase.
+cmds.parentConstraint("bot_uparm_mover", "bot_uparm_pistonC_top_mover_UP", weight=1, mo=True)
+cmds.parentConstraint("bot_loarm_mover", "bot_loarm_pistonB_top_mover_UP", weight=1, mo=True)
+cmds.parentConstraint("bot_loarm_twist_mover", "bot_loarm_pistonC_bot_mover_UP", weight=1, mo=True)
+cmds.parentConstraint("bot_loarm_twist_mover", "bot_loarm_pistonD_bot_mover_UP", weight=1, mo=True)
+'''
\ No newline at end of file
