Branching

diff --git a/Engine/Extras/Maya_AnimationRiggingTools/ARTv1/General/ART/Projects/Orion_Cine/Scripts/postScripts/RiftMage_Post_Script.py b/Engine/Extras/Maya_AnimationRiggingTools/ARTv1/General/ART/Projects/Orion_Cine/Scripts/postScripts/RiftMage_Post_Script.py
new file mode 100644
--- /dev/null
+++ b/Engine/Extras/Maya_AnimationRiggingTools/ARTv1/General/ART/Projects/Orion_Cine/Scripts/postScripts/RiftMage_Post_Script.py
@@ -0,0 +1,237 @@
+import maya.cmds as cmds
+import maya.mel as mel
+
+# ------------------------------------------------------------------
+# scale all of the controls to be more visible with the mesh.
+def scaleControls():
+    controls = ["cowl_l_anim", "cowl_r_anim", "cowl_bk_anim", "outer_ring_A_piston01_l_anim", "outer_ring_A_piston01_r_anim", "outer_ring_A_piston02_l_anim", "outer_ring_A_piston02_r_anim", "outer_ring_B_piston01_l_anim", "outer_ring_B_piston01_r_anim", "hood_l_anim", "hood_r_anim", "outer_ring_A_l_anim", "outer_ring_A_r_anim", "outer_ring_B_l_anim", "outer_ring_B_r_anim", "outer_ring_C_l_anim", "outer_ring_C_r_anim"]
+    for control in controls:
+        cmds.select(control+".cv[0:32]")
+        cmds.scale(0.2, 0.2, 0.2, r=True)
+
+
+    cmds.select(["portal_anim.cv[0:32]"])
+    cmds.scale(1.5, 0.5, 1.5, r = True, ocp = True)
+
+    cmds.select(["outer_ring_base_anim.cv[0:32]"])
+    cmds.scale(4.5, 0.001, 4.5, r = True, ocp = True)
+
+    cmds.select(["portal_ring_anim.cv[0:32]"])
+    cmds.scale(2.5, 0.001, 2.5, r = True, ocp = True)
+
+    cmds.select(["outer_ring_A_l_anim.cv[0:32]"])
+    cmds.move(4, 0, 4, r = True)
+
+    cmds.select(["outer_ring_A_r_anim.cv[0:32]"])
+    cmds.move(-4, 0, 4, r = True)
+
+    cmds.select(["outer_ring_B_l_anim.cv[0:32]"])
+    cmds.move(4, 0, 0, r = True)
+
+    cmds.select(["outer_ring_B_r_anim.cv[0:32]"])
+    cmds.move(-4, 0, 0, r = True)
+
+    cmds.select(["outer_ring_C_l_anim.cv[0:32]"])
+    cmds.move(4, 0, -4, r = True)
+
+    cmds.select(["outer_ring_C_r_anim.cv[0:32]"])
+    cmds.move(-4, 0, -4, r = True)
+
+    cmds.select(["ik_foot_anim_l.cv[0:17]"])
+    cmds.move(0, 3, -2.22651, r = True)
+
+    cmds.select(["ik_foot_anim_r.cv[0:17]"])
+    cmds.move(0, 3, -2.22651, r = True)
+
+    cmds.select(["cowl_bk_anim.cv[0:32]"])
+    cmds.move(0, 0, 6.5, r = True)
+
+    cmds.select(["cowl_l_anim.cv[0:32]"])
+    cmds.move(0, 0, 2, r = True)
+
+    cmds.select(["cowl_r_anim.cv[0:32]"])
+    cmds.move(0, 0, 2, r = True)
+
+    cmds.select(None)
+
+scaleControls()
+
+# ------------------------------------------------------------------
+# Build the rig for the rings and connect them to the arm rig.
+def buildRingRig():
+    import Tools.ART_addSpaceSwitching as space
+    reload(space)
+
+    #space.addSpaceSwitching("portal_ring_anim", "outer_ring_base_anim", True)
+    
+    controlsList = {"outer_ring_A_l":"outer_ring_base_anim", "outer_ring_B_l":"outer_ring_A_l_anim", "outer_ring_C_l":"outer_ring_B_l_anim", "outer_ring_A_piston01_l":"outer_ring_A_l_anim", "outer_ring_A_piston02_l":"outer_ring_A_l_anim", "outer_ring_B_piston01_l":"outer_ring_B_l_anim", "outer_ring_A_r":"outer_ring_base_anim", "outer_ring_B_r":"outer_ring_A_r_anim", "outer_ring_C_r":"outer_ring_B_r_anim", "outer_ring_A_piston01_r":"outer_ring_A_r_anim", "outer_ring_A_piston02_r":"outer_ring_A_r_anim", "outer_ring_B_piston01_r":"outer_ring_B_r_anim", "portal":"portal_ring_anim"}
+    for one in controlsList:
+        # Re-create the scale constraint between the anim control and the driver joint.
+        sourceConn = cmds.listConnections("driver_"+one+".s", c=True, p=True)
+        cmds.disconnectAttr(sourceConn[1], sourceConn[0])
+        cmds.scaleConstraint(one+"_anim", "driver_"+one, mo=True)
+        
+        # scale constraint the control to the leaf control we want to scale it.
+        cmds.scaleConstraint(controlsList[one], one+"_parent_grp", mo=True)
+
+        # set the joints (driver and deform) to have segment scale compensate off.
+        cmds.setAttr("driver_"+one+".segmentScaleCompensate", 0)
+        cmds.setAttr(one+".segmentScaleCompensate", 0)
+            
+    cmds.select("outer_ring_base_anim")
+
+buildRingRig()
+
+# ------------------------------------------------------------------
+# Delete this once the animators lock down how big they want the ring to be.
+def TEMPRINGMORPH():
+    cmds.addAttr("outer_ring_base_anim", ln="Thickness", keyable=True)
+    cmds.connectAttr("outer_ring_base_anim.Thickness", "blendShape1.outer_ring_top_geo_scaled", f=True)
+    cmds.connectAttr("outer_ring_base_anim.Thickness", "blendShape2.outer_ring_bott_geo_scaled", f=True)
+    cmds.connectAttr("outer_ring_base_anim.Thickness", "blendShape3.outer_ring_inside_geo_scaled", f=True)
+#TEMPRINGMORPH()
+
+# ------------------------------------------------------------------
+def dynamicsOff():
+    for one in ["cloth_chest_bk_l_anim", "cloth_chest_bk_r_anim", "cloth_chest_fr_l_anim", "cloth_chest_fr_r_anim", "chest_cloth_dyn_anim", "hips_cloth_bk_l_dyn_anim", "hips_cloth_bk_r_dyn_anim", "hips_cloth_fr_l_dyn_anim", "hips_cloth_fr_r_dyn_anim", "loin_cloth_bk_dyn_anim", "loin_cloth_fr_dyn_anim"]:
+        cmds.setAttr(one+".chainStartEnvelope", 0)
+        cmds.setAttr(one+".chainStartEnvelope", l=True)
+dynamicsOff()
+
+# ------------------------------------------------------------------
+def createGameCam():
+    gameCam = cmds.camera(centerOfInterest=5, focalLength=18, lensSqueezeRatio=1, cameraScale=1, horizontalFilmAperture=1.41732, horizontalFilmOffset=0, verticalFilmAperture=0.94488, verticalFilmOffset=0, filmFit="Overscan", overscan=1.05, motionBlur=0, shutterAngle=144, nearClipPlane=1.0, farClipPlane=10000, orthographic=0, orthographicWidth=30, panZoomEnabled=0, horizontalPan=0, verticalPan=0, zoom=1, displayFilmGate=True)[0]
+    print gameCam
+    gameCam = cmds.rename(gameCam, "GameCam")
+    cmds.camera(gameCam, e=True, horizontalFilmAperture=1.6724376, focalLength=21.2399857011)
+    cmds.setAttr(gameCam+".displayGateMaskOpacity", 1)
+    cmds.setAttr(gameCam+".displayGateMaskColor", 0, 0, 0, type="double3")
+    cmds.setAttr(gameCam+".r", 90, 0, 180, type="double3")
+    
+    # Create the control
+    control = cmds.circle(name=gameCam+"_Ctrl", ch=True, o=True, r=15)[0]
+    cmds.makeIdentity(control, a=True, t=True, r=True, s=True)
+    # Create groups for the sdk and the constraints
+    const_grp = cmds.group(control, name=control+"_grp")
+    if cmds.objExists("rig_grp"):
+        cmds.parent(const_grp, "rig_grp")
+    # Parent Constrain the cam to the control
+    cmds.parentConstraint(control, gameCam, mo=True)
+    
+    cmds.connectAttr(control+".visibility", gameCam+".visibility", f=True)
+    cmds.setAttr(const_grp+".v", 0)
+    
+    cmds.setAttr(const_grp+".t", 0, 500, 200, type="double3")
+createGameCam()
+
+# ------------------------------------------------------------------
+def addCapsule(radius, halfHeight):
+    trueHalfHeight = (halfHeight*2)-(radius*2)
+    capsule = cmds.polyCylinder(name="collision_capsule", r=radius, h=trueHalfHeight, sx=18, sy=1, sz=5, ax=[0,0,1], rcp=True, cuv=3, ch=True)
+    cmds.setAttr(capsule[0]+".tz", halfHeight+2.2)
+    cmds.setAttr(capsule[0]+".rz", 90)
+
+    cmds.addAttr(capsule[0], ln="Radius", dv=radius, keyable=True)
+    cmds.addAttr(capsule[0], ln="HalfHeight", dv=halfHeight, keyable=True)
+
+    expr = capsule[1]+".radius = "+capsule[0]+".Radius;\n"+capsule[1]+".height = ("+capsule[0]+".HalfHeight*2) - ("+capsule[0]+".Radius*2);"
+    cmds.expression(s=expr, o=capsule[1], ae=1, uc="all")
+
+    cmds.setAttr(capsule[0]+".overrideEnabled", 1)
+    cmds.setAttr(capsule[0]+".overrideColor", 18)
+    cmds.setAttr(capsule[0]+".overrideShading", 0)
+
+    cmds.addAttr(capsule[0], ln="Shaded", dv=0, min=0, max=1, keyable=True)
+    cmds.connectAttr(capsule[0]+".Shaded", capsule[0]+".overrideShading")
+
+    cmds.parentConstraint("root", capsule[0], mo=True)
+    cmds.select(capsule[0], r=True)
+    capsuleDisplayLayer = cmds.createDisplayLayer(nr=True, name="Collision Capsule")
+    cmds.setAttr(capsuleDisplayLayer+".displayType", 2)
+    cmds.setAttr(capsuleDisplayLayer+".v", 0)
+    
+    for attr in ["tx", "ty", "tz", "rx", "ry", "rz", "sx", "sy", "sz"]:
+        cmds.setAttr(capsule[0]+"."+attr, lock=True, keyable=False, channelBox=False)
+
+    cmds.select(clear=True)
+addCapsule(42, 86)
+
+# ------------------------------------------------------------------
+# add ik bones Function that is called later after we open the export file.
+def addIKBones():
+    try:
+		# create the joints
+		cmds.select(clear = True)
+		ikFootRoot = cmds.joint(name = "ik_foot_root")
+		cmds.select(clear = True)
+		
+		cmds.select(clear = True)
+		ikFootLeft = cmds.joint(name = "ik_foot_l")
+		cmds.select(clear = True)
+		
+		cmds.select(clear = True)
+		ikFootRight = cmds.joint(name = "ik_foot_r")
+		cmds.select(clear = True)
+		
+		cmds.select(clear = True)
+		ikHandRoot = cmds.joint(name = "ik_hand_root")
+		cmds.select(clear = True)
+		
+		cmds.select(clear = True)
+		ikHandGun = cmds.joint(name = "ik_hand_gun")
+		cmds.select(clear = True)
+		
+		cmds.select(clear = True)
+		ikHandLeft = cmds.joint(name = "ik_hand_l")
+		cmds.select(clear = True)
+		
+		cmds.select(clear = True)
+		ikHandRight = cmds.joint(name = "ik_hand_r")
+		cmds.select(clear = True)
+		
+		# create hierarhcy
+		cmds.parent(ikFootRoot, "root")
+		cmds.parent(ikHandRoot, "root")
+		cmds.parent(ikFootLeft, ikFootRoot)
+		cmds.parent(ikFootRight, ikFootRoot)
+		cmds.parent(ikHandGun, ikHandRoot)
+		cmds.parent(ikHandLeft, ikHandGun)
+		cmds.parent(ikHandRight, ikHandGun)
+		
+		# constrain the joints
+		leftHandConstraint = cmds.parentConstraint("hand_l", ikHandLeft)[0]
+		rightHandGunConstraint = cmds.parentConstraint("hand_r", ikHandGun)[0]
+		rightHandConstraint = cmds.parentConstraint("hand_r", ikHandRight)[0]
+		leftFootConstraint = cmds.parentConstraint("foot_l", ikFootLeft)[0]
+		rightFootConstraint = cmds.parentConstraint("foot_r", ikFootRight)[0]
+	
+    except:
+		cmds.warning("Something went wrong")
+
+# ------------------------------------------------------------------
+# ADD the IK bones to the Export file.
+#save the AnimRig file before opening the export file.
+cmds.file(save = True, type = "mayaBinary", force = True)
+
+#ADD IK BONES
+#open export file
+sceneName = cmds.file(q = True, sceneName = True)
+exportFile = sceneName.partition("AnimRigs")[0] + "ExportFiles" +  sceneName.partition("AnimRigs")[2].partition(".mb")[0] + "_Export.mb"
+cmds.file(exportFile, open = True, force = True)
+
+#add ik bones
+addIKBones()
+
+# ------------------------------------------------------------------
+def addCustomAttrs():
+    attrs = cmds.listAttr("root")
+    exists = False
+    for one in attrs:
+        if one.find("DistanceCurve") != -1:
+            exists = True
+    if exists == False:
+        cmds.addAttr("root", ln="DistanceCurve", keyable=True)
+addCustomAttrs()
+
+#save the export file.
+cmds.file(save = True, type = "mayaBinary", force = True)
\ No newline at end of file
diff --git a/Engine/Extras/Maya_AnimationRiggingTools/ARTv1/General/ART/Projects/Orion_Cine/Scripts/postScripts/Sparrow_Post_Script.py b/Engine/Extras/Maya_AnimationRiggingTools/ARTv1/General/ART/Projects/Orion_Cine/Scripts/postScripts/Sparrow_Post_Script.py
new file mode 100644
--- /dev/null
+++ b/Engine/Extras/Maya_AnimationRiggingTools/ARTv1/General/ART/Projects/Orion_Cine/Scripts/postScripts/Sparrow_Post_Script.py
@@ -0,0 +1,300 @@
+import maya.cmds as cmds
+
+# NOTE: This character has special steps to run during Rig Pose creation.  See end of script.
+
+# ------------------------------------------------------------------
+# scale all of the controls to be more visible with the mesh.
+def scaleControls():
+    sides = ["_l", "_r"]
+    for side in sides:
+        # feet
+        cmds.select(["ik_foot_anim"+side+".cv[*]"])
+        cmds.scale(1, 1, 1, r = True, ocp = False)
+        cmds.move(0, 0, 0, r = True)
+
+        cmds.select(["ik_wrist"+side+"_anim.cv[*]"])
+        cmds.scale(1, 1, 1, r = True, ocp = False)
+
+    objects = ["hair_top_anim", "fk_ponytail_01_anim", "fk_ponytail_02_anim", "fk_ponytail_03_anim", "fk_ponytail_04_anim", "quiver_anim", "arrow_nock_anim", "bow_base_anim", "bow_shaft_top_01_anim", "bow_shaft_top_02_anim", "bow_shaft_top_03_anim", "bow_shaft_bott_01_anim", "bow_shaft_bott_02_anim", "bow_shaft_bott_03_anim", "bow_string_bott_anim", "bow_string_mid_anim", "bow_string_top_anim", "front_strap_small_anim", "fk_front_strap_01_anim", "fk_front_strap_02_anim", "fk_skirt_back_cloth_01_anim", "fk_skirt_back_cloth_02_anim", "fk_skirt_cloth_l_01_anim", "fk_skirt_cloth_l_02_anim", "fk_skirt_front_cloth_01_anim", "fk_skirt_front_cloth_02_anim", "breast_chain_anim", "feathers_anim", "fk_hair_curls_A_l_01_anim", "fk_hair_curls_A_l_02_anim", "fk_hair_curls_A_r_01_anim", "fk_hair_curls_A_r_02_anim", "fk_hair_curls_B_l_01_anim", "fk_hair_curls_B_l_02_anim", "fk_hair_curls_B_r_01_anim", "fk_hair_curls_B_r_02_anim", "fk_chest_cloth_l_01_anim", "fk_chest_cloth_l_02_anim", "fk_chest_cloth_r_01_anim", "fk_chest_cloth_r_02_anim", "fk_back_tail_cloth_01_anim", "fk_back_tail_cloth_02_anim"]
+    for one in objects:
+        cmds.select([one+".cv[*]"])
+        cmds.scale(0.227057, 0.227057, 0.227057, r=True, ocp=False)
+        cmds.move(0, 0, 0, r = True, os=True, wd=True)
+
+    cmds.select(["arrow_nock_anim.cv[*]"])
+    cmds.scale(0, 0.3, 4, r = True, ocp = True)
+
+    cmds.select(["bow_string_mid_anim.cv[*]"])
+    cmds.scale(0, 0.6, 4, r = True, ocp = True)
+
+    cmds.select(None)
+scaleControls()
+
+# ------------------------------------------------------------------
+def bowAndArrowRig():
+    import Tools.ART_addSpaceSwitching as space
+    reload(space)
+        
+    # Create Arrow base and tip controls
+    arrow_base_anim = cmds.duplicate("arrow_nock_anim", n="arrow_base_anim")
+    cmds.xform(arrow_base_anim, t=[0, 0, -35.51354])
+    cmds.makeIdentity(arrow_base_anim, a=True, t=True)
+
+    arrow_tip_anim = cmds.duplicate("arrow_nock_anim", n="arrow_tip_anim")
+    cmds.xform(arrow_tip_anim, t=[0, 0, -78.67631])
+    cmds.makeIdentity(arrow_tip_anim, a=True, t=True)
+
+    cmds.parent(arrow_base_anim, "arrow_nock_anim")
+    cmds.parent(arrow_tip_anim, arrow_base_anim)
+
+    cmds.delete("driver_arrow_nock_parentConstraint1")
+    cmds.parentConstraint(arrow_tip_anim, "driver_arrow_nock", mo=True)
+    
+    # Create aimConstraints for the ends of the strings
+    cmds.aimConstraint("bow_string_mid_anim", "bow_string_bott_anim_grp", mo=True, aimVector=(0, 0, -1), upVector=(0, -1, 0), worldUpType="object", worldUpObject="bow_base_anim")
+    cmds.aimConstraint("bow_string_mid_anim", "bow_string_top_anim_grp", mo=True, aimVector=(0, 0, -1), upVector=(0, -1, 0), worldUpType="object", worldUpObject="bow_base_anim")
+
+    # Create SDK's for the bow shafts to flex when the string is pulled.
+    cmds.setAttr("bow_string_mid_anim.translateY", 0)
+    cmds.setAttr("bow_shaft_top_01_anim_grp.rotateX", 0)
+    cmds.setAttr("bow_shaft_top_02_anim_grp.rotateX", 0)
+    cmds.setAttr("bow_shaft_top_03_anim_grp.rotateX", 0)
+    cmds.setAttr("bow_shaft_bott_01_anim_grp.rotateX", 0)
+    cmds.setAttr("bow_shaft_bott_02_anim_grp.rotateX", 0)
+    cmds.setAttr("bow_shaft_bott_03_anim_grp.rotateX", 0)
+    cmds.setDrivenKeyframe(["bow_shaft_top_01_anim_grp.rx", "bow_shaft_top_02_anim_grp.rx", "bow_shaft_top_03_anim_grp.rx", "bow_shaft_bott_01_anim_grp.rx", "bow_shaft_bott_02_anim_grp.rx", "bow_shaft_bott_03_anim_grp.rx"], cd = "bow_string_mid_anim.translateY", itt = "spline", ott = "spline")
+
+    cmds.setAttr("bow_string_mid_anim.translateY", 120)
+    cmds.setAttr("bow_shaft_top_01_anim_grp.rotateX", -10)
+    cmds.setAttr("bow_shaft_top_02_anim_grp.rotateX", -14)
+    cmds.setAttr("bow_shaft_top_03_anim_grp.rotateX", -19.4)
+    cmds.setAttr("bow_shaft_bott_01_anim_grp.rotateX", -10)
+    cmds.setAttr("bow_shaft_bott_02_anim_grp.rotateX", -14)
+    cmds.setAttr("bow_shaft_bott_03_anim_grp.rotateX", -19.4)
+    cmds.setDrivenKeyframe(["bow_shaft_top_01_anim_grp.rx", "bow_shaft_top_02_anim_grp.rx", "bow_shaft_top_03_anim_grp.rx", "bow_shaft_bott_01_anim_grp.rx", "bow_shaft_bott_02_anim_grp.rx", "bow_shaft_bott_03_anim_grp.rx"], cd = "bow_string_mid_anim.translateY", itt = "spline", ott = "spline")
+
+    cmds.setAttr("bow_string_mid_anim.translateY", 0)
+    cmds.setAttr("bow_shaft_top_01_anim_grp.rotateX", 0)
+    cmds.setAttr("bow_shaft_top_02_anim_grp.rotateX", 0)
+    cmds.setAttr("bow_shaft_top_03_anim_grp.rotateX", 0)
+    cmds.setAttr("bow_shaft_bott_01_anim_grp.rotateX", 0)
+    cmds.setAttr("bow_shaft_bott_02_anim_grp.rotateX", 0)
+    cmds.setAttr("bow_shaft_bott_03_anim_grp.rotateX", 0)
+
+    # Move the arrow to the string
+    cmds.setAttr("arrow_nock_anim_grp.translateX", 0.05019)
+    cmds.setAttr("arrow_nock_anim_grp.translateY", 15.07779)
+    cmds.setAttr("arrow_nock_anim_grp.translateZ", -35.42169)
+    cmds.setAttr("arrow_nock_anim_grp.rotateX", -90)
+    cmds.setAttr("arrow_nock_anim_grp.rotateY", 0)
+    cmds.setAttr("arrow_nock_anim_grp.rotateZ", 0)
+
+    # Make the arrow follow the string using space switching
+    space.addSpaceSwitching("arrow_nock_anim", "bow_string_mid_anim", True)
+
+    #Move the string to an appropriate position relative to where it might be aligned when shooting.
+    cmds.setAttr("bow_string_mid_anim.translateX", -2.45216)
+    cmds.setAttr("bow_string_mid_anim.translateZ", 5.9665)
+
+    # Create an aim constraint rig so that the arrow will always stay aimed at its position on the shaft.
+    arrow_aim_target = cmds.circle(n="arrow_notched_target", r=3.4)[0]
+    cmds.xform(arrow_aim_target, t=(-2.18154, 0.91122, 6.11286))
+    arrow_up_target = cmds.spaceLocator(n="arrow_up_target")
+    cmds.xform(arrow_up_target, t=(-2.18154, 0.91122, 36))
+    cmds.parent(arrow_aim_target, arrow_up_target, "bow_base_anim")
+    cmds.makeIdentity(arrow_aim_target, a=True, t=True, r=True)
+    cmds.aimConstraint(arrow_aim_target, "arrow_nock_anim_grp", mo=True, aimVector=(0, 0, -1), upVector=(0, -1, 0), worldUpType="object", worldUpObject=arrow_up_target[0])
+
+    # Add space switching for the string to follow the hand.  Off by default.
+    space.addSpaceSwitching("bow_string_mid_anim", "driver_hand_r", False)
+    
+    # Move the bow to the left hand.
+    cmds.setAttr("bow_base_anim.translateX", 74.70517)
+    cmds.setAttr("bow_base_anim.translateY", 6.12268)
+    cmds.setAttr("bow_base_anim.translateZ", 144.79125)
+    cmds.setAttr("bow_base_anim.rotateX", 183.95171)
+    cmds.setAttr("bow_base_anim.rotateY", -86.42158)
+    cmds.setAttr("bow_base_anim.rotateZ", -83.33129)
+
+bowAndArrowRig()
+
+# ------------------------------------------------------------------
+def armsLegsRig():
+    cmds.setAttr("Rig_Settings.lUpperarmTwist2Amount", 0.3)
+    cmds.setAttr("Rig_Settings.rUpperarmTwist2Amount", 0.3)
+
+    cmds.setAttr("Rig_Settings.lForearmTwist2Amount", 0.3)
+    cmds.setAttr("Rig_Settings.rForearmTwist2Amount", 0.3)
+
+    cmds.setAttr("Rig_Settings.lThighTwist2Amount", 0.3)
+    cmds.setAttr("Rig_Settings.rThighTwist2Amount", 0.3)
+
+    cmds.setAttr("Rig_Settings.lCalfTwistAmount", 0.7)
+    cmds.setAttr("Rig_Settings.rCalfTwistAmount", 0.7)
+
+armsLegsRig()
+
+# ------------------------------------------------------------------
+def dynamicsOff():
+    for one in ["back_tail_cloth_dyn_anim", "chest_cloth_l_dyn_anim", "chest_cloth_r_dyn_anim", "front_strap_dyn_anim", "hair_curls_A_l_dyn_anim", "hair_curls_A_r_dyn_anim", "hair_curls_B_l_dyn_anim", "hair_curls_B_r_dyn_anim", "ponytail_dyn_anim", "skirt_back_cloth_dyn_anim", "skirt_cloth_l_dyn_anim", "skirt_front_cloth_dyn_anim"]:
+        cmds.setAttr(one+".chainStartEnvelope", 0)
+        cmds.setAttr(one+".chainStartEnvelope", l=True)
+dynamicsOff()
+
+# ------------------------------------------------------------------
+def addCapsule(radius, halfHeight):
+    trueHalfHeight = (halfHeight*2)-(radius*2)
+    capsule = cmds.polyCylinder(name="collision_capsule", r=radius, h=trueHalfHeight, sx=18, sy=1, sz=5, ax=[0,0,1], rcp=True, cuv=3, ch=True)
+    cmds.setAttr(capsule[0]+".tz", halfHeight+2.2)
+    cmds.setAttr(capsule[0]+".rz", 90)
+
+    cmds.addAttr(capsule[0], ln="Radius", dv=radius, keyable=True)
+    cmds.addAttr(capsule[0], ln="HalfHeight", dv=halfHeight, keyable=True)
+
+    expr = capsule[1]+".radius = "+capsule[0]+".Radius;\n"+capsule[1]+".height = ("+capsule[0]+".HalfHeight*2) - ("+capsule[0]+".Radius*2);"
+    cmds.expression(s=expr, o=capsule[1], ae=1, uc="all")
+
+    cmds.setAttr(capsule[0]+".overrideEnabled", 1)
+    cmds.setAttr(capsule[0]+".overrideColor", 18)
+    cmds.setAttr(capsule[0]+".overrideShading", 0)
+
+    cmds.addAttr(capsule[0], ln="Shaded", dv=0, min=0, max=1, keyable=True)
+    cmds.connectAttr(capsule[0]+".Shaded", capsule[0]+".overrideShading")
+
+    cmds.parentConstraint("root", capsule[0], mo=True)
+    cmds.select(capsule[0], r=True)
+    capsuleDisplayLayer = cmds.createDisplayLayer(nr=True, name="Collision Capsule")
+    cmds.setAttr(capsuleDisplayLayer+".displayType", 2)
+    cmds.setAttr(capsuleDisplayLayer+".v", 0)
+    
+    for attr in ["tx", "ty", "tz", "rx", "ry", "rz", "sx", "sy", "sz"]:
+        cmds.setAttr(capsule[0]+"."+attr, lock=True, keyable=False, channelBox=False)
+
+    cmds.select(clear=True)
+addCapsule(42, 86)
+
+# ------------------------------------------------------------------
+def createGameCam():
+    gameCam = cmds.camera(centerOfInterest=5, focalLength=18, lensSqueezeRatio=1, cameraScale=1, horizontalFilmAperture=1.41732, horizontalFilmOffset=0, verticalFilmAperture=0.94488, verticalFilmOffset=0, filmFit="Overscan", overscan=1.05, motionBlur=0, shutterAngle=144, nearClipPlane=1.0, farClipPlane=10000, orthographic=0, orthographicWidth=30, panZoomEnabled=0, horizontalPan=0, verticalPan=0, zoom=1, displayFilmGate=True)[0]
+    print gameCam
+    gameCam = cmds.rename(gameCam, "GameCam")
+    cmds.camera(gameCam, e=True, horizontalFilmAperture=1.6724376, focalLength=21.2399857011)
+    cmds.setAttr(gameCam+".displayGateMaskOpacity", 1)
+    cmds.setAttr(gameCam+".displayGateMaskColor", 0, 0, 0, type="double3")
+    cmds.setAttr(gameCam+".r", 90, 0, 180, type="double3")
+    
+    # Create the control
+    control = cmds.circle(name=gameCam+"_Ctrl", ch=True, o=True, r=15)[0]
+    cmds.makeIdentity(control, a=True, t=True, r=True, s=True)
+    # Create groups for the sdk and the constraints
+    const_grp = cmds.group(control, name=control+"_grp")
+    if cmds.objExists("rig_grp"):
+        cmds.parent(const_grp, "rig_grp")
+    # Parent Constrain the cam to the control
+    cmds.parentConstraint(control, gameCam, mo=True)
+    
+    cmds.connectAttr(control+".visibility", gameCam+".visibility", f=True)
+    cmds.setAttr(const_grp+".v", 0)
+    
+    cmds.setAttr(const_grp+".t", 0, 500, 200, type="double3")
+createGameCam()
+
+# ------------------------------------------------------------------
+# add ik bones Function that is called later after we open the export file.
+def addIKBones():
+    try:
+		# create the joints
+		cmds.select(clear = True)
+		ikFootRoot = cmds.joint(name = "ik_foot_root")
+		cmds.select(clear = True)
+		
+		cmds.select(clear = True)
+		ikFootLeft = cmds.joint(name = "ik_foot_l")
+		cmds.select(clear = True)
+		
+		cmds.select(clear = True)
+		ikFootRight = cmds.joint(name = "ik_foot_r")
+		cmds.select(clear = True)
+		
+		cmds.select(clear = True)
+		ikHandRoot = cmds.joint(name = "ik_hand_root")
+		cmds.select(clear = True)
+		
+		cmds.select(clear = True)
+		ikHandGun = cmds.joint(name = "ik_hand_gun")
+		cmds.select(clear = True)
+		
+		cmds.select(clear = True)
+		ikHandLeft = cmds.joint(name = "ik_hand_l")
+		cmds.select(clear = True)
+		
+		cmds.select(clear = True)
+		ikHandRight = cmds.joint(name = "ik_hand_r")
+		cmds.select(clear = True)
+		
+		# create hierarhcy
+		cmds.parent(ikFootRoot, "root")
+		cmds.parent(ikHandRoot, "root")
+		cmds.parent(ikFootLeft, ikFootRoot)
+		cmds.parent(ikFootRight, ikFootRoot)
+		cmds.parent(ikHandGun, ikHandRoot)
+		cmds.parent(ikHandLeft, ikHandGun)
+		cmds.parent(ikHandRight, ikHandGun)
+		
+		# constrain the joints
+		leftHandConstraint = cmds.parentConstraint("hand_l", ikHandLeft)[0]
+		rightHandGunConstraint = cmds.parentConstraint("hand_r", ikHandGun)[0]
+		rightHandConstraint = cmds.parentConstraint("hand_r", ikHandRight)[0]
+		leftFootConstraint = cmds.parentConstraint("foot_l", ikFootLeft)[0]
+		rightFootConstraint = cmds.parentConstraint("foot_r", ikFootRight)[0]
+	
+    except:
+		cmds.warning("Something went wrong")
+
+# ------------------------------------------------------------------
+# ADD the IK bones to the Export file.
+#save the AnimRig file before opening the export file.
+cmds.file(save = True, type = "mayaBinary", force = True)
+
+#ADD IK BONES
+#open export file
+sceneName = cmds.file(q = True, sceneName = True)
+exportFile = sceneName.partition("AnimRigs")[0] + "ExportFiles" +  sceneName.partition("AnimRigs")[2].partition(".mb")[0] + "_Export.mb"
+cmds.file(exportFile, open = True, force = True)
+
+#add ik bones
+addIKBones()
+
+# ------------------------------------------------------------------
+def addCustomAttrs():
+    attrs = cmds.listAttr("root")
+    exists = False
+    for one in attrs:
+        if one.find("DistanceCurve") != -1:
+            exists = True
+    if exists == False:
+        cmds.addAttr("root", ln="DistanceCurve", keyable=True)
+addCustomAttrs()
+
+# ------------------------------------------------------------------
+#save the export file.
+cmds.file(save = True, type = "mayaBinary", force = True)
+
+
+
+#########################################
+# RUN THIS COMMAND DURING THE RIG POSE CREATION STEP TO GET THE SHIELD TO THE CORRECT LOCATION.  LEAVE IT COMMENTED OUT SO THAT IT DOESNT RUN DURING RIGGING.
+'''
+# The following 2 lines are used to find the locations and then input below.
+#cmds.xform(q=True, ws=True, rp=True)
+#cmds.xform(q=True, ro=True)
+
+
+cmds.xform("arrow_nock_mover", ws=True, t=[9.611692837947271e-14, -2.0994091418383203e-16, 35.41695785522444])
+cmds.xform("arrow_nock_mover", ro=[-13.78281392154952, -64.119833503541, 16.763202342514436])
+
+cmds.xform("bow_base_mover", ws=True, t=[3.3590546269413172e-15, -3.762654592984341e-15, -7.492745461319391e-15])
+cmds.xform("bow_base_mover", ro=[-13.782813921549694, 64.11983350354099, -16.763202342514536])
+
+'''
+#########################################
\ No newline at end of file
diff --git a/Engine/Extras/Maya_AnimationRiggingTools/ARTv1/General/ART/Projects/Orion_Cine/Scripts/postScripts/Steel_Post_Script.py b/Engine/Extras/Maya_AnimationRiggingTools/ARTv1/General/ART/Projects/Orion_Cine/Scripts/postScripts/Steel_Post_Script.py
new file mode 100644
--- /dev/null
+++ b/Engine/Extras/Maya_AnimationRiggingTools/ARTv1/General/ART/Projects/Orion_Cine/Scripts/postScripts/Steel_Post_Script.py
@@ -0,0 +1,672 @@
+import maya.cmds as cmds
+
+# NOTE: This character has special steps to run during Rig Pose creation.  See end of script.
+# ------------------------------------------------------------------
+def shieldRig():
+    #Move Shield to correct position
+    cmds.setAttr("sheild_main_anim_grp.tx", 190.409*0.65)
+    cmds.setAttr("sheild_main_anim_grp.ty", -101.619*0.65)
+    cmds.setAttr("sheild_main_anim_grp.tz", -130.408*0.65)
+    cmds.setAttr("sheild_main_anim_grp.ry", 90)
+
+    #Shield Anim Control Scale
+    cmds.select("sheild_main_anim" + ".cv[0:32]")
+    cmds.scale(6*0.65, 6*0.65, 0, r=True)
+
+    #add attrs to shield anim
+    attrsToAdd = ["shieldDeploy1", "shieldDeploy2", "shieldDeploy3", "shieldDeploy4", "panelA_deploy1", "panelA_deploy2", "panelB_deploy1", "panelB_deploy2", "panelC_deploy1", "panelC_deploy2", "panelD_deploy1", "panelD_deploy2"] 
+    for attr in attrsToAdd:
+        cmds.addAttr("sheild_main_anim", ln = attr, dv = 0, min = 0, keyable = True)
+
+    #Shield SDKS
+    #deploy1
+    cmds.setAttr("sheild_main_anim.shieldDeploy1", 0)
+    cmds.setDrivenKeyframe(["loarm_shield_top_l_anim.tz", "loarm_shield_top_l_anim.ry", "loarm_shield_bot_l_anim.tz", "loarm_shield_bot_l_anim.ry"], cd = "sheild_main_anim.shieldDeploy1", itt = "spline", ott = "spline")
+
+    cmds.setAttr("sheild_main_anim.shieldDeploy1", 10)
+    cmds.setAttr("loarm_shield_top_l_anim.tz", 13.22*0.65)
+    cmds.setAttr("loarm_shield_top_l_anim.ry", -11.846)
+    cmds.setAttr("loarm_shield_bot_l_anim.tz", -14.016*0.65)
+    cmds.setAttr("loarm_shield_bot_l_anim.ry", 14.991)
+    cmds.setDrivenKeyframe(["loarm_shield_top_l_anim.tz", "loarm_shield_top_l_anim.ry", "loarm_shield_bot_l_anim.tz", "loarm_shield_bot_l_anim.ry"], cd = "sheild_main_anim.shieldDeploy1", itt = "spline", ott = "spline")
+    cmds.setAttr("sheild_main_anim.shieldDeploy1", 0)
+
+    cmds.setInfinity("loarm_shield_top_l_anim.tz", pri = "linear", poi = "linear")
+    cmds.setInfinity("loarm_shield_top_l_anim.ry", pri = "linear", poi = "linear")
+    cmds.setInfinity("loarm_shield_bot_l_anim.tz", pri = "linear", poi = "linear")
+    cmds.setInfinity("loarm_shield_bot_l_anim.ry", pri = "linear", poi = "linear")
+
+    cmds.setAttr("loarm_shield_top_l_anim.tz", keyable = False, channelBox = False)
+    cmds.setAttr("loarm_shield_top_l_anim.ry", keyable = False, channelBox = False)
+    cmds.setAttr("loarm_shield_bot_l_anim.tz", keyable = False, channelBox = False)
+    cmds.setAttr("loarm_shield_bot_l_anim.ry", keyable = False, channelBox = False)
+
+    #deploy2
+    cmds.setAttr("sheild_main_anim.shieldDeploy2", 0)
+    cmds.setDrivenKeyframe(["sheild_main_anim_grp.ty"], cd = "sheild_main_anim.shieldDeploy2", itt = "spline", ott = "spline")
+
+    cmds.setAttr("sheild_main_anim.shieldDeploy2", 10)
+    cmds.setAttr("sheild_main_anim_grp.ty", -55*0.65)
+    cmds.setDrivenKeyframe(["sheild_main_anim_grp.ty"], cd = "sheild_main_anim.shieldDeploy2", itt = "spline", ott = "spline")
+    cmds.setAttr("sheild_main_anim.shieldDeploy2", 0)
+    cmds.setInfinity("sheild_main_anim_grp.ty", pri = "linear", poi = "linear")
+    cmds.setAttr("sheild_main_anim_grp.ty", keyable = False, channelBox = False)
+
+    #deploy3
+    cmds.setAttr("sheild_main_anim.shieldDeploy3", 10)
+    cmds.setDrivenKeyframe(["panel_bot_root_l_anim.sy", "panel_bot_root_r_anim.sy", "panel_top_root_l_anim.sy", "panel_top_root_r_anim.sy"], cd = "sheild_main_anim.shieldDeploy3", itt = "spline", ott = "spline")
+
+    cmds.setAttr("sheild_main_anim.shieldDeploy3", 0)
+
+    cmds.setAttr("panel_bot_root_l_anim.sy", .5*0.65)
+    cmds.setAttr("panel_bot_root_r_anim.sy", .5*0.65)
+    cmds.setAttr("panel_top_root_l_anim.sy", .5*0.65)
+    cmds.setAttr("panel_top_root_r_anim.sy", .5*0.65)
+    cmds.setDrivenKeyframe(["panel_bot_root_l_anim.sy", "panel_bot_root_r_anim.sy", "panel_top_root_l_anim.sy", "panel_top_root_r_anim.sy"], cd = "sheild_main_anim.shieldDeploy3", itt = "spline", ott = "spline")
+    cmds.setInfinity("panel_bot_root_l_anim.sy", pri = "linear", poi = "linear")
+    cmds.setInfinity("panel_bot_root_r_anim.sy", pri = "linear", poi = "linear")
+    cmds.setInfinity("panel_top_root_l_anim.sy", pri = "linear", poi = "linear")
+    cmds.setInfinity("panel_top_root_r_anim.sy", pri = "linear", poi = "linear")
+
+    cmds.setAttr("panel_bot_root_l_anim.sy", keyable = False, channelBox = False)
+    cmds.setAttr("panel_bot_root_r_anim.sy", keyable = False, channelBox = False)
+    cmds.setAttr("panel_top_root_l_anim.sy", keyable = False, channelBox = False)
+    cmds.setAttr("panel_top_root_r_anim.sy", keyable = False, channelBox = False)
+
+    #deploy4
+    cmds.setAttr("sheild_main_anim.shieldDeploy4", 10)
+    cmds.setDrivenKeyframe(["panel_bot_r_anim.sz", "panel_bot_l_anim.sz", "panel_top_r_anim.sz", "panel_top_l_anim.sz"], cd = "sheild_main_anim.shieldDeploy4", itt = "spline", ott = "spline")
+
+    cmds.setAttr("sheild_main_anim.shieldDeploy4", 0)
+
+    cmds.setAttr("panel_bot_r_anim.sz", .3*0.65)
+    cmds.setAttr("panel_bot_l_anim.sz", .3*0.65)
+    cmds.setAttr("panel_top_r_anim.sz", .3*0.65)
+    cmds.setAttr("panel_top_l_anim.sz", .3*0.65)
+    cmds.setDrivenKeyframe(["panel_bot_r_anim.sz", "panel_bot_l_anim.sz", "panel_top_r_anim.sz", "panel_top_l_anim.sz"], cd = "sheild_main_anim.shieldDeploy4", itt = "spline", ott = "spline")
+    cmds.setInfinity("panel_bot_r_anim.sz", pri = "linear", poi = "linear")
+    cmds.setInfinity("panel_bot_l_anim.sz", pri = "linear", poi = "linear")
+    cmds.setInfinity("panel_top_r_anim.sz", pri = "linear", poi = "linear")
+    cmds.setInfinity("panel_top_l_anim.sz", pri = "linear", poi = "linear")
+
+    cmds.setAttr("panel_bot_r_anim.sz", keyable = False, channelBox = False)
+    cmds.setAttr("panel_bot_l_anim.sz", keyable = False, channelBox = False)
+    cmds.setAttr("panel_top_r_anim.sz", keyable = False, channelBox = False)
+    cmds.setAttr("panel_top_l_anim.sz", keyable = False, channelBox = False)
+
+    #panelA 1
+    cmds.setAttr("sheild_main_anim.panelA_deploy1", 10)
+    cmds.setDrivenKeyframe(["panel_bot_root_l_anim.rz"], cd = "sheild_main_anim.panelA_deploy1", itt = "spline", ott = "spline")
+
+    cmds.setAttr("sheild_main_anim.panelA_deploy1", 0)
+    cmds.setAttr("panel_bot_root_l_anim.rz", -90)
+    cmds.setDrivenKeyframe(["panel_bot_root_l_anim.rz"], cd = "sheild_main_anim.panelA_deploy1", itt = "spline", ott = "spline")
+    cmds.setInfinity("panel_bot_root_l_anim.rz", pri = "linear", poi = "linear")
+    cmds.setAttr("panel_bot_root_l_anim.rz", keyable = False, channelBox = False)
+
+    #panelA 2
+    cmds.setAttr("sheild_main_anim.panelA_deploy2", 10)
+    cmds.setDrivenKeyframe(["panel_bot_l_anim.sx"], cd = "sheild_main_anim.panelA_deploy2", itt = "spline", ott = "spline")
+
+    cmds.setAttr("sheild_main_anim.panelA_deploy2", 0)
+    cmds.setAttr("panel_bot_l_anim.sx", 0.05*0.65)
+    cmds.setDrivenKeyframe(["panel_bot_l_anim.sx"], cd = "sheild_main_anim.panelA_deploy2", itt = "spline", ott = "spline")
+    cmds.setInfinity("panel_bot_l_anim.sx", pri = "linear", poi = "linear")
+    cmds.setAttr("panel_bot_l_anim.sx", keyable = False, channelBox = False)
+
+    #panelB 1
+    cmds.setAttr("sheild_main_anim.panelB_deploy1", 10)
+    cmds.setDrivenKeyframe(["panel_top_root_l_anim.rz"], cd = "sheild_main_anim.panelB_deploy1", itt = "spline", ott = "spline")
+
+    cmds.setAttr("sheild_main_anim.panelB_deploy1", 0)
+    cmds.setAttr("panel_top_root_l_anim.rz", -90)
+    cmds.setDrivenKeyframe(["panel_top_root_l_anim.rz"], cd = "sheild_main_anim.panelB_deploy1", itt = "spline", ott = "spline")
+    cmds.setInfinity("panel_top_root_l_anim.rz", pri = "linear", poi = "linear")
+    cmds.setAttr("panel_top_root_l_anim.rz", keyable = False, channelBox = False)
+
+    #panelB 2
+    cmds.setAttr("sheild_main_anim.panelB_deploy2", 10)
+    cmds.setDrivenKeyframe(["panel_top_l_anim.sx"], cd = "sheild_main_anim.panelB_deploy2", itt = "spline", ott = "spline")
+
+    cmds.setAttr("sheild_main_anim.panelB_deploy2", 0)
+    cmds.setAttr("panel_top_l_anim.sx", 0.05*0.65)
+    cmds.setDrivenKeyframe(["panel_top_l_anim.sx"], cd = "sheild_main_anim.panelB_deploy2", itt = "spline", ott = "spline")
+    cmds.setInfinity("panel_top_l_anim.sx", pri = "linear", poi = "linear")
+    cmds.setAttr("panel_top_l_anim.sx", keyable = False, channelBox = False)
+
+    #panelC 1
+    cmds.setAttr("sheild_main_anim.panelC_deploy1", 10)
+    cmds.setDrivenKeyframe(["panel_bot_root_r_anim.rz"], cd = "sheild_main_anim.panelC_deploy1", itt = "spline", ott = "spline")
+
+    cmds.setAttr("sheild_main_anim.panelC_deploy1", 0)
+    cmds.setAttr("panel_bot_root_r_anim.rz", -90)
+    cmds.setDrivenKeyframe(["panel_bot_root_r_anim.rz"], cd = "sheild_main_anim.panelC_deploy1", itt = "spline", ott = "spline")
+    cmds.setInfinity("panel_bot_root_r_anim.rz", pri = "linear", poi = "linear")
+    cmds.setAttr("panel_bot_root_r_anim.rz", keyable = False, channelBox = False)
+
+    #panelC 2
+    cmds.setAttr("sheild_main_anim.panelC_deploy2", 10)
+    cmds.setDrivenKeyframe(["panel_bot_r_anim.sx"], cd = "sheild_main_anim.panelC_deploy2", itt = "spline", ott = "spline")
+
+    cmds.setAttr("sheild_main_anim.panelC_deploy2", 0)
+    cmds.setAttr("panel_bot_r_anim.sx", 0.05*0.65)
+    cmds.setDrivenKeyframe(["panel_bot_r_anim.sx"], cd = "sheild_main_anim.panelC_deploy2", itt = "spline", ott = "spline")
+    cmds.setInfinity("panel_bot_r_anim.sx", pri = "linear", poi = "linear")
+    cmds.setAttr("panel_bot_r_anim.sx", keyable = False, channelBox = False)
+
+    #panelD 1
+    cmds.setAttr("sheild_main_anim.panelD_deploy1", 10)
+    cmds.setDrivenKeyframe(["panel_top_root_r_anim.rz"], cd = "sheild_main_anim.panelD_deploy1", itt = "spline", ott = "spline")
+
+    cmds.setAttr("sheild_main_anim.panelD_deploy1", 0)
+    cmds.setAttr("panel_top_root_r_anim.rz", -90)
+    cmds.setDrivenKeyframe(["panel_top_root_r_anim.rz"], cd = "sheild_main_anim.panelD_deploy1", itt = "spline", ott = "spline")
+    cmds.setInfinity("panel_top_root_r_anim.rz", pri = "linear", poi = "linear")
+    cmds.setAttr("panel_top_root_r_anim.rz", keyable = False, channelBox = False)
+
+    #panelD 2
+    cmds.setAttr("sheild_main_anim.panelD_deploy2", 10)
+    cmds.setDrivenKeyframe(["panel_top_r_anim.sx"], cd = "sheild_main_anim.panelD_deploy2", itt = "spline", ott = "spline")
+
+    cmds.setAttr("sheild_main_anim.panelD_deploy2", 0)
+    cmds.setAttr("panel_top_r_anim.sx", 0.05*0.65)
+    cmds.setDrivenKeyframe(["panel_top_r_anim.sx"], cd = "sheild_main_anim.panelD_deploy2", itt = "spline", ott = "spline")
+    cmds.setInfinity("panel_top_r_anim.sx", pri = "linear", poi = "linear")
+    cmds.setAttr("panel_top_r_anim.sx", keyable = False, channelBox = False)
+
+shieldRig()
+
+# ------------------------------------------------------------------
+def legRig():
+    #Boot SDKS
+    for side in ["_l", "_r"]:
+        cmds.setDrivenKeyframe(["calf_armor"+side+"_anim_grp.rx", "toe_plate"+side+"_anim_grp.rz"], cd = "driver_foot"+side+".rz", itt = "linear", ott = "linear")
+        cmds.setAttr("driver_foot"+side+".rz", 30)
+        cmds.setAttr("calf_armor"+side+"_anim_grp.rx", -90)
+        cmds.setAttr("toe_plate"+side+"_anim_grp.rz", -90)
+        cmds.setDrivenKeyframe(["calf_armor"+side+"_anim_grp.rx", "toe_plate"+side+"_anim_grp.rz"], cd = "driver_foot"+side+".rz", itt = "linear", ott = "linear")
+        cmds.setAttr("driver_foot"+side+".rz", 0)
+
+        cmds.setDrivenKeyframe(["calf_armor"+side+"_anim_grp.rx", "calf_armor"+side+"_anim_grp.ry", "calf_armor"+side+"_anim_grp.rz"], cd = "driver_foot"+side+".rx", itt = "linear", ott = "linear")
+        cmds.setAttr("driver_foot"+side+".rx", 45)
+        cmds.setAttr("calf_armor"+side+"_anim_grp.rx", -11.7)
+        cmds.setAttr("calf_armor"+side+"_anim_grp.ry", -41)
+        cmds.setAttr("calf_armor"+side+"_anim_grp.rz", 24.9)
+        cmds.setDrivenKeyframe(["calf_armor"+side+"_anim_grp.rx", "calf_armor"+side+"_anim_grp.ry", "calf_armor"+side+"_anim_grp.rz"], cd = "driver_foot"+side+".rx", itt = "linear", ott = "linear")
+        cmds.setAttr("driver_foot"+side+".rx", -45)
+        cmds.setAttr("calf_armor"+side+"_anim_grp.rx", -11.7)
+        cmds.setAttr("calf_armor"+side+"_anim_grp.ry", 41)
+        cmds.setAttr("calf_armor"+side+"_anim_grp.rz", -24.9)
+        cmds.setDrivenKeyframe(["calf_armor"+side+"_anim_grp.rx", "calf_armor"+side+"_anim_grp.ry", "calf_armor"+side+"_anim_grp.rz"], cd = "driver_foot"+side+".rx", itt = "linear", ott = "linear")
+        cmds.setAttr("driver_foot"+side+".rx", 0)
+        
+        cmds.setDrivenKeyframe(["calf_armor"+side+"_anim_grp.rx", "calf_armor"+side+"_anim_grp.ry", "calf_armor"+side+"_anim_grp.rz"], cd = "driver_foot"+side+".ry", itt = "linear", ott = "linear")
+        cmds.setAttr("driver_foot"+side+".ry", 45)
+        cmds.setAttr("calf_armor"+side+"_anim_grp.rx", -17)
+        cmds.setAttr("calf_armor"+side+"_anim_grp.ry", 24)
+        cmds.setAttr("calf_armor"+side+"_anim_grp.rz", 0.3)
+        cmds.setDrivenKeyframe(["calf_armor"+side+"_anim_grp.rx", "calf_armor"+side+"_anim_grp.ry", "calf_armor"+side+"_anim_grp.rz"], cd = "driver_foot"+side+".ry", itt = "linear", ott = "linear")
+        cmds.setAttr("driver_foot"+side+".ry", -45)
+        cmds.setAttr("calf_armor"+side+"_anim_grp.rx", -17)
+        cmds.setAttr("calf_armor"+side+"_anim_grp.ry", -24)
+        cmds.setAttr("calf_armor"+side+"_anim_grp.rz", -0.3)
+        cmds.setDrivenKeyframe(["calf_armor"+side+"_anim_grp.rx", "calf_armor"+side+"_anim_grp.ry", "calf_armor"+side+"_anim_grp.rz"], cd = "driver_foot"+side+".ry", itt = "linear", ott = "linear")
+        cmds.setAttr("driver_foot"+side+".ry", 0)
+    
+legRig()
+
+# ------------------------------------------------------------------
+def tailFlapRig():
+    cmds.setDrivenKeyframe(["tail_pad_anim_grp.rx"], cd = "driver_spine_01.rz", itt = "linear", ott = "linear")
+    cmds.setAttr("driver_spine_01.rz", 90)
+    cmds.setAttr("tail_pad_anim_grp.rx", -90)
+    cmds.setDrivenKeyframe(["tail_pad_anim_grp.rx"], cd = "driver_spine_01.rz", itt = "linear", ott = "linear")
+    cmds.setAttr("driver_spine_01.rz", 15)
+    cmds.setAttr("tail_pad_anim_grp.rx", -38)
+    cmds.setDrivenKeyframe(["tail_pad_anim_grp.rx"], cd = "driver_spine_01.rz", itt = "spline", ott = "spline")
+    cmds.setAttr("driver_spine_01.rz", 0)
+
+    cmds.setDrivenKeyframe(["tail_pad_anim_grp.ry"], cd = "driver_spine_01.ry", itt = "linear", ott = "linear")
+    cmds.setAttr("driver_spine_01.ry", 90)
+    cmds.setAttr("tail_pad_anim_grp.ry", 25)
+    cmds.setDrivenKeyframe(["tail_pad_anim_grp.ry"], cd = "driver_spine_01.ry", itt = "linear", ott = "linear")
+    cmds.setAttr("driver_spine_01.ry", -90)
+    cmds.setAttr("tail_pad_anim_grp.ry", -25)
+    cmds.setDrivenKeyframe(["tail_pad_anim_grp.ry"], cd = "driver_spine_01.ry", itt = "linear", ott = "linear")
+    cmds.setAttr("driver_spine_01.ry", 0)
+
+    cmds.setDrivenKeyframe(["tail_pad_anim_grp.rz"], cd = "driver_spine_01.rx", itt = "linear", ott = "linear")
+    cmds.setAttr("driver_spine_01.rx", 90)
+    cmds.setAttr("tail_pad_anim_grp.rz", 20)
+    cmds.setDrivenKeyframe(["tail_pad_anim_grp.rz"], cd = "driver_spine_01.rx", itt = "linear", ott = "linear")
+    cmds.setAttr("driver_spine_01.rx", -90)
+    cmds.setAttr("tail_pad_anim_grp.rz", -20)
+    cmds.setDrivenKeyframe(["tail_pad_anim_grp.rz"], cd = "driver_spine_01.rx", itt = "linear", ott = "linear")
+    cmds.setAttr("driver_spine_01.rx", 0)
+tailFlapRig()
+
+# ------------------------------------------------------------------
+def scaleControls():
+    #Hide controls
+    '''controlsToHide = ["toe_l_anim", "toe_r_anim", "loarm_shield_top_l_anim", "loarm_shield_bot_l_anim", "loarm_shield_back_l_anim", "panels_l_anim", "panels_main_anim", "panels_r_anim", "panel_top_root_l_anim",
+    "panel_top_root_r_anim",  "panel_bot_root_l_anim", "panel_bot_root_r_anim", "panel_bot_l_anim", "panel_bot_r_anim", "panel_top_l_anim", "panel_top_r_anim", "toe_plate_l_anim", "toe_plate_r_anim"]
+    for ctrl in controlsToHide:
+        if cmds.objExists(ctrl):
+            grp = ctrl+"_grp"
+            cmds.setAttr(grp + ".visibility", 0)
+            cmds.setAttr(grp + ".visibility", lock = True)'''
+
+    for side in ["_l", "_r"]:
+
+        #scale shin controls
+        cmds.select("shinpad"+side+"_anim" + ".cv[0:32]")
+        cmds.scale(0, 2*0.65, 2*0.65, r=True)
+
+        #calf armor controls
+        cmds.select("calf_armor"+side+"_anim" + ".cv[0:32]")
+        cmds.scale(2*0.65, 2*0.65, 0, r=True)
+
+        #straps controls
+        cmds.select("arm_strap_top"+side+"_anim" + ".cv[0:32]")
+        cmds.scale(0.22, 0.22, 0.22, r=True)
+
+        cmds.select("arm_strap_bott"+side+"_anim" + ".cv[0:32]")
+        cmds.scale(0.22, 0.22, 0.22, r=True)
+
+        cmds.select("chest_strap"+side+"_anim" + ".cv[0:32]")
+        cmds.scale(0.22, 0.22, 0.22, r=True)
+
+        cmds.select("hip_strap"+side+"_anim" + ".cv[0:32]")
+        cmds.scale(0.22, 0.22, 0.22, r=True)
+
+        #wrist anims
+        cmds.select("ik_wrist"+side+"_anim" + ".cv[0:32]")
+        cmds.scale(2.8*0.65, 2.8*0.65, 2.8*0.65, r=True)
+
+        cmds.select("ik_elbow"+side+"_anim" + ".cv[0:32]")
+        cmds.scale(3*0.65, 3*0.65, 3*0.65, r=True)
+
+        #ik feet
+        cmds.select("ik_foot_anim"+side + ".cv[0:32]")
+        cmds.scale(1.7*0.65, 1.75*0.65, 1*0.65, r=True)
+    
+        #fingers
+        fingers = ["index_finger_fk_ctrl_1", "index_finger_fk_ctrl_2", "middle_finger_fk_ctrl_1", "middle_finger_fk_ctrl_2", "ring_finger_fk_ctrl_1", "ring_finger_fk_ctrl_2", "pinky_finger_fk_ctrl_1", "pinky_finger_fk_ctrl_2", "thumb_finger_fk_ctrl_1", "thumb_finger_fk_ctrl_2"]
+        for finger in fingers:
+            cmds.select(finger+side+".cv[0:32]")
+            cmds.scale(3.7*0.65, 3.7*0.65, 3.7*0.65, r=True)
+
+    # Other
+
+    cmds.select("tail_pad_anim" + ".cv[0:32]")
+    cmds.scale(0.3, 0.3, 0.3, r=True)
+    cmds.move(0, 3, 7.5, r=True)
+
+    cmds.select("flap_back_anim" + ".cv[0:32]")
+    cmds.scale(0.3, 0.3, 0.3, r=True)
+    cmds.move(0, 5, 0, r=True)
+
+    cmds.select("flap_front_anim" + ".cv[0:32]")
+    cmds.scale(0.3, 0.3, 0.3, r=True)
+    cmds.move(0, -7, 0, r=True)
+
+    #turbine controls
+    cmds.select("turbine_back_main_anim" + ".cv[0:32]")
+    cmds.scale(1, 0.07, 1, r=True)
+    cmds.move(0, 5.6, 0, r=True)
+
+    cmds.select("turbine_back_inner_anim" + ".cv[0:32]")
+    cmds.scale(0.75, 0.07, 0.75, r=True)
+    cmds.move(0, -0.6, 0, r=True)
+
+    #chest ik
+    cmds.select("chest_ik_anim" + ".cv[0:32]")
+    cmds.scale(1.5*0.65, 1.5*0.65, 1.5*0.65, r=True)
+    
+    #fk spine
+    cmds.select("spine_01_anim" + ".cv[0:32]")
+    cmds.scale(1.28*0.65, 1.28*0.65, 1.28*0.65, r=True)
+        
+    cmds.select("spine_02_anim" + ".cv[0:32]")
+    cmds.scale(1.32*0.65, 1.32*0.65, 1.32*0.65, r=True)
+
+    cmds.select("spine_03_anim" + ".cv[0:32]")
+    cmds.scale(1.54*0.65, 1.54*0.65, 1.54*0.65, r=True)    
+    
+    #neck
+    cmds.select("neck_01_fk_anim" + ".cv[0:32]")
+    cmds.scale(2.3*0.65, 2.3*0.65, 2.3*0.65, r=True)
+    
+    # football pads
+    cmds.select("football_pads_front_anim" + ".cv[0:32]")
+    cmds.scale(0.2, 0.2, 0.2, r=True)
+    cmds.move(0, 9, 16, r=True)
+
+    cmds.select("football_pads_back_anim" + ".cv[0:32]")
+    cmds.scale(0.2, 0.2, 0.2, r=True)
+    cmds.move(0, 9, 16, r=True)
+
+scaleControls()
+
+# ------------------------------------------------------------------
+def rigArms():
+    cmds.setAttr("Rig_Settings.lUpperarmTwistAmount", 1)
+    cmds.setAttr("Rig_Settings.lUpperarmTwist2Amount", 0.5)
+    cmds.setAttr("Rig_Settings.rUpperarmTwistAmount", 1)
+    cmds.setAttr("Rig_Settings.rUpperarmTwist2Amount", 0.5)    
+
+    cmds.setAttr("Rig_Settings.lArmMode", 0)
+    cmds.setAttr("Rig_Settings.lClavMode", 0)
+    cmds.setAttr("Rig_Settings.rArmMode", 0)
+    cmds.setAttr("Rig_Settings.rClavMode", 0)
+
+    for side in ["_l", "_r"]:
+        # Create SDK's to drive the shoulder blendshapes by thex upperarm joint
+        cmds.setAttr("fk_arm"+side+"_anim.ry", 62)
+        cmds.setDrivenKeyframe(["bicep_straps_geo_blendShape.bicep_straps_geo_shoulders_up"+side, "shoulders_geo_blendShape.shoulders_geo_shoulders_up"+side], cd="upperarm"+side+".rotateY", itt="flat", ott="flat")
+        cmds.setDrivenKeyframe(["bicep_straps_geo_blendShape.bicep_straps_geo_shoulders_down"+side, "shoulders_geo_blendShape.shoulders_geo_shoulders_down"+side], cd="upperarm"+side+".rotateY", itt="flat", ott="flat")
+
+        cmds.setAttr("fk_arm"+side+"_anim.ry", 0)
+        cmds.setAttr("bicep_straps_geo_blendShape.bicep_straps_geo_shoulders_up"+side, 1)
+        cmds.setAttr("shoulders_geo_blendShape.shoulders_geo_shoulders_up"+side, 1)
+        cmds.setDrivenKeyframe(["bicep_straps_geo_blendShape.bicep_straps_geo_shoulders_up"+side, "shoulders_geo_blendShape.shoulders_geo_shoulders_up"+side], cd="upperarm"+side+".rotateY", itt="flat", ott="flat")
+        cmds.setDrivenKeyframe(["bicep_straps_geo_blendShape.bicep_straps_geo_shoulders_down"+side, "shoulders_geo_blendShape.shoulders_geo_shoulders_down"+side], cd="upperarm"+side+".rotateY", itt="flat", ott="flat")
+        cmds.setAttr("bicep_straps_geo_blendShape.bicep_straps_geo_shoulders_up"+side, 0)
+        cmds.setAttr("shoulders_geo_blendShape.shoulders_geo_shoulders_up"+side, 0)
+
+        cmds.setAttr("fk_arm"+side+"_anim.ry", 82)
+        cmds.setAttr("bicep_straps_geo_blendShape.bicep_straps_geo_shoulders_down"+side, 1)
+        cmds.setAttr("shoulders_geo_blendShape.shoulders_geo_shoulders_down"+side, 1)
+        cmds.setDrivenKeyframe(["bicep_straps_geo_blendShape.bicep_straps_geo_shoulders_up"+side, "shoulders_geo_blendShape.shoulders_geo_shoulders_up"+side], cd="upperarm"+side+".rotateY", itt="flat", ott="flat")
+        cmds.setDrivenKeyframe(["bicep_straps_geo_blendShape.bicep_straps_geo_shoulders_down"+side, "shoulders_geo_blendShape.shoulders_geo_shoulders_down"+side], cd="upperarm"+side+".rotateY", itt="flat", ott="flat")
+        cmds.setAttr("bicep_straps_geo_blendShape.bicep_straps_geo_shoulders_down"+side, 0)
+        cmds.setAttr("shoulders_geo_blendShape.shoulders_geo_shoulders_down"+side, 0)
+
+        cmds.setAttr("fk_arm"+side+"_anim.ry", 0)
+
+        # Create SDK's to drive the bicep blendshape by the elbow joint
+        cmds.setAttr("fk_elbow"+side+"_anim.rz", 0)
+        cmds.setAttr("bicep_straps_geo_blendShape.bicep_straps_geo_elbows_bent"+side, 0)
+        cmds.setAttr("shoulders_geo_blendShape.shoulders_geo_elbows_bent"+side, 0)
+        cmds.setDrivenKeyframe(["bicep_straps_geo_blendShape.bicep_straps_geo_elbows_bent"+side, "shoulders_geo_blendShape.shoulders_geo_elbows_bent"+side], cd="lowerarm"+side+".rotateZ", itt="flat", ott="flat")
+        
+        cmds.setAttr("fk_elbow"+side+"_anim.rz", -100)
+        cmds.setAttr("bicep_straps_geo_blendShape.bicep_straps_geo_elbows_bent"+side, 1)
+        cmds.setAttr("shoulders_geo_blendShape.shoulders_geo_elbows_bent"+side, 1)
+        cmds.setDrivenKeyframe(["bicep_straps_geo_blendShape.bicep_straps_geo_elbows_bent"+side, "shoulders_geo_blendShape.shoulders_geo_elbows_bent"+side], cd="lowerarm"+side+".rotateZ", itt="flat", ott="flat")
+        
+        cmds.setAttr("fk_elbow"+side+"_anim.rz", 0)
+
+    cmds.setAttr("Rig_Settings.lArmMode", 1)
+    cmds.setAttr("Rig_Settings.lClavMode", 1)
+    cmds.setAttr("Rig_Settings.rArmMode", 1)
+    cmds.setAttr("Rig_Settings.rClavMode", 1)
+rigArms()
+
+# ------------------------------------------------------------------
+def createGameCam():
+    gameCam = cmds.camera(centerOfInterest=5, focalLength=18, lensSqueezeRatio=1, cameraScale=1, horizontalFilmAperture=1.41732, horizontalFilmOffset=0, verticalFilmAperture=0.94488, verticalFilmOffset=0, filmFit="Overscan", overscan=1.05, motionBlur=0, shutterAngle=144, nearClipPlane=1.0, farClipPlane=10000, orthographic=0, orthographicWidth=30, panZoomEnabled=0, horizontalPan=0, verticalPan=0, zoom=1, displayFilmGate=True)[0]
+    print gameCam
+    gameCam = cmds.rename(gameCam, "GameCam")
+    cmds.camera(gameCam, e=True, horizontalFilmAperture=1.6724376, focalLength=21.2399857011)
+    cmds.setAttr(gameCam+".displayGateMaskOpacity", 1)
+    cmds.setAttr(gameCam+".displayGateMaskColor", 0, 0, 0, type="double3")
+    cmds.setAttr(gameCam+".r", 90, 0, 180, type="double3")
+    
+    # Create the control
+    control = cmds.circle(name=gameCam+"_Ctrl", ch=True, o=True, r=15)[0]
+    cmds.makeIdentity(control, a=True, t=True, r=True, s=True)
+    # Create groups for the sdk and the constraints
+    const_grp = cmds.group(control, name=control+"_grp")
+    if cmds.objExists("rig_grp"):
+        cmds.parent(const_grp, "rig_grp")
+    # Parent Constrain the cam to the control
+    cmds.parentConstraint(control, gameCam, mo=True)
+    
+    cmds.connectAttr(control+".visibility", gameCam+".visibility", f=True)
+    cmds.setAttr(const_grp+".v", 0)
+    
+    cmds.setAttr(const_grp+".t", 0, 500, 200, type="double3")
+createGameCam()
+
+# ------------------------------------------------------------------
+def addCapsule(radius, halfHeight):
+    trueHalfHeight = (halfHeight*2)-(radius*2)
+    capsule = cmds.polyCylinder(name="collision_capsule", r=radius, h=trueHalfHeight, sx=18, sy=1, sz=5, ax=[0,0,1], rcp=True, cuv=3, ch=True)
+    cmds.setAttr(capsule[0]+".tz", halfHeight+2.2)
+    cmds.setAttr(capsule[0]+".rz", 90)
+
+    cmds.addAttr(capsule[0], ln="Radius", dv=radius, keyable=True)
+    cmds.addAttr(capsule[0], ln="HalfHeight", dv=halfHeight, keyable=True)
+
+    expr = capsule[1]+".radius = "+capsule[0]+".Radius;\n"+capsule[1]+".height = ("+capsule[0]+".HalfHeight*2) - ("+capsule[0]+".Radius*2);"
+    cmds.expression(s=expr, o=capsule[1], ae=1, uc="all")
+
+    cmds.setAttr(capsule[0]+".overrideEnabled", 1)
+    cmds.setAttr(capsule[0]+".overrideColor", 18)
+    cmds.setAttr(capsule[0]+".overrideShading", 0)
+
+    cmds.addAttr(capsule[0], ln="Shaded", dv=0, min=0, max=1, keyable=True)
+    cmds.connectAttr(capsule[0]+".Shaded", capsule[0]+".overrideShading")
+
+    cmds.parentConstraint("root", capsule[0], mo=True)
+    cmds.select(capsule[0], r=True)
+    capsuleDisplayLayer = cmds.createDisplayLayer(nr=True, name="Collision Capsule")
+    cmds.setAttr(capsuleDisplayLayer+".displayType", 2)
+    cmds.setAttr(capsuleDisplayLayer+".v", 0)
+    
+    for attr in ["tx", "ty", "tz", "rx", "ry", "rz", "sx", "sy", "sz"]:
+        cmds.setAttr(capsule[0]+"."+attr, lock=True, keyable=False, channelBox=False)
+
+    cmds.select(clear=True)
+addCapsule(64, 110)
+
+# ------------------------------------------------------------------
+# add ik bones Function that is called later after we open the export file.
+def addIKBones():
+    try:
+		# create the joints
+		cmds.select(clear=True)
+		ikFootRoot = cmds.joint(name = "ik_foot_root")
+		cmds.select(clear=True)
+		
+		cmds.select(clear=True)
+		ikFootLeft = cmds.joint(name = "ik_foot_l")
+		cmds.select(clear=True)
+		
+		cmds.select(clear=True)
+		ikFootRight = cmds.joint(name = "ik_foot_r")
+		cmds.select(clear=True)
+		
+		cmds.select(clear=True)
+		ikHandRoot = cmds.joint(name = "ik_hand_root")
+		cmds.select(clear=True)
+		
+		cmds.select(clear=True)
+		ikHandGun = cmds.joint(name = "ik_hand_gun")
+		cmds.select(clear=True)
+		
+		cmds.select(clear=True)
+		ikHandLeft = cmds.joint(name = "ik_hand_l")
+		cmds.select(clear=True)
+		
+		cmds.select(clear=True)
+		ikHandRight = cmds.joint(name = "ik_hand_r")
+		cmds.select(clear=True)
+		
+		# create hierarhcy
+		cmds.parent(ikFootRoot, "root")
+		cmds.parent(ikHandRoot, "root")
+		cmds.parent(ikFootLeft, ikFootRoot)
+		cmds.parent(ikFootRight, ikFootRoot)
+		cmds.parent(ikHandGun, ikHandRoot)
+		cmds.parent(ikHandLeft, ikHandGun)
+		cmds.parent(ikHandRight, ikHandGun)
+		
+		# constrain the joints
+		leftHandConstraint = cmds.parentConstraint("hand_l", ikHandLeft)[0]
+		rightHandGunConstraint = cmds.parentConstraint("hand_r", ikHandGun)[0]
+		rightHandConstraint = cmds.parentConstraint("hand_r", ikHandRight)[0]
+		leftFootConstraint = cmds.parentConstraint("foot_l", ikFootLeft)[0]
+		rightFootConstraint = cmds.parentConstraint("foot_r", ikFootRight)[0]
+	
+    except:
+		cmds.warning("Something went wrong")
+
+# ------------------------------------------------------------------
+# ADD the IK bones to the Export file.
+#save the AnimRig file before opening the export file.
+cmds.file(save = True, type = "mayaBinary", force = True)
+
+#ADD IK BONES
+#open export file
+sceneName = cmds.file(q = True, sceneName = True)
+exportFile = sceneName.partition("AnimRigs")[0] + "ExportFiles" +  sceneName.partition("AnimRigs")[2].partition(".mb")[0] + "_Export.mb"
+cmds.file(exportFile, open = True, force = True)
+
+#add ik bones
+addIKBones()
+
+#BEGIN JERNST CINE BLOCK
+#create the foot armor rig (left)
+cmds.createNode("joint", name = "l_foot_top_armor")
+cmds.createNode("joint", name = "l_foot_top_armor_end")
+
+cmds.xform("l_foot_top_armor", ws = True, rotation = [-90, 35.4419, -90], t = [27.52, -6.1206, 25.20597])
+cmds.xform("l_foot_top_armor_end", ws = True, rotation = [-90, 22.5059, -90], t = [27.52, -16.97217, 17.4694])
+cmds.parent("l_foot_top_armor_end", "l_foot_top_armor")
+cmds.parent("l_foot_top_armor", "calf_l")
+cmds.makeIdentity("l_foot_top_armor", t = 0 , r = 1, s = 0, apply = True)
+
+lFootArmorIK = cmds.ikHandle(sj = "l_foot_top_armor", ee = "l_foot_top_armor_end", sol = "ikSCsolver")
+
+
+#create the toe armor rig (left)
+cmds.createNode("joint", name = "l_foot_toe_armor")
+cmds.createNode("joint", name = "l_foot_toe_armor_end")
+
+
+cmds.xform("l_foot_toe_armor", ws = True, rotation = [-90, 22.5059, -90], t = [27.52, -12.9240, 19.26320])
+cmds.xform("l_foot_toe_armor_end", ws = True, rotation = [-90, 22.5059, -90], t = [27.52, -22.5982, 15.25484])
+cmds.parent("l_foot_toe_armor_end", "l_foot_toe_armor")
+cmds.parent("l_foot_toe_armor", "foot_l")
+cmds.makeIdentity("l_foot_toe_armor", t = 0 , r = 1, s = 0, apply = True)
+
+lToeArmorIK = cmds.ikHandle(sj = "l_foot_toe_armor", ee = "l_foot_toe_armor_end", sol = "ikSCsolver")
+
+#parent IK handles (left)
+cmds.parent(lFootArmorIK[0], "l_foot_toe_armor")
+cmds.parent(lToeArmorIK[0], "ball_l")
+
+#create the foot armor rig (right)
+cmds.createNode("joint", name = "r_foot_top_armor")
+cmds.createNode("joint", name = "r_foot_top_armor_end")
+
+cmds.xform("r_foot_top_armor", ws = True, rotation = [-90, 35.4419, -90], t = [-27.52, -6.1206, 25.20597])
+cmds.xform("r_foot_top_armor_end", ws = True, rotation = [-90, 22.5059, -90], t = [-27.52, -16.97217, 17.4694])
+cmds.parent("r_foot_top_armor_end", "r_foot_top_armor")
+cmds.parent("r_foot_top_armor", "calf_r")
+cmds.makeIdentity("r_foot_top_armor", t = 0 , r = 1, s = 0, apply = True)
+
+lFootArmorIK = cmds.ikHandle(sj = "r_foot_top_armor", ee = "r_foot_top_armor_end", sol = "ikSCsolver")
+
+
+#create the toe armor rig (right)
+cmds.createNode("joint", name = "r_foot_toe_armor")
+cmds.createNode("joint", name = "r_foot_toe_armor_end")
+
+
+cmds.xform("r_foot_toe_armor", ws = True, rotation = [-90, 22.5059, -90], t = [-27.52, -12.9240, 19.26320])
+cmds.xform("r_foot_toe_armor_end", ws = True, rotation = [-90, 22.5059, -90], t = [-27.52, -22.5982, 15.25484])
+cmds.parent("r_foot_toe_armor_end", "r_foot_toe_armor")
+cmds.parent("r_foot_toe_armor", "foot_r")
+cmds.makeIdentity("r_foot_toe_armor", t = 0 , r = 1, s = 0, apply = True)
+
+lToeArmorIK = cmds.ikHandle(sj = "r_foot_toe_armor", ee = "r_foot_toe_armor_end", sol = "ikSCsolver")
+
+#parent IK handles (right)
+cmds.parent(lFootArmorIK[0], "r_foot_toe_armor")
+cmds.parent(lToeArmorIK[0], "ball_r")
+
+
+#Add joints to skinClusters
+skinClusters = cmds.ls(type = "skinCluster")
+
+for cluster in skinClusters:
+    geo = cmds.skinCluster(cluster, q = True, geometry = True)[0]
+    if geo == "boot_cover_geoShape":
+        cmds.skinCluster(cluster, edit = True, dr = 4, lw = True, wt = 0, ai = ["l_foot_top_armor", "l_foot_top_armor_end", "r_foot_top_armor", "r_foot_top_armor_end"])
+        
+        cmds.skinPercent(cluster, "boot_cover_geo.vtx[0:123]", tv = ["l_foot_top_armor", 1])
+        cmds.skinPercent(cluster, "boot_cover_geo.vtx[124:247]", tv = ["r_foot_top_armor", 1])
+        
+        
+    if geo == "toe_plates_geoShape":
+        cmds.skinCluster(cluster, edit = True, dr = 4, lw = True, wt = 0, ai = ["l_foot_toe_armor", "l_foot_toe_armor_end", "r_foot_toe_armor", "r_foot_toe_armor_end"])
+        
+        cmds.skinPercent(cluster, "toe_plates_geo.vtx[0:185]", tv = ["l_foot_toe_armor", 1])
+        cmds.skinPercent(cluster, "toe_plates_geo.vtx[186:371]", tv = ["r_foot_toe_armor", 1])
+# END JERNST CINE BLOCK
+
+#save the export file.
+cmds.file(save = True, type = "mayaBinary", force = True)
+
+
+
+#########################################
+# RUN THIS COMMAND DURING THE RIG POSE CREATION STEP TO GET THE SHIELD TO THE CORRECT LOCATION.  LEAVE IT COMMENTED OUT SO THAT IT DOESNT RUN DURING RIGGING.
+'''
+#cmds.xform(q=True, ws=True, rp=True)
+#cmds.xform(q=True, ro=True)
+cmds.xform("sheild_main_mover", ws=True, t=[-0.018884450479987436, 59.40691643000001, 272.98635249999995])
+cmds.xform("sheild_main_mover", ro=[-14.885476763734184, 61.20693988721698, -15.277192149145698])
+'''
+#########################################
+
+
+
+
+
+
+
+
+
+
+
+
+
+
+
+
+
+
+
+
+
+
+
+
+
+
+
+
+
+
+
+
+
+
+
+
+
+
+
+
+
+
+
+
+
+
+
+
diff --git a/Engine/Extras/Maya_AnimationRiggingTools/ARTv1/General/ART/Projects/Orion_Cine/Scripts/postScripts/TwinBlast_Post_Script.py b/Engine/Extras/Maya_AnimationRiggingTools/ARTv1/General/ART/Projects/Orion_Cine/Scripts/postScripts/TwinBlast_Post_Script.py
new file mode 100644
--- /dev/null
+++ b/Engine/Extras/Maya_AnimationRiggingTools/ARTv1/General/ART/Projects/Orion_Cine/Scripts/postScripts/TwinBlast_Post_Script.py
@@ -0,0 +1,218 @@
+import maya.cmds as cmds
+# NOTE: If you mirror his movers, lock the weapon controls first so that they dont mirror because they are set up to not.  And putting them back is a pain.
+# NOTE: Use the "GUN_locator___USE_FOR_RIG_POSE" in the scene during the Create Rig Pose phase to put the weapon movers adn the grenade movers back to the correct location.
+
+# ------------------------------------------------------------------
+# scale all of the controls to be more visible with the mesh.
+def scaleControls():
+    sides = ["_l", "_r"]
+    for side in sides:
+        # feet
+        cmds.select(["ik_foot_anim"+side+".cv[0:17]"])
+        cmds.scale(1, 1, 1, r = True, ocp = False)
+        cmds.move(0, 0, -4.78  , r = True)
+
+        cmds.select(["q_weapon"+side+"_anim.cv[0:32]"])
+        cmds.scale(0.4, 0.4, 0.4, r = True, ocp = False)
+        if side == "_l":
+            cmds.move(0, 20, 0, r = True)
+        else:
+            cmds.move(0, 20, 0, r = True)            
+
+        cmds.select(["ult_weapon"+side+"_anim.cv[0:32]"])
+        cmds.scale(0.5, 0.5, 0.5, r = True, ocp = False)
+        if side == "_l":
+            cmds.move(0, 25, 0, r = True)
+        else:
+            cmds.move(0, 25, 0, r = True)            
+
+        cmds.select(["forearm_gears"+side+"_anim.cv[0:32]"])
+        cmds.scale(0, 1, 1, r = True, ocp = False)
+
+    # Other
+    objects = ["weapon_forestock_bott_l_anim", "weapon_forestock_bott_r_anim", "weapon_forestock_top_l_anim", "weapon_forestock_top_r_anim", "weapon_l_anim", "weapon_r_anim", "weapon_stock_l_anim", "weapon_stock_r_anim", "low_lowarm_plate_A_l_anim", "low_lowarm_plate_B_l_anim", "low_lowarm_plate_C_l_anim", "low_lowarm_plate_D_l_anim", "up_lowarm_plate_A_l_anim", "up_lowarm_plate_B_l_anim", "up_lowarm_plate_C_l_anim", "up_lowarm_plate_D_l_anim", "grenade_anim", "grenade_front_anim", "grenade_fin_01_anim", "grenade_fin_02_anim", "grenade_fin_03_anim", "grenade_fin_04_anim", "grenade_fin_05_anim", "grenade_fin_06_anim", "collar_in_l_anim", "collar_out_l_anim", "collar_out_flap_bott_l_anim", "collar_out_flap_top_l_anim", "low_lowarm_plate_A_r_anim", "low_lowarm_plate_B_r_anim", "low_lowarm_plate_C_r_anim", "low_lowarm_plate_D_r_anim", "up_lowarm_plate_A_r_anim", "up_lowarm_plate_B_r_anim", "up_lowarm_plate_C_r_anim", "up_lowarm_plate_D_r_anim", "collar_in_r_anim", "collar_out_r_anim", "collar_out_flap_bott_r_anim", "collar_out_flap_top_r_anim", "dogtag_anim", "collar_back_anim", "collar_l_anim", "collar_r_anim"]
+    for one in objects:
+        cmds.select([one+".cv[*]"])
+        cmds.scale(0.227057, 0.227057, 0.227057, r = True, ocp = True)
+        cmds.move(0, 0, 0, r = True, os=True, wd=True)
+
+
+    cmds.select(None)
+
+scaleControls()
+
+# ------------------------------------------------------------------
+def rigWeapons():
+    cmds.move(0, 4.4, 24.558962, "ult_weapon_l_anim", r=True, os=True, wd=True)
+    cmds.scale(0.05, 0.05, 0.05, "ult_weapon_l_anim", r=True)
+
+    cmds.move(0, -4.4, -24.558962, "ult_weapon_r_anim", r=True, os=True, wd=True)
+    cmds.scale(0.05, 0.05, 0.05, "ult_weapon_r_anim", r=True)
+
+    cmds.move(0, 4.4, 24.558962, "q_weapon_l_anim", r=True, os=True, wd=True)
+    cmds.scale(0.05, 0.05, 0.05, "q_weapon_l_anim", r=True)
+
+    cmds.move(0, -4.4, -24.558962, "q_weapon_r_anim", r=True, os=True, wd=True)
+    cmds.scale(0.05, 0.05, 0.05, "q_weapon_r_anim", r=True)
+
+    cmds.xform("weapon_l_anim_grp", t=[82.96438708999413, 10.607007776552306, 147.1500662787879], ro=[37.910151258451656, -65.27865699485243, 49.67785824895247])
+    cmds.xform("weapon_r_anim_grp", t=[-82.96438708999413, 10.607007776552306, 147.1500662787879], ro=[37.910151258451656, 65.27865699485243, -49.67785824895247])
+    cmds.xform("grenade_front_anim", t=[0, 0.4509, 0], ro=[0, 0, 0])
+    cmds.xform("grenade_fin_01_anim", t=[0, 0, 0], ro=[-180, 0, 0])
+    cmds.xform("grenade_fin_02_anim", t=[0, 0, 0], ro=[-180, 0, 0])
+    cmds.xform("grenade_fin_03_anim", t=[0, 0, 0], ro=[-180, 0, 0])
+    cmds.xform("grenade_fin_04_anim", t=[0, 0, 0], ro=[-180, 0, 0])
+    cmds.xform("grenade_fin_05_anim", t=[0, 0, 0], ro=[-180, 0, 0])
+    cmds.xform("grenade_fin_06_anim", t=[0, 0, 0], ro=[-180, 0, 0])
+    cmds.xform("grenade_anim_grp", t=[-59.36194438801096, 12.504397019491371, 150.4395268893561], ro=[0.0, 0.0, -90])
+    
+    cmds.setAttr("Rig_Settings.lUpperarmTwistAmount", 0.9)
+    cmds.setAttr("Rig_Settings.lUpperarmTwist2Amount", 0.3)
+    cmds.setAttr("Rig_Settings.rUpperarmTwistAmount", 0.9)
+    cmds.setAttr("Rig_Settings.rUpperarmTwist2Amount", 0.3)
+
+    cmds.setAttr("Rig_Settings.lForearmTwistAmount", 1)
+    cmds.setAttr("Rig_Settings.rForearmTwistAmount", 1)
+    
+rigWeapons()
+
+# ------------------------------------------------------------------
+def createGameCam():
+    gameCam = cmds.camera(centerOfInterest=5, focalLength=18, lensSqueezeRatio=1, cameraScale=1, horizontalFilmAperture=1.41732, horizontalFilmOffset=0, verticalFilmAperture=0.94488, verticalFilmOffset=0, filmFit="Overscan", overscan=1.05, motionBlur=0, shutterAngle=144, nearClipPlane=1.0, farClipPlane=10000, orthographic=0, orthographicWidth=30, panZoomEnabled=0, horizontalPan=0, verticalPan=0, zoom=1, displayFilmGate=True)[0]
+    print gameCam
+    gameCam = cmds.rename(gameCam, "GameCam")
+    cmds.camera(gameCam, e=True, horizontalFilmAperture=1.6724376, focalLength=21.2399857011)
+    cmds.setAttr(gameCam+".displayGateMaskOpacity", 1)
+    cmds.setAttr(gameCam+".displayGateMaskColor", 0, 0, 0, type="double3")
+    cmds.setAttr(gameCam+".r", 90, 0, 180, type="double3")
+    
+    # Create the control
+    control = cmds.circle(name=gameCam+"_Ctrl", ch=True, o=True, r=15)[0]
+    cmds.makeIdentity(control, a=True, t=True, r=True, s=True)
+    # Create groups for the sdk and the constraints
+    const_grp = cmds.group(control, name=control+"_grp")
+    if cmds.objExists("rig_grp"):
+        cmds.parent(const_grp, "rig_grp")
+    # Parent Constrain the cam to the control
+    cmds.parentConstraint(control, gameCam, mo=True)
+    
+    cmds.connectAttr(control+".visibility", gameCam+".visibility", f=True)
+    cmds.setAttr(const_grp+".v", 0)
+    
+    cmds.setAttr(const_grp+".t", 0, 500, 200, type="double3")
+createGameCam()
+# ------------------------------------------------------------------
+
+##################################################################################
+##################################################################################
+
+def addCapsule(radius, halfHeight):
+    trueHalfHeight = (halfHeight*2)-(radius*2)
+    capsule = cmds.polyCylinder(name="collision_capsule", r=radius, h=trueHalfHeight, sx=18, sy=1, sz=5, ax=[0,0,1], rcp=True, cuv=3, ch=True)
+    cmds.setAttr(capsule[0]+".tz", halfHeight+2.2)
+    cmds.setAttr(capsule[0]+".rz", 90)
+
+    cmds.addAttr(capsule[0], ln="Radius", dv=radius, keyable=True)
+    cmds.addAttr(capsule[0], ln="HalfHeight", dv=halfHeight, keyable=True)
+
+    expr = capsule[1]+".radius = "+capsule[0]+".Radius;\n"+capsule[1]+".height = ("+capsule[0]+".HalfHeight*2) - ("+capsule[0]+".Radius*2);"
+    cmds.expression(s=expr, o=capsule[1], ae=1, uc="all")
+
+    cmds.setAttr(capsule[0]+".overrideEnabled", 1)
+    cmds.setAttr(capsule[0]+".overrideColor", 18)
+    cmds.setAttr(capsule[0]+".overrideShading", 0)
+
+    cmds.addAttr(capsule[0], ln="Shaded", dv=0, min=0, max=1, keyable=True)
+    cmds.connectAttr(capsule[0]+".Shaded", capsule[0]+".overrideShading")
+
+    cmds.parentConstraint("root", capsule[0], mo=True)
+    cmds.select(capsule[0], r=True)
+    capsuleDisplayLayer = cmds.createDisplayLayer(nr=True, name="Collision Capsule")
+    cmds.setAttr(capsuleDisplayLayer+".displayType", 2)
+    cmds.setAttr(capsuleDisplayLayer+".v", 0)
+    
+    for attr in ["tx", "ty", "tz", "rx", "ry", "rz", "sx", "sy", "sz"]:
+        cmds.setAttr(capsule[0]+"."+attr, lock=True, keyable=False, channelBox=False)
+
+    cmds.select(clear=True)
+addCapsule(42, 86)
+
+# ------------------------------------------------------------------
+# add ik bones Function that is called later after we open the export file.
+def addIKBones():
+    try:
+		# create the joints
+		cmds.select(clear = True)
+		ikFootRoot = cmds.joint(name = "ik_foot_root")
+		cmds.select(clear = True)
+		
+		cmds.select(clear = True)
+		ikFootLeft = cmds.joint(name = "ik_foot_l")
+		cmds.select(clear = True)
+		
+		cmds.select(clear = True)
+		ikFootRight = cmds.joint(name = "ik_foot_r")
+		cmds.select(clear = True)
+		
+		cmds.select(clear = True)
+		ikHandRoot = cmds.joint(name = "ik_hand_root")
+		cmds.select(clear = True)
+		
+		cmds.select(clear = True)
+		ikHandGun = cmds.joint(name = "ik_hand_gun")
+		cmds.select(clear = True)
+		
+		cmds.select(clear = True)
+		ikHandLeft = cmds.joint(name = "ik_hand_l")
+		cmds.select(clear = True)
+		
+		cmds.select(clear = True)
+		ikHandRight = cmds.joint(name = "ik_hand_r")
+		cmds.select(clear = True)
+		
+		# create hierarhcy
+		cmds.parent(ikFootRoot, "root")
+		cmds.parent(ikHandRoot, "root")
+		cmds.parent(ikFootLeft, ikFootRoot)
+		cmds.parent(ikFootRight, ikFootRoot)
+		cmds.parent(ikHandGun, ikHandRoot)
+		cmds.parent(ikHandLeft, ikHandGun)
+		cmds.parent(ikHandRight, ikHandGun)
+		
+		# constrain the joints
+		leftHandConstraint = cmds.parentConstraint("hand_l", ikHandLeft)[0]
+		rightHandGunConstraint = cmds.parentConstraint("hand_r", ikHandGun)[0]
+		rightHandConstraint = cmds.parentConstraint("hand_r", ikHandRight)[0]
+		leftFootConstraint = cmds.parentConstraint("foot_l", ikFootLeft)[0]
+		rightFootConstraint = cmds.parentConstraint("foot_r", ikFootRight)[0]
+	
+    except:
+		cmds.warning("Something went wrong")
+
+# ------------------------------------------------------------------
+# ADD the IK bones to the Export file.
+#save the AnimRig file before opening the export file.
+cmds.file(save = True, type = "mayaBinary", force = True)
+
+#ADD IK BONES
+#open export file
+sceneName = cmds.file(q = True, sceneName = True)
+exportFile = sceneName.partition("AnimRigs")[0] + "ExportFiles" +  sceneName.partition("AnimRigs")[2].partition(".mb")[0] + "_Export.mb"
+cmds.file(exportFile, open = True, force = True)
+
+#add ik bones
+addIKBones()
+
+# ------------------------------------------------------------------
+def addCustomAttrs():
+    attrs = cmds.listAttr("root")
+    exists = False
+    for one in attrs:
+        if one.find("DistanceCurve") != -1:
+            exists = True
+    if exists == False:
+        cmds.addAttr("root", ln="DistanceCurve", keyable=True)
+addCustomAttrs()
+
+#save the export file.
+cmds.file(save = True, type = "mayaBinary", force = True)
\ No newline at end of file
diff --git a/Engine/Extras/Maya_AnimationRiggingTools/ARTv1/General/ART/Projects/Rocket/Scripts/Mannequin_Post_Script.py b/Engine/Extras/Maya_AnimationRiggingTools/ARTv1/General/ART/Projects/Rocket/Scripts/Mannequin_Post_Script.py
new file mode 100644
--- /dev/null
+++ b/Engine/Extras/Maya_AnimationRiggingTools/ARTv1/General/ART/Projects/Rocket/Scripts/Mannequin_Post_Script.py
@@ -0,0 +1,65 @@
+def addIKBones():
+    
+    try:
+	#create the joints
+	cmds.select(clear = True)
+	ikFootRoot = cmds.joint(name = "ik_foot_root")
+	cmds.select(clear = True)
+	
+	cmds.select(clear = True)
+	ikFootLeft = cmds.joint(name = "ik_foot_l")
+	cmds.select(clear = True)
+	
+	cmds.select(clear = True)
+	ikFootRight = cmds.joint(name = "ik_foot_r")
+	cmds.select(clear = True)
+	
+	cmds.select(clear = True)
+	ikHandRoot = cmds.joint(name = "ik_hand_root")
+	cmds.select(clear = True)
+	
+	cmds.select(clear = True)
+	ikHandGun = cmds.joint(name = "ik_hand_gun")
+	cmds.select(clear = True)
+	
+	cmds.select(clear = True)
+	ikHandLeft = cmds.joint(name = "ik_hand_l")
+	cmds.select(clear = True)
+	
+	cmds.select(clear = True)
+	ikHandRight = cmds.joint(name = "ik_hand_r")
+	cmds.select(clear = True)
+	
+	
+	#create hierarhcy
+	cmds.parent(ikFootRoot, "root")
+	cmds.parent(ikHandRoot, "root")
+	cmds.parent(ikFootLeft, ikFootRoot)
+	cmds.parent(ikFootRight, ikFootRoot)
+	cmds.parent(ikHandGun, ikHandRoot)
+	cmds.parent(ikHandLeft, ikHandGun)
+	cmds.parent(ikHandRight, ikHandGun)
+	
+	#constrain the joints
+	leftHandConstraint = cmds.parentConstraint("hand_l", ikHandLeft)[0]
+	rightHandGunConstraint = cmds.parentConstraint("hand_r", ikHandGun)[0]
+	rightHandConstraint = cmds.parentConstraint("hand_r", ikHandRight)[0]
+	leftFootConstraint = cmds.parentConstraint("foot_l", ikFootLeft)[0]
+	rightFootConstraint = cmds.parentConstraint("foot_r", ikFootRight)[0]
+	
+    except:
+	cmds.warning("Something went wrong")
+
+
+
+#open export file
+sceneName = cmds.file(q = True, sceneName = True)
+exportFile = sceneName.partition("AnimRigs")[0] + "ExportFiles" +  sceneName.partition("AnimRigs")[2].partition(".mb")[0] + "_Export.mb"
+cmds.file(exportFile, open = True, force = True)
+
+#add ik bones
+addIKBones()
+
+#save
+cmds.file(save = True, type = "mayaBinary", force = True)
+
diff --git a/Engine/Extras/Maya_AnimationRiggingTools/ARTv1/General/ART/Projects/Unreal_Tournament/Scripts/ut4_base_postScript.py b/Engine/Extras/Maya_AnimationRiggingTools/ARTv1/General/ART/Projects/Unreal_Tournament/Scripts/ut4_base_postScript.py
new file mode 100644
--- /dev/null
+++ b/Engine/Extras/Maya_AnimationRiggingTools/ARTv1/General/ART/Projects/Unreal_Tournament/Scripts/ut4_base_postScript.py
@@ -0,0 +1,90 @@
+def addCapsule():
+    capsule = cmds.polyCylinder(name = "collision_capsule", ch=1, r=46, h=95, cuv=3, sc=3, rcp=1)
+    cmds.setAttr(capsule[1] + ".subdivisionsAxis", 8)
+    cmds.setAttr(capsule[0] + ".rx", -90)
+    cmds.setAttr(capsule[0] + ".tz", 1)
+    
+    cmds.move(0, 0, 93.5, capsule, wd=1)
+    cmds.xform(capsule[0], ws = True, piv = [0,0,0])
+    cmds.makeIdentity(capsule[0], t = 1, r = 1, s = 1, apply = True)
+    
+    cmds.setAttr(capsule[0] + ".overrideEnabled", 1)
+    cmds.setAttr(capsule[0] + ".overrideColor", 18)
+    cmds.setAttr(capsule[0] + ".overrideShading", 0)
+    
+    cmds.setAttr(capsule[0] + ".sx", lock = True, keyable = False, channelBox = False)
+    cmds.setAttr(capsule[0] + ".sy", lock = True, keyable = False, channelBox = False)
+    cmds.setAttr(capsule[0] + ".sz", lock = True, keyable = False, channelBox = False)
+    
+    cmds.addAttr(capsule[0], ln = "Shaded", dv = 0, min = 0, max = 1, keyable = True)
+    cmds.connectAttr(capsule[0] + ".Shaded", capsule[0] + ".overrideShading")
+    
+    cmds.parentConstraint("root", capsule[0], mo = True)
+    cmds.select("collision_capsule", r=True)
+    cmds.createDisplayLayer(nr=True, name="Collision Capsule")
+    cmds.select(clear=True)
+
+def ikSetup():
+    cmds.select(clear = True)
+    ikFootRoot = cmds.joint(name = "ik_foot_root")
+    cmds.select(clear = True)
+    ikFootLeft = cmds.joint(name = "ik_foot_l")
+    cmds.select(clear = True)
+    ikFootRight = cmds.joint(name = "ik_foot_r")
+    cmds.select(clear = True)
+    ikHandRoot = cmds.joint(name = "ik_hand_root")
+    cmds.select(clear = True)
+    ikHandGun = cmds.joint(name = "ik_hand_gun")
+    cmds.select(clear = True)
+    ikHandLeft = cmds.joint(name = "ik_hand_l")
+    cmds.select(clear = True)
+    ikHandRight = cmds.joint(name = "ik_hand_r")
+
+    #create hierarhcy
+    cmds.parent(ikFootRoot, "root")
+    cmds.parent(ikHandRoot, "root")
+    cmds.parent(ikFootLeft, ikFootRoot)
+    cmds.parent(ikFootRight, ikFootRoot)
+    cmds.parent(ikHandGun, ikHandRoot)
+    cmds.parent(ikHandLeft, ikHandGun)
+    cmds.parent(ikHandRight, ikHandGun)
+
+    leftHandConstraint = cmds.parentConstraint("hand_l", 'ik_hand_l')[0]
+    rightHandGunConstraint = cmds.parentConstraint("hand_r", 'ik_hand_gun')[0]
+    rightHandConstraint = cmds.parentConstraint("hand_r", 'ik_hand_r')[0]
+    leftFootConstraint = cmds.parentConstraint("foot_l", 'ik_foot_l')[0]
+    rightFootConstraint = cmds.parentConstraint("foot_r", 'ik_foot_r')[0]
+
+    cmds.select(clear = True)
+    ikFootRoot = cmds.joint(name = "ikalt_foot_root")
+    cmds.select(clear = True)
+    ikFootLeft = cmds.joint(name = "ikalt_foot_l")
+    cmds.select(clear = True)
+    ikFootRight = cmds.joint(name = "ikalt_foot_r")
+    cmds.select(clear = True)
+    ikHandRoot = cmds.joint(name = "ikalt_hand_root")
+    cmds.select(clear = True)
+    ikHandLeft = cmds.joint(name = "ikalt_hand_l")
+    cmds.select(clear = True)
+    ikHandRight = cmds.joint(name = "ikalt_hand_r")
+
+    #create hierarhcy
+    cmds.parent(ikFootRoot, "root")
+    cmds.parent(ikHandRoot, "root")
+    cmds.parent(ikFootLeft, ikFootRoot)
+    cmds.parent(ikFootRight, ikFootRoot)
+    cmds.parent(ikHandLeft, ikHandRoot)
+    cmds.parent(ikHandRight, ikHandRoot)
+
+    leftHandConstraint = cmds.parentConstraint("hand_l", 'ik_hand_l')[0]
+    rightHandConstraint = cmds.parentConstraint("hand_r", 'ik_hand_r')[0]
+    leftFootConstraint = cmds.parentConstraint("foot_l", 'ik_foot_l')[0]
+    rightFootConstraint = cmds.parentConstraint("foot_r", 'ik_foot_r')[0]
+    
+    cmds.parentConstraint("hand_l", 'ikalt_hand_l')
+    cmds.parentConstraint("hand_r", 'ikalt_hand_r')
+    cmds.parentConstraint("foot_l", 'ikalt_foot_l')
+    cmds.parentConstraint("foot_r", 'ikalt_foot_r')
+
+ikSetup()
+addCapsule()
\ No newline at end of file
diff --git a/Engine/Extras/Maya_AnimationRiggingTools/ARTv1/General/Scripts/ART_Settings.py b/Engine/Extras/Maya_AnimationRiggingTools/ARTv1/General/Scripts/ART_Settings.py
new file mode 100644
--- /dev/null
+++ b/Engine/Extras/Maya_AnimationRiggingTools/ARTv1/General/Scripts/ART_Settings.py
@@ -0,0 +1,257 @@
+import maya.cmds as cmds
+import maya.mel as mel
+from functools import partial
+import os, cPickle
+
+class ART_Settings():
+
+    def __init__(self):
+
+        #get access to our maya tools
+        toolsPath = cmds.internalVar(usd = True) + "mayaTools.txt"
+        if os.path.exists(toolsPath):
+
+            f = open(toolsPath, 'r')
+            self.mayaToolsDir = f.readline()
+            f.close()
+
+
+        #check to see if window exists, if so, delete
+        if cmds.window("AnimationRiggingTool_SettingsUI", exists = True):
+            cmds.deleteUI("AnimationRiggingTool_SettingsUI")
+
+        self.path = self.mayaToolsDir
+        self.widgets = {}
+
+        #create the window
+        self.widgets["window"] = cmds.window("AnimationRiggingTool_SettingsUI", w = 400, h = 300, titleBarMenu = False, title = "Settings", sizeable = False)
+
+
+        #main layout
+        self.widgets["main"] = cmds.formLayout(w = 400, h = 300)
+
+        #background image
+        image = self.mayaToolsDir + "/General/Icons/ART/settings.jpg"
+        self.widgets["background"] = cmds.image(w = 400, h = 300, parent = self.widgets["main"], image = image)
+
+
+        #close button
+        self.widgets["closeButton"] = cmds.symbolButton(w = 26, h = 20, parent = self.widgets["main"], image = self.mayaToolsDir + "/General/Icons/ART/xbutton.bmp", c = self.close)
+        cmds.formLayout(self.widgets["main"], edit = True, af = [(self.widgets["closeButton"], "top", 2), (self.widgets["closeButton"], "right", 2)])
+
+        #tools path
+        self.widgets["toolsPathLabel"] = cmds.text(h = 40, w = 100, label = "Tools Path: ", font = "boldLabelFont", parent = self.widgets["main"])
+        self.widgets["toolsPathField"] = cmds.textField(h = 40, w = 225, text = "", parent = self.widgets["main"], editable = False)
+        self.widgets["toolsPathBrowseButton"] = cmds.symbolButton(w = 40, h = 40, parent = self.widgets["main"], image = self.mayaToolsDir + "/General/Icons/ART/browse.bmp", c = self.browseNewToolsPath)
+
+        cmds.formLayout(self.widgets["main"], edit = True, af = [(self.widgets["toolsPathLabel"], "top", 40), (self.widgets["toolsPathLabel"], "left", 10)])
+        cmds.formLayout(self.widgets["main"], edit = True, af = [(self.widgets["toolsPathField"], "top", 40), (self.widgets["toolsPathField"], "right", 50)])
+        cmds.formLayout(self.widgets["main"], edit = True, af = [(self.widgets["toolsPathBrowseButton"], "top", 40), (self.widgets["toolsPathBrowseButton"], "right", 5)])
+
+        #Source control settings
+        divider = cmds.separator(w = 400, h = 10, style = "out")
+        self.widgets["useSourceControl"] = cmds.checkBox( h = 40, label = "Use Source Control", v = False, parent = self.widgets["main"], onc = partial(self.useSourceControl, True), ofc = partial(self.useSourceControl, False))
+
+        cmds.formLayout(self.widgets["main"], edit = True, af = [(divider, "top", 95), (divider, "left", 0)])
+        cmds.formLayout(self.widgets["main"], edit = True, af = [(self.widgets["useSourceControl"], "top", 105), (self.widgets["useSourceControl"], "left", 10)])
+
+        #source control test button
+        self.widgets["testConnection"] = cmds.button(w = 225, h = 40, label = "Test Perforce Connection", parent = self.widgets["main"], enable = False, c = self.testConnection)
+        cmds.formLayout(self.widgets["main"], edit = True, af = [(self.widgets["testConnection"], "top", 105), (self.widgets["testConnection"], "right", 25)]) 
+
+
+        #favorite project option
+        favoriteProjLabel = cmds.text(label = "Set Favorite Project: ", parent = self.widgets["main"])
+        cmds.formLayout(self.widgets["main"], edit = True, af = [(favoriteProjLabel, "top", 170), (favoriteProjLabel, "left", 10)]) 
+        self.widgets["favoriteProject_OM"] = cmds.optionMenu(w = 225, h = 40, label = "", parent = self.widgets["main"])
+        cmds.menuItem(label = "None", parent = self.widgets["favoriteProject_OM"])
+
+        cmds.formLayout(self.widgets["main"], edit = True, af = [(self.widgets["favoriteProject_OM"], "top", 160), (self.widgets["favoriteProject_OM"], "right", 25)]) 
+
+
+        #get projects and add to menu
+        self.findProjects()
+
+
+        #save settings button
+        self.widgets["saveSettings"] = cmds.button(w = 380, h = 40, label = "Save Settings and Close", parent = self.widgets["main"], c = self.saveSettings)
+        cmds.formLayout(self.widgets["main"], edit = True, af = [(self.widgets["saveSettings"], "bottom", 10), (self.widgets["saveSettings"], "right", 10)])   
+
+        #show window
+        cmds.showWindow(self.widgets["window"])
+
+        #get tools path
+        self.getToolsPath()
+
+
+        #load settings from disk
+        settingsLocation = self.mayaToolsDir + "/General/Scripts/projectSettings.txt"
+        if os.path.exists(settingsLocation):
+            f = open(settingsLocation, 'r')
+            settings = cPickle.load(f)	    
+            useSource = settings.get("UseSourceControl")
+            favoriteProject = settings.get("FavoriteProject")
+
+            f.close()
+
+            try:
+                #set the checkbox to the settings value
+                cmds.checkBox(self.widgets["useSourceControl"], edit = True, v = useSource)
+                self.useSourceControl(useSource)
+
+            except:
+                pass
+
+            try:
+                #set the optionMenu to favorite project
+                cmds.optionMenu(self.widgets["favoriteProject_OM"], edit = True, v = favoriteProject)
+
+            except:
+                pass
+
+
+
+
+    # # # # # # # # # # # # # # # # # # # # # # # # # # # # # # # # # # # # #
+    def close(self, *args):
+
+        cmds.deleteUI("AnimationRiggingTool_SettingsUI")
+
+
+
+    # # # # # # # # # # # # # # # # # # # # # # # # # # # # # # # # # # # # #
+    def getToolsPath(self, *args):
+
+        toolsPath = cmds.internalVar(usd = True) + "mayaTools.txt"
+        if os.path.exists(toolsPath):
+
+            f = open(toolsPath, 'r')
+            self.mayaToolsDir = f.readline()
+            f.close()
+
+        #edit text field
+        cmds.textField(self.widgets["toolsPathField"], edit = True, text = self.mayaToolsDir)
+
+
+    # # # # # # # # # # # # # # # # # # # # # # # # # # # # # # # # # # # # #
+    def browseNewToolsPath(self, *args):
+
+        newPath = cmds.fileDialog2(fm = 3)
+        if newPath != None:
+            newPath = newPath[0]
+
+            if newPath.rpartition("/")[2] != "MayaTools":
+                cmds.warning("Selected directory is not valid. Please locate the MayaTools directory.")      
+                return
+
+            else:
+                #create file that contains this path
+                path = cmds.internalVar(usd = True) + "mayaTools.txt"
+
+                f = open(path, 'w')
+                f.write(newPath)
+                f.close()
+
+                self.getToolsPath()
+                self.close()
+                cmds.confirmDialog(title = "Settings", icon = "warning", message = "New tools path will not be used until Maya has been restarted.")
+
+
+    # # # # # # # # # # # # # # # # # # # # # # # # # # # # # # # # # # # # #
+    def useSourceControl(self, state, *args):
+
+        if state:
+            cmds.button(self.widgets["testConnection"], edit = True, enable = True)
+            cmds.menuItem("perforceSubmitMenuItem", edit = True, enable = True)
+            cmds.menuItem("perforceAddAndSubmitMenuItem", edit = True, enable = True)
+            cmds.menuItem("perforceCheckOutMenuItem", edit = True, enable = True)
+            cmds.menuItem("perforceFileHistoryMenuItem", edit = True, enable = True)
+            cmds.menuItem("perforceGetLatestMenuItem", edit = True, enable = True)
+            cmds.menuItem("perforceProjectList", edit = True, enable = True)
+
+        if state == False:
+            cmds.button(self.widgets["testConnection"], edit = True, enable = False)
+            cmds.menuItem("perforceSubmitMenuItem", edit = True, enable = False)
+            cmds.menuItem("perforceAddAndSubmitMenuItem", edit = True, enable = False)
+            cmds.menuItem("perforceCheckOutMenuItem", edit = True, enable = False)
+            cmds.menuItem("perforceFileHistoryMenuItem", edit = True, enable = False)
+            cmds.menuItem("perforceGetLatestMenuItem", edit = True, enable = False)
+            cmds.menuItem("perforceProjectList", edit = True, enable = False)
+
+
+    # # # # # # # # # # # # # # # # # # # # # # # # # # # # # # # # # # # # #
+    def testConnection(self, *args):
+
+
+        try:
+            from P4 import P4, P4Exception
+
+        except:
+            cmds.confirmDialog(title = "Test Connection", icon = "critical", message = "Perforce Python modules not found.")
+            return
+
+
+        try:
+            p4 = P4()
+            p4.connect()
+
+
+            #client info
+            spec = p4.run( "client", "-o" )[0]
+
+            client = spec.get("Client")
+            owner = spec.get("Owner")
+            host = spec.get("Host")
+
+            p4.disconnect()
+
+            cmds.confirmDialog(title = "Test Connection", message = "Connection Sucessful!\n\nUsername: " + owner + "\nClient: " + client + "\nHost: " + host + ".")
+
+
+        except:
+            cmds.confirmDialog(title = "Test Connection", message = "Unable to connect to Perforce server.")
+
+
+    # # # # # # # # # # # # # # # # # # # # # # # # # # # # # # # # # # # # #
+    def saveSettings(self, *args):
+
+        #this function will save out the user's preferences they have set in the UI to disk
+        settingsLocation = self.mayaToolsDir + "/General/Scripts/projectSettings.txt"
+
+        try:
+            f = open(settingsLocation, 'w')
+
+
+            #create a dictionary with  values
+            settings = {}
+            settings["UseSourceControl"] = cmds.checkBox(self.widgets["useSourceControl"], q = True, v = True)
+            settings["FavoriteProject"] = cmds.optionMenu(self.widgets["favoriteProject_OM"], q = True, v = True)
+
+            #write our dictionary to file
+            cPickle.dump(settings, f)
+            f.close()
+
+        except:
+            cmds.confirmDialog(title = "Error", icon = "critical", message = settingsLocation + " is not writeable. Please make sure this file is not set to read only.")
+
+
+
+        #close the UI
+        cmds.deleteUI("AnimationRiggingTool_SettingsUI")
+
+
+
+    # # # # # # # # # # # # # # # # # # # # # # # # # # # # # # # # # # # # #
+    def findProjects(self):
+
+        projectPath = self.mayaToolsDir + "/General/ART/Projects/"
+
+        try:
+            projects = os.listdir(projectPath)
+
+            if len(projects) > 0:
+                for proj in sorted(projects):
+                    cmds.menuItem(label = proj, parent = self.widgets["favoriteProject_OM"])
+
+        except:
+            pass
\ No newline at end of file
