Branching

diff --git a/Engine/Extras/Maya_AnimationRiggingTools/ARTv1/General/Scripts/ART_addCharacter_UI.py b/Engine/Extras/Maya_AnimationRiggingTools/ARTv1/General/Scripts/ART_addCharacter_UI.py
new file mode 100644
--- /dev/null
+++ b/Engine/Extras/Maya_AnimationRiggingTools/ARTv1/General/Scripts/ART_addCharacter_UI.py
@@ -0,0 +1,200 @@
+import maya.cmds as cmds
+import maya.mel as mel
+from functools import partial
+import os, cPickle
+
+
+class AddCharacter_UI():
+
+    def __init__(self):
+
+        #get access to our maya tools
+        toolsPath = cmds.internalVar(usd = True) + "mayaTools.txt"
+        if os.path.exists(toolsPath):
+
+            f = open(toolsPath, 'r')
+            self.mayaToolsDir = f.readline()
+            f.close()
+
+
+        #create a dictionary for our UI widgets
+        self.widgets = {}
+
+        #create window
+        if cmds.window("addCharacter_UI", exists = True):
+            cmds.deleteUI("addCharacter_UI")
+
+        self.widgets["window"] = cmds.window("addCharacter_UI", title = "Add Character", w = 430, h = 260, sizeable = True, mxb = False, mnb = False)
+
+        #create the main layout
+        main = cmds.columnLayout(w = 430)
+
+        #banner image
+        cmds.image(w = 430, h = 50, image = self.mayaToolsDir + "/General/Icons/ART/artBanner430px.bmp", parent = main)
+
+        self.widgets["formLayout"] = cmds.formLayout(w = 430, h = 260, parent = main)
+
+        #create the widgets
+        self.widgets["project"] = cmds.optionMenu(label = "Project:", w = 200, cc = self.findCharacterRigs)
+        self.widgets["search"] = cmds.textFieldGrp(label = "Search:", cl2 = ["left", "left"], ct2 = ["left", "left"], cw2 = [40, 160], text = "coming soon", enable = False)
+        self.widgets["characterList"] = cmds.textScrollList(allowMultiSelection = False, w = 200, h = 180, selectCommand = self.findCharacterThumbnail)
+        self.widgets["thumbnailBorder"] = cmds.image(w = 210, h = 210, image = "", ebg = True, bgc = [0, 0, 0])
+        self.widgets["thumbnail"] = cmds.image(w = 200, h = 200, image = "")
+        self.widgets["addButton"] = cmds.button(label = "Add", w = 60, c = partial(self.addCharacter, False))
+        self.widgets["addButton2"] = cmds.button(label = "Add and Close", w = 80, c = partial(self.addCharacter, True))
+        self.widgets["cancelButton"] = cmds.button(label = "Cancel", w = 60, c = self.cancel)
+
+
+        #attach widgets
+        cmds.formLayout(self.widgets["formLayout"], edit = True, af = [(self.widgets["project"],  'top', 10),(self.widgets["project"],  'left', 10)])
+        cmds.formLayout(self.widgets["formLayout"], edit = True, af = [(self.widgets["search"],  'top', 40),(self.widgets["search"],  'left', 10)])
+
+        cmds.formLayout(self.widgets["formLayout"], edit = True, af = [(self.widgets["characterList"],  'top', 70),(self.widgets["characterList"],  'left', 10)])
+        cmds.formLayout(self.widgets["formLayout"], edit = True, af = [(self.widgets["thumbnailBorder"],  'top', 5),(self.widgets["thumbnailBorder"],  'right', 5)])
+        cmds.formLayout(self.widgets["formLayout"], edit = True, af = [(self.widgets["thumbnail"],  'top', 10),(self.widgets["thumbnail"],  'right', 10)])
+        cmds.formLayout(self.widgets["formLayout"], edit = True, af = [(self.widgets["addButton"],  'bottom', 10),(self.widgets["addButton"],  'left', 215)])
+        cmds.formLayout(self.widgets["formLayout"], edit = True, af = [(self.widgets["addButton2"],  'bottom', 10),(self.widgets["addButton2"],  'left', 280)])
+        cmds.formLayout(self.widgets["formLayout"], edit = True, af = [(self.widgets["cancelButton"],  'bottom', 10),(self.widgets["cancelButton"],  'right', 5)])
+
+
+        #show the window
+        cmds.showWindow(self.widgets["window"])
+        self.findProjects()
+
+        #set favorite project if it exists
+        settingsLocation = self.mayaToolsDir + "/General/Scripts/projectSettings.txt"
+        if os.path.exists(settingsLocation):
+            f = open(settingsLocation, 'r')
+            settings = cPickle.load(f)
+            favoriteProject = settings.get("FavoriteProject")
+
+            try:
+                cmds.optionMenu(self.widgets["project"], edit = True, v = favoriteProject)
+            except:
+                pass
+
+            self.findCharacterRigs()
+
+
+    def findProjects(self):
+
+        children = cmds.optionMenu(self.widgets["project"], q = True, itemListLong = True)
+        if children != None:
+            for child in children:
+                cmds.deleteUI(child)
+
+
+        projectPath = self.mayaToolsDir + "/General/ART/Projects/"
+        try:
+            projects = os.listdir(projectPath)
+
+            if len(projects) > 0:
+                for proj in sorted(projects):
+                    cmds.menuItem(label = proj, parent = self.widgets["project"])
+
+
+                self.findCharacterRigs()
+
+        except:
+            cmds.confirmDialog(title = "Add Character", message = "There are no projects to add a character from.")
+            cmds.deleteUI(self.widgets["window"])
+
+
+
+
+    def findCharacterRigs(self, *args):
+
+        #remove any existing child elements in the text scroll list
+        cmds.textScrollList(self.widgets["characterList"], edit = True, removeAll = True)
+
+        #find all character rigs
+        project = cmds.optionMenu(self.widgets["project"], q = True, value = True)
+
+        rigPath = self.mayaToolsDir + "/General/ART/Projects/" + project + "/AnimRigs/"
+        rigs = os.listdir(rigPath)
+
+        for rig in rigs:
+            niceName = rig.partition(".mb")[0]
+            if os.path.isfile(rigPath + rig):
+                cmds.textScrollList(self.widgets["characterList"], edit = True, append = niceName)
+
+
+        #select the first item
+        cmds.textScrollList(self.widgets["characterList"], edit = True, selectIndexedItem = 1)
+
+        self.findCharacterThumbnail()
+
+    def findCharacterThumbnail(self):
+        project = cmds.optionMenu(self.widgets["project"], q = True, value = True)
+        selectedCharacter = cmds.textScrollList(self.widgets["characterList"], q = True, si = True)[0]
+        thumbnailPath = self.mayaToolsDir + "/General/Icons/ART/Thumbnails/" + project
+
+        try:
+            thumbnails = os.listdir(thumbnailPath)
+
+            for thumb in thumbnails:
+                if thumb.find("_small") == -1:
+                    if thumb.find(selectedCharacter) == 0:
+                        cmds.image(self.widgets["thumbnail"], edit = True, image = thumbnailPath + "/" + thumb)
+
+        except:
+            cmds.confirmDialog(icon = "critical", title = "Error", message = "No thumbnail found for character")
+            return
+
+
+    def addCharacter(self, close, *args):
+        project = cmds.optionMenu(self.widgets["project"], q = True, value = True)
+        selectedCharacter = cmds.textScrollList(self.widgets["characterList"], q = True, si = True)[0]
+        rigPath = self.mayaToolsDir + "/General/ART/Projects/" + project + "/AnimRigs/" + selectedCharacter + ".mb"
+
+
+        #find existing namespaces in scene
+        namespaces = cmds.namespaceInfo(listOnlyNamespaces = True)
+
+        #reference the rig file
+        cmds.file(rigPath, r = True, type = "mayaBinary", loadReferenceDepth = "all", namespace = selectedCharacter, options = "v=0")
+
+        #clear selection and fit view
+        cmds.select(clear = True)
+        
+        #query autoKeyFrame state
+        if not cmds.autoKeyframe(q = True, state = True):
+            cmds.viewFit()
+        #cmds.viewFit()
+        panels = cmds.getPanel(type = 'modelPanel')
+
+
+        #turn on smooth shading
+        for panel in panels:
+            editor = cmds.modelPanel(panel, q = True, modelEditor = True)
+            cmds.modelEditor(editor, edit = True, displayAppearance = "smoothShaded", displayTextures = True, textures = True )
+
+
+
+        #find new namespaces in scene (this is here in case I need to do something later and I need the new name that was created)
+        newCharacterName = selectedCharacter
+        newNamespaces = cmds.namespaceInfo(listOnlyNamespaces = True)
+
+        for name in newNamespaces:
+            if name not in namespaces:
+                newCharacterName = name
+
+        #launch UI
+        import ART_animationUI
+        reload(ART_animationUI)
+        ART_animationUI.AnimationUI()
+
+        if close:
+            cmds.deleteUI(self.widgets["window"])
+
+
+
+
+
+
+    def cancel(self, *args):
+        cmds.deleteUI(self.widgets["window"])
+
+
+
+
