Branching

diff --git a/Engine/Extras/Maya_AnimationRiggingTools/ARTv1/General/Scripts/ART_fbxBatch.py b/Engine/Extras/Maya_AnimationRiggingTools/ARTv1/General/Scripts/ART_fbxBatch.py
new file mode 100644
--- /dev/null
+++ b/Engine/Extras/Maya_AnimationRiggingTools/ARTv1/General/Scripts/ART_fbxBatch.py
@@ -0,0 +1,339 @@
+#fbx export batcher
+
+import maya.cmds as cmds
+from functools import partial
+import os, cPickle, math
+import maya.mel as mel
+import maya.utils
+    
+class FBX_Batcher():
+    def __init__(self):
+        #get access to our maya tools
+        toolsPath = cmds.internalVar(usd = True) + "mayaTools.txt"
+        if os.path.exists(toolsPath):
+            
+            f = open(toolsPath, 'r')
+            self.mayaToolsDir = f.readline()
+            f.close()   
+
+        self.widgets = {}
+        
+        if cmds.window("ART_FBX_BATCH", exists = True):
+            cmds.deleteUI("ART_FBX_BATCH")
+            
+        self.widgets["window"] = cmds.window("ART_FBX_BATCH", title = "FBX Batch", w = 400, h = 150, sizeable = False, mnb = False, mxb = False)
+        
+        #main layout
+        self.widgets["layout"] = cmds.formLayout(w = 400, h = 150)
+        
+        #textField and browse button
+        self.widgets["textField"] = cmds.textField(w = 300, text = "")
+        self.widgets["browse"] = cmds.button(w = 70, label = "Browse", c = self.browse)
+        
+        cmds.formLayout(self.widgets["layout"], edit = True, af = [(self.widgets["textField"], "left", 10),(self.widgets["textField"], "top", 10)])
+        cmds.formLayout(self.widgets["layout"], edit = True, af = [(self.widgets["browse"], "right", 10),(self.widgets["browse"], "top", 10)])
+        
+        #process button
+        self.widgets["process"] = cmds.button(w = 380, h = 50, label = "BEGIN BATCH", c = self.process)
+        cmds.formLayout(self.widgets["layout"], edit = True, af = [(self.widgets["process"], "left", 10),(self.widgets["process"], "top", 40)])
+        
+        #progress bar
+        self.widgets["currentFileProgressBar"] = cmds.progressBar(w = 390)
+        cmds.formLayout(self.widgets["layout"], edit = True, af = [(self.widgets["currentFileProgressBar"], "left", 5),(self.widgets["currentFileProgressBar"], "bottom", 35)])
+        
+        self.widgets["progressBar"] = cmds.progressBar(w = 390)
+        cmds.formLayout(self.widgets["layout"], edit = True, af = [(self.widgets["progressBar"], "left", 5),(self.widgets["progressBar"], "bottom", 10)])
+        
+        #show window
+        cmds.showWindow(self.widgets["window"])
+
+    def browse(self, *args):
+        try:
+            directory = cmds.fileDialog2(dialogStyle=2, fm = 3)[0]
+            cmds.textField(self.widgets["textField"], edit = True, text = directory)
+        except:
+            pass
+    
+    def process(self, *args):
+        failed = []
+        succeeds = []
+        
+        #new file
+        cmds.file(new = True, force = True)
+        
+        #get the files in the directory
+        directory = cmds.textField(self.widgets["textField"], q = True, text = True)
+        filesInDir = os.listdir(directory)
+        mayaFiles = []
+        
+        for each in filesInDir:
+            
+            if each.rpartition(".")[2] == "mb":
+                mayaFiles.append(each)
+            
+            #go through loop
+            if len(mayaFiles) > 0:
+                
+                #get number of maya files
+                numFiles = len(mayaFiles)
+                amount = 100/numFiles
+                newAmount = 0
+                
+                #print mayaFiles
+                
+                for mayaFile in mayaFiles:
+                    cmds.progressBar(self.widgets["progressBar"], edit = True, progress = newAmount + amount)
+                
+                #open file
+                directory = os.path.normpath(directory)
+                cmds.file(os.path.join(directory, mayaFile), open = True, force = True)
+                    
+                #update
+                cmds.progressBar(self.widgets["currentFileProgressBar"], edit = True, progress = 50)
+                
+                #get characters
+                characters = []
+                references = cmds.ls(type = "reference")
+                
+                for reference in references:
+                    niceName = reference.rpartition("RN")[0]
+                    suffix = reference.rpartition("RN")[2]
+                    if suffix != "":
+                        if cmds.objExists(niceName + suffix + ":" + "Skeleton_Settings"):
+                            characters.append(niceName + suffix)
+                        
+                    else:
+                        if cmds.objExists(niceName + ":" + "Skeleton_Settings"):
+                            characters.append(niceName) 
+                        
+                #export FBX for each character
+                for character in characters:
+                    #check to see if file has cached data:
+                    if cmds.objExists("ExportAnimationSettings"):
+                        sequeneces = cmds.listAttr("ExportAnimationSettings", string = "sequence*")
+                    
+                        for seq in sequeneces:
+                            #print seq
+                            
+                            #get data
+                            data = cmds.getAttr("ExportAnimationSettings." + seq)
+                            name = data.partition("::")[0]
+                            data = data.partition("::")[2]
+                            start = data.partition("::")[0]
+                            data = data.partition("::")[2]
+                            end = data.partition("::")[0]
+                            data = data.partition("::")[2]
+                            fps = data.partition("::")[0]
+                            interp = data.partition("::")[2]
+                            
+                            if interp == "Independent Euler Angle":
+                                interp = 1
+                            
+                            if interp == "Synchronized Euler Angle":
+                                interp = 2
+                            
+                            if interp == "Quaternion Slerp":
+                                interp = 3
+                            
+                            #print name
+                            #print start
+                            #print end
+                            #print fps
+                            #print interp
+                            
+                            #self.getBlendshapes(character, start, end)
+                            
+                            #export
+                            self.exportFBX(character, os.path.join(name), start, end, fps, interp)
+                    else:
+                        #create the fbx path with /FBX
+                        fbxPath = os.path.join(directory, "FBX")
+                        if not os.path.exists(fbxPath):
+                            os.makedirs(fbxPath)
+                    
+                        #get name
+                        fileName = mayaFile.rpartition(".")[0]          
+                
+                        #get current frame range
+                        start = cmds.playbackOptions(q = True, min = True)
+                        end = cmds.playbackOptions(q = True, max = True)
+                        #self.getBlendshapes(character, start, end)
+                        self.exportFBX(character, os.path.join(fbxPath, fileName + ".fbx"), start, end, "30 FPS", 3)
+                        
+                    cmds.progressBar(self.widgets["currentFileProgressBar"], edit = True, progress = 100)
+                    
+                #increment new amount
+                newAmount = newAmount + amount
+            
+        #new scene
+        cmds.file(new = True, force = True)
+        
+        #delete UI
+        cmds.deleteUI(self.widgets["window"])
+        
+        #confirm message
+        cmds.confirmDialog(title = "Complete!", message = "Operation is complete!")
+        
+    def exportFBX(self, character, path, startFrame, endFrame, fps, interp, *args):
+        #get the current time range values
+        originalStart = cmds.playbackOptions(q = True, min = True)
+        originalEnd = cmds.playbackOptions(q = True, max = True)
+        currentFPS = cmds.currentUnit(q = True, time = True)
+        currentInterp = 1
+        
+        #set rotation interp
+        cmds.optionVar(iv = ("rotationInterpolationDefault", int(interp)))
+        
+        #duplicate the skeleton
+        newSkeleton = cmds.duplicate(character + ":" + "root", un = False, ic = False)
+        
+        joints = []
+        for each in newSkeleton:
+            if cmds.nodeType(each) != "joint":
+                cmds.delete(each)
+            else:
+                joints.append(each)
+        
+        constraints = []
+        for joint in joints:
+            #do some checks to make sure that this is valid
+            parent = cmds.listRelatives(joint, parent = True)
+
+            if parent != None:
+                if parent[0] in joints:
+                    constraint = cmds.parentConstraint(character + ":" + parent[0] + "|" + character + ":" + joint, joint)[0]
+                    constraints.append(constraint)
+            else:
+                #root bone?
+                if joint == "root":
+                    constraint = cmds.parentConstraint(character + ":" + joint, joint)[0]
+                    constraints.append(constraint)              
+
+        #Set start and end frame
+        cmds.playbackOptions(min = startFrame, animationStartTime = startFrame)
+        cmds.playbackOptions(max = endFrame, animationEndTime = endFrame)
+        
+        #set to the new fps
+        if fps == "30 FPS":
+            fps = "ntsc"
+            
+        if fps == "60 FPS":
+            fps = "ntscf"
+            
+        if fps == "120 FPS":
+            fps  = "120fps"
+            
+        cmds.currentUnit(time = fps)
+        
+        #set the new frame range
+        if fps == "ntscf":
+            cmds.playbackOptions(min = startFrame * 2, animationStartTime = startFrame * 2)
+            cmds.playbackOptions(max = endFrame * 2, animationEndTime = endFrame * 2)
+            startFrame = startFrame * 2
+            endFrame = endFrame * 2
+            
+        if fps == "120fps":
+            cmds.playbackOptions(min = startFrame * 4, animationStartTime = startFrame * 4)
+            cmds.playbackOptions(max = endFrame * 4, animationEndTime = endFrame * 4)
+            startFrame = startFrame * 4
+            endFrame = endFrame * 4
+            
+        #bake results
+        cmds.select("root", hi = True)
+        cmds.bakeResults(simulation = True, t = (startFrame, endFrame))
+        cmds.delete(constraints)
+                
+        #run an euler filter
+        cmds.select("root", hi = True)
+        cmds.filterCurve()
+                
+
+        #print "startFrame: ",startFrame
+        #print "endFrame: ",endFrame
+        #print "BLENDSHAPES CODE IS RUNNING!!!!!!!!"
+        
+        #get the blendshapes to export
+        allBlends = cmds.ls(type = "blendShape")
+        jmBlends = ["calf_l_shapes",  "calf_r_shapes",  "head_shapes",  "l_elbow_shapes",  "r_elbow_shapes",  "l_lowerarm_shapes",  "r_lowerarm_shapes",  "l_shoulder_shapes",  "r_shoulder_shapes",  "l_upperarm_shapes",  "r_upperarm_shapes",
+                    "neck1_shapes",  "neck2_shapes",  "neck3_shapes",  "pelvis_shapes",  "spine1_shapes",  "spine2_shapes",  "spine3_shapes",  "spine4_shapes",  "spine5_shapes",  "thigh_l_shapes",  "thigh_r_shapes"]
+        blendshapes = []
+        #print "allBlends list:", allBlends
+        #print "jmBlends list:", jmBlends
+
+        if allBlends != None:
+            for shape in allBlends:
+                if shape.partition(":")[2] not in jmBlends:
+                    if shape.partition(":")[0] == character:
+                        blendshapes.append(shape)
+        #print "blendshapes list:", blendshapes
+
+        if blendshapes != None:
+            if cmds.objExists("custom_export"):
+                cmds.delete("custom_export")
+
+            cube = cmds.polyCube(name = "custom_export")[0]
+            i = 1
+            for shape in blendshapes:
+                attrs = cmds.listAttr(shape, m = True, string = "weight")
+                #counter for index
+
+                for attr in attrs:
+                    #print attr
+                    keys = cmds.keyframe( shape, attribute=attr, query=True, timeChange = True )
+                    #print keys
+                    keyValues = cmds.keyframe( shape, attribute=attr, query=True, valueChange = True )
+                    #print keyValues
+
+                    morph = cmds.polyCube(name = attr)[0]
+                    if cmds.objExists("custom_export_shapes"):
+                        cmds.blendShape("custom_export_shapes", edit = True, t = (cube, i, morph, 1.0))
+                        #print "editing blendshape to include " + cube + str(i) + morph
+
+                    else:
+                        cmds.select([morph, cube], r = True)
+                        cmds.blendShape(name = "custom_export_shapes")
+                        cmds.select(clear = True)
+
+                    cmds.delete(morph)
+                    #print i
+                    #transfer keys from original to new morph
+                    if keys != None:
+                        for x in range(int(len(keys))):
+                            cmds.setKeyframe("custom_export_shapes." + attr, t = (keys[x]), v = keyValues[x])
+
+                    if keys == None:
+                        for x in range(int(startFrame), (int(endFrame) + 1)):
+                            value = cmds.getAttr(shape + "." + attr, time = x)
+                            cmds.setKeyframe("custom_export_shapes." + attr, t = (x), v = value)
+
+                    i = i + 1
+
+        #export selected
+        #first change some fbx properties
+        string = "FBXExportConstraints -v 1;"
+        string += "FBXExportCacheFile -v 0;"
+        mel.eval(string)
+                
+        cmds.select("root", hi = True)
+        if cmds.objExists("custom_export"):
+            cmds.select("custom_export", add = True)
+
+        cmds.file(path, es = True, force = True, prompt = False, type = "FBX export")
+        
+        #clean scene
+        cmds.delete("root")
+        if cmds.objExists("custom_export"):
+            cmds.delete("custom_export")
+            
+        #reset fps to original
+        cmds.currentUnit(time = currentFPS)     
+
+        #reset timeslider to original values
+        cmds.playbackOptions(min = originalStart, max = originalEnd)
+        
+        #reset rotation interp
+        cmds.optionVar(iv = ("rotationInterpolationDefault", currentInterp))
+        
+    
+
+        
\ No newline at end of file
diff --git a/Engine/Extras/Maya_AnimationRiggingTools/ARTv1/General/Scripts/ART_fbxBatcher.py b/Engine/Extras/Maya_AnimationRiggingTools/ARTv1/General/Scripts/ART_fbxBatcher.py
new file mode 100644
--- /dev/null
+++ b/Engine/Extras/Maya_AnimationRiggingTools/ARTv1/General/Scripts/ART_fbxBatcher.py
@@ -0,0 +1,414 @@
+import maya.cmds as cmds
+from functools import partial
+import os, cPickle, math
+import maya.mel as mel
+import maya.utils
+
+class FBX_Batcher():
+
+    def __init__(self):
+
+        #get access to our maya tools
+        toolsPath = cmds.internalVar(usd = True) + "mayaTools.txt"
+        if os.path.exists(toolsPath):
+
+            f = open(toolsPath, 'r')
+            self.mayaToolsDir = f.readline()
+            f.close()	
+
+        self.widgets = {}
+
+        if cmds.window("ART_RetargetTool_UI", exists = True):
+            cmds.deleteUI("ART_RetargetTool_UI")
+
+        self.widgets["window"] = cmds.window("ART_RetargetTool_UI", title = "ART FBX Export Batcher", w = 400, h = 400, sizeable = False, mnb = False, mxb = False)
+
+        #main layout
+        self.widgets["layout"] = cmds.formLayout(w = 400, h = 400)
+
+        #textField and browse button
+        label = cmds.text(label = "Directory to Batch:")
+
+        self.widgets["textField"] = cmds.textField(w = 300, text = "")
+        self.widgets["browse"] = cmds.button(w = 70, label = "Browse", c = self.browse)
+
+        cmds.formLayout(self.widgets["layout"], edit = True, af = [(label, "left", 10),(label, "top", 10)])	
+        cmds.formLayout(self.widgets["layout"], edit = True, af = [(self.widgets["textField"], "left", 10),(self.widgets["textField"], "top", 40)])
+        cmds.formLayout(self.widgets["layout"], edit = True, af = [(self.widgets["browse"], "right", 10),(self.widgets["browse"], "top", 40)])
+
+        #frameLayout for custom file script
+        self.widgets["frame"] = cmds.frameLayout(w = 380, h = 100, bs = "etchedIn", cll = False, label = "Advanced", parent = self.widgets["layout"])
+        cmds.formLayout(self.widgets["layout"], edit = True, af = [(self.widgets["frame"], "right", 10),(self.widgets["frame"], "top", 80)])
+        self.widgets["advancedForm"] = cmds.formLayout(w = 380, h = 100, parent = self.widgets["frame"])
+
+        label2 = cmds.text("Custom File Script:", parent =self.widgets["advancedForm"])
+
+        self.widgets["scriptField"] = cmds.textField(w = 280, text = "", parent =self.widgets["advancedForm"])
+        self.widgets["scriptBrowse"] = cmds.button(w = 70, label = "Browse", c = self.scriptBrowse, parent =self.widgets["advancedForm"])
+
+        cmds.formLayout(self.widgets["advancedForm"], edit = True, af = [(label2, "left", 10),(label2, "top", 10)])	
+        cmds.formLayout(self.widgets["advancedForm"], edit = True, af = [(self.widgets["scriptField"], "left", 10),(self.widgets["scriptField"], "top", 40)])
+        cmds.formLayout(self.widgets["advancedForm"], edit = True, af = [(self.widgets["scriptBrowse"], "right", 10),(self.widgets["scriptBrowse"], "top", 40)])
+
+
+        #option checkbox for remove motion on root
+        self.widgets["removeRoot"] = cmds.checkBox(label = "Remove root animation?", v = False, parent = self.widgets["layout"])
+        cmds.formLayout(self.widgets["layout"], edit = True, af = [(self.widgets["removeRoot"], "left", 10),(self.widgets["removeRoot"], "top", 220)])	
+
+        self.widgets["exportFBX"] = cmds.checkBox(label = "Export an FBX?", v = True, parent = self.widgets["layout"])
+        cmds.formLayout(self.widgets["layout"], edit = True, af = [(self.widgets["exportFBX"], "right", 10),(self.widgets["exportFBX"], "top", 220)])	
+
+        #process button
+        self.widgets["process"] = cmds.button(w = 380, h = 50, label = "BEGIN BATCH", c = self.process, parent = self.widgets["layout"])
+        cmds.formLayout(self.widgets["layout"], edit = True, af = [(self.widgets["process"], "left", 10),(self.widgets["process"], "top", 250)])
+
+        #progress bar
+        text = cmds.text(label = "File Progress: ", parent = self.widgets["layout"])
+        self.widgets["currentFileProgressBar"] = cmds.progressBar(w = 250, parent = self.widgets["layout"])
+
+        cmds.formLayout(self.widgets["layout"], edit = True, af = [(text, "left", 10),(text, "top", 320)])
+        cmds.formLayout(self.widgets["layout"], edit = True, af = [(self.widgets["currentFileProgressBar"], "left", 110),(self.widgets["currentFileProgressBar"], "top", 320)])
+
+        text2 = cmds.text(label = "Total Progress: ", parent = self.widgets["layout"])
+        self.widgets["progressBar"] = cmds.progressBar(w = 250, parent = self.widgets["layout"])
+
+        cmds.formLayout(self.widgets["layout"], edit = True, af = [(text2, "left", 10),(text2, "top", 360)])
+        cmds.formLayout(self.widgets["layout"], edit = True, af = [(self.widgets["progressBar"], "left", 110),(self.widgets["progressBar"], "top", 360)])
+
+        #show window
+        cmds.showWindow(self.widgets["window"])
+
+
+    # # # # # # # # # # # # # # # # # # # # # # # # # # # # # # # # # # # # # # # # # # # # # # # # # # # # # # # # # # # # # # # 
+    # # # # # # # # # # # # # # # # # # # # # # # # # # # # # # # # # # # # # # # # # # # # # # # # # # # # # # # # # # # # # # # 
+    # # # # # # # # # # # # # # # # # # # # # # # # # # # # # # # # # # # # # # # # # # # # # # # # # # # # # # # # # # # # # # # 
+    # # # # # # # # # # # # # # # # # # # # # # # # # # # # # # # # # # # # # # # # # # # # # # # # # # # # # # # # # # # # # # # 
+    # # # # # # # # # # # # # # # # # # # # # # # # # # # # # # # # # # # # # # # # # # # # # # # # # # # # # # # # # # # # # # # 
+    def browse(self, *args):
+
+        try:
+            directory = cmds.fileDialog2(dialogStyle=2, fm = 3)[0]
+            cmds.textField(self.widgets["textField"], edit = True, text = directory)
+
+        except:
+            pass
+
+    # # # # # # # # # # # # # # # # # # # # # # # # # # # # # # # # # # # # # # # # # # # # # # # # # # # # # # # # # # # # # # # 
+    # # # # # # # # # # # # # # # # # # # # # # # # # # # # # # # # # # # # # # # # # # # # # # # # # # # # # # # # # # # # # # # 
+    # # # # # # # # # # # # # # # # # # # # # # # # # # # # # # # # # # # # # # # # # # # # # # # # # # # # # # # # # # # # # # # 
+    # # # # # # # # # # # # # # # # # # # # # # # # # # # # # # # # # # # # # # # # # # # # # # # # # # # # # # # # # # # # # # # 
+    # # # # # # # # # # # # # # # # # # # # # # # # # # # # # # # # # # # # # # # # # # # # # # # # # # # # # # # # # # # # # # # 	
+    def scriptBrowse(self, *args):
+
+        try:
+            directory = cmds.fileDialog2(dialogStyle=2, fm = 1)[0]
+            cmds.textField(self.widgets["scriptField"], edit = True, text = directory)
+
+        except:
+            pass
+
+    # # # # # # # # # # # # # # # # # # # # # # # # # # # # # # # # # # # # # # # # # # # # # # # # # # # # # # # # # # # # # # # 
+    # # # # # # # # # # # # # # # # # # # # # # # # # # # # # # # # # # # # # # # # # # # # # # # # # # # # # # # # # # # # # # # 
+    # # # # # # # # # # # # # # # # # # # # # # # # # # # # # # # # # # # # # # # # # # # # # # # # # # # # # # # # # # # # # # # 
+    # # # # # # # # # # # # # # # # # # # # # # # # # # # # # # # # # # # # # # # # # # # # # # # # # # # # # # # # # # # # # # # 
+    # # # # # # # # # # # # # # # # # # # # # # # # # # # # # # # # # # # # # # # # # # # # # # # # # # # # # # # # # # # # # # # 	
+    def process(self, *args):
+
+        failed = []
+        succeeds = []
+
+
+        #new file
+        cmds.file(new = True, force = True)
+
+        #get the files in the directory
+        directory = cmds.textField(self.widgets["textField"], q = True, text = True)
+        filesInDir = os.listdir(directory)
+        mayaFiles = []
+
+
+
+
+        for each in filesInDir:
+
+            if each.rpartition(".")[2] == "mb":
+                mayaFiles.append(each)
+            if each.rpartition(".")[2] == "ma":
+                mayaFiles.append(each)
+
+        #go through loop
+        if len(mayaFiles) > 0:
+
+
+            #get number of maya files
+            numFiles = len(mayaFiles)
+            amount = 100/ (numFiles + 1)
+            newAmount = 0
+
+            for mayaFile in mayaFiles:
+                cmds.progressBar(self.widgets["progressBar"], edit = True, progress = newAmount + amount)
+
+
+                #get name
+                fileName = mayaFile.rpartition(".")[0]
+
+                #open file
+                directory = os.path.normpath(directory)
+                cmds.file(os.path.join(directory, mayaFile), open = True, force = True)
+
+                #update
+                cmds.progressBar(self.widgets["currentFileProgressBar"], edit = True, progress = 10)
+
+                #execute custom script if it exists
+                cmds.progressBar(self.widgets["currentFileProgressBar"], edit = True, progress = 20)
+                script = cmds.textField(self.widgets["scriptField"], q = True, text = True)
+                sourceType = ""		
+
+                if script.find(".py") != -1:
+                    sourceType = "python"
+
+                if script.find(".mel") != -1:
+                    sourceType = "mel"	
+
+
+                #execute script
+                if sourceType == "mel":
+                    try:
+                        command = ""
+                        #open the file, and for each line in the file, add it to our command string.
+                        f = open(script, 'r')
+                        lines = f.readlines()
+                        for line in lines:
+                            command += line
+
+                        import maya.mel as mel
+                        mel.eval(command)
+
+                    except Exception, e:
+                        print str(e)		    
+
+
+                if sourceType == "python":
+                    try:
+                        execfile("" + script + "")
+
+
+                    except Exception, e:
+                        print str(e)	
+
+
+                #export FBX
+                fbxPath = os.path.join(directory, "FBX")
+
+
+                #get current frame range
+                cmds.progressBar(self.widgets["currentFileProgressBar"], edit = True, progress = 30)
+                start = cmds.playbackOptions(q = True, min = True)
+                end = cmds.playbackOptions(q = True, max = True)
+                exportFBX = cmds.checkBox(self.widgets["exportFBX"], q = True, v = True)
+                if exportFBX:
+                    if not os.path.exists(fbxPath):
+                        os.makedirs(fbxPath)
+                    self.exportFBX(os.path.join(fbxPath, fileName + ".fbx"), start, end)
+
+                cmds.progressBar(self.widgets["currentFileProgressBar"], edit = True, progress = 40)
+
+                #increment new amount
+                newAmount = newAmount + amount		
+
+        #new scene
+        cmds.file(new = True, force = True)
+        cmds.progressBar(self.widgets["progressBar"], edit = True, progress = 100)
+
+        #delete UI
+        cmds.deleteUI(self.widgets["window"])
+
+        #confirm message
+        cmds.confirmDialog(title = "Complete!", message = "Operation is complete!")
+
+
+
+    # # # # # # # # # # # # # # # # # # # # # # # # # # # # # # # # # # # # # # # # # # # # # # # # # # # # # # # # # # # # # # # 
+    # # # # # # # # # # # # # # # # # # # # # # # # # # # # # # # # # # # # # # # # # # # # # # # # # # # # # # # # # # # # # # # 
+    # # # # # # # # # # # # # # # # # # # # # # # # # # # # # # # # # # # # # # # # # # # # # # # # # # # # # # # # # # # # # # # 
+    # # # # # # # # # # # # # # # # # # # # # # # # # # # # # # # # # # # # # # # # # # # # # # # # # # # # # # # # # # # # # # # 
+    # # # # # # # # # # # # # # # # # # # # # # # # # # # # # # # # # # # # # # # # # # # # # # # # # # # # # # # # # # # # # # # 
+    def exportFBX(self, path, startFrame, endFrame, *args):
+
+        #get the characters in the scene
+        characters = self.getCharacters()
+
+        #get rotation interp
+        options = cmds.optionVar(list = True)
+        for op in options:
+            if op == "rotationInterpolationDefault":
+                interp = cmds.optionVar(q = op)	
+        cmds.optionVar(iv = ("rotationInterpolationDefault", 3))
+
+
+        #Loop through each character in the scene, and export the fbx for that character
+        cmds.progressBar(self.widgets["currentFileProgressBar"], edit = True, progress = 40)
+
+        #create increment ammount for progress bar
+        increment = 50/len(characters)
+        step = increment/5	
+
+        for character in characters:
+
+
+            #add character suffix to fbx path file
+            exportPath = path.rpartition(".fbx")[0]
+            exportPath = exportPath + "_" + character + ".fbx"
+
+            #get blendshapes
+            cmds.progressBar(self.widgets["currentFileProgressBar"], edit = True, step = step)
+            allBlends = cmds.ls(type = "blendShape")
+            jmBlends = ["calf_l_shapes",  "calf_r_shapes",  "head_shapes",  "l_elbow_shapes",  "r_elbow_shapes",  "l_lowerarm_shapes",  "r_lowerarm_shapes",  "l_shoulder_shapes",  "r_shoulder_shapes",  "l_upperarm_shapes",  "r_upperarm_shapes",
+                        "neck1_shapes",  "neck2_shapes",  "neck3_shapes",  "pelvis_shapes",  "spine1_shapes",  "spine2_shapes",  "spine3_shapes",  "spine4_shapes",  "spine5_shapes",  "thigh_l_shapes",  "thigh_r_shapes"]
+            blendshapes = []
+
+            if allBlends != None:
+                for shape in allBlends:
+                    if shape.partition(":")[2] not in jmBlends:
+                        if shape.partition(":")[0] == character:
+                            blendshapes.append(shape)
+
+            #if our character has blendshapes, deal with those now
+            cmds.progressBar(self.widgets["currentFileProgressBar"], edit = True, step = step)
+            if blendshapes != None:
+                cube = cmds.polyCube(name = "custom_export")[0]
+
+                for shape in blendshapes:
+                    try:
+                        attrs = cmds.listAttr(shape, m = True, string = "weight")
+                        print attrs
+                        #counter for index
+                        x = 1
+
+                        for attr in attrs:
+                            morph = cmds.polyCube(name = attr)[0]
+                            if cmds.objExists("custom_export_shapes"):
+                                cmds.blendShape("custom_export_shapes", edit = True, t = (cube, x, morph, 1.0))
+
+                            else:
+                                cmds.select([morph, cube], r = True)
+                                cmds.blendShape(name = "custom_export_shapes")
+                                cmds.select(clear = True)
+
+                            cmds.delete(morph)
+
+                            #transfer animation from original to new morph
+                            for i in range(int(startFrame), int(endFrame) + 1):
+                                cmds.currentTime(i)
+                                value = cmds.getAttr(shape + "." + attr)
+                                print i, shape, attr, value
+                                cmds.setKeyframe("custom_export_shapes." + attr, t = (i), v = value)
+
+                            x = x + 1	
+
+                    except Exception, e:
+                        print str(e)
+
+
+
+
+            #duplicate the skeleton
+            cmds.progressBar(self.widgets["currentFileProgressBar"], edit = True, step = step)
+            newSkeleton = cmds.duplicate(character + ":" + "root", un = False, ic = False)
+
+            joints = []
+            for each in newSkeleton:
+                if cmds.nodeType(each) != "joint":
+                    cmds.delete(each)
+
+                else:
+                    joints.append(each)
+
+
+            #constrain the dupe skeleton to the original skeleton
+            constraints = []
+            for joint in joints:
+                #do some checks to make sure that this is valid
+                parent = cmds.listRelatives(joint, parent = True)
+
+                if parent != None:
+                    if parent[0] in joints:
+                        constraint = cmds.parentConstraint(character + ":" + parent[0] + "|" + character + ":" + joint, joint)[0]
+                        constraints.append(constraint)
+
+                else:
+                    #root bone?
+                    if joint == "root":
+                        constraint = cmds.parentConstraint(character + ":" + joint, joint)[0]
+                        constraints.append(constraint)			    
+
+
+            #bake the animation down onto the export skeleton
+            cmds.progressBar(self.widgets["currentFileProgressBar"], edit = True, step = step)
+            cmds.select("root", hi = True)
+            cmds.bakeResults(simulation = True, t = (startFrame, endFrame))
+            cmds.delete(constraints)
+
+
+            #check remove root animation checkbox. if true, delete keys off of root and zero out
+            removeRoot = cmds.checkBox(self.widgets["removeRoot"], q = True, v = True)
+            if removeRoot:
+                cmds.select("root")
+                cmds.cutKey()
+                for attr in ["tx", "ty", "tz", "rx", "ry", "rz"]:
+                    cmds.setAttr("root." + attr, 0)
+
+            #run an euler filter
+            cmds.select("root", hi = True)
+            cmds.filterCurve()
+
+
+            #export selected
+            cmds.progressBar(self.widgets["currentFileProgressBar"], edit = True, step = step)
+            #first change some fbx properties
+            string = "FBXExportConstraints -v 1;"
+            string += "FBXExportCacheFile -v 0;"
+            mel.eval(string)
+
+            cmds.select("root", hi = True)
+            if cmds.objExists("custom_export"):
+                cmds.select("custom_export", add = True)
+
+
+            cmds.file(exportPath, es = True, force = True, prompt = False, type = "FBX export")
+
+            #clean scene
+            cmds.delete("root")
+            if cmds.objExists("custom_export"):
+                cmds.delete("custom_export")
+
+
+
+        #reset rotation interp
+        cmds.optionVar(iv = ("rotationInterpolationDefault", interp))	
+
+
+
+    # # # # # # # # # # # # # # # # # # # # # # # # # # # # # # # # # # # # # # # # # # # # # # # # # # # # # # # # # # # # # # # 
+    # # # # # # # # # # # # # # # # # # # # # # # # # # # # # # # # # # # # # # # # # # # # # # # # # # # # # # # # # # # # # # # 
+    # # # # # # # # # # # # # # # # # # # # # # # # # # # # # # # # # # # # # # # # # # # # # # # # # # # # # # # # # # # # # # # 
+    # # # # # # # # # # # # # # # # # # # # # # # # # # # # # # # # # # # # # # # # # # # # # # # # # # # # # # # # # # # # # # # 
+    # # # # # # # # # # # # # # # # # # # # # # # # # # # # # # # # # # # # # # # # # # # # # # # # # # # # # # # # # # # # # # # 
+    def getCharacters(self):
+
+        referenceNodes = []
+        references = cmds.ls(type = "reference")
+
+
+        for reference in references:
+            niceName = reference.rpartition("RN")[0]
+            suffix = reference.rpartition("RN")[2]
+            if suffix != "":
+                if cmds.objExists(niceName + suffix + ":" + "Skeleton_Settings"):
+                    referenceNodes.append(niceName + suffix)
+
+            else:
+                if cmds.objExists(niceName + ":" + "Skeleton_Settings"):
+                    referenceNodes.append(niceName)
+
+        return referenceNodes
\ No newline at end of file
