Branching

diff --git a/Engine/Extras/Maya_AnimationRiggingTools/ARTv1/General/Scripts/Modules/ART_Spine.py b/Engine/Extras/Maya_AnimationRiggingTools/ARTv1/General/Scripts/Modules/ART_Spine.py
new file mode 100644
--- /dev/null
+++ b/Engine/Extras/Maya_AnimationRiggingTools/ARTv1/General/Scripts/Modules/ART_Spine.py
@@ -0,0 +1,117 @@
+import maya.cmds as cmds
+import maya.mel as mel
+import os
+import ART_rigUtils as utils
+reload(utils)
+
+
+class Spine():
+
+    #pass in items we need to know about (do we build a clavicle? What is the root name of the module?(If set to None, just use upperarm, lowerarm, etc), 
+    def __init__(self, includeClavicle, name, prefix, suffix, armGrpParent, color):
+
+
+        print "test"
+
+
+
+    # # # # # # # # # # # # # # # # # # # # # # # # # # # # # # # # # # # # # # # # # # # # # # # # # # # # # # # # # # # # # # # 
+    # # # # # # # # # # # # # # # # # # # # # # # # # # # # # # # # # # # # # # # # # # # # # # # # # # # # # # # # # # # # # # # 
+    # # # # # # # # # # # # # # # # # # # # # # # # # # # # # # # # # # # # # # # # # # # # # # # # # # # # # # # # # # # # # # # 
+    # # # # # # # # # # # # # # # # # # # # # # # # # # # # # # # # # # # # # # # # # # # # # # # # # # # # # # # # # # # # # # # 
+    # # # # # # # # # # # # # # # # # # # # # # # # # # # # # # # # # # # # # # # # # # # # # # # # # # # # # # # # # # # # # # # 
+    def buildFKSpine(self):
+
+        #find the number of spine bones from the skeleton settings
+        spineJoints = self.getSpineJoints()
+        fkControls = []
+        parent = None
+
+        for joint in spineJoints:
+            if joint == "spine_01":
+                #add space switcher node to base of spine
+                spaceSwitcherFollow = cmds.group(empty = True, name = joint + "_space_switcher_follow")
+                constraint = cmds.parentConstraint(joint, spaceSwitcherFollow)[0]
+                cmds.delete(constraint)
+                spaceSwitcher = cmds.duplicate(spaceSwitcherFollow, name = joint + "_space_switcher")[0]
+                cmds.parent(spaceSwitcher, spaceSwitcherFollow)
+
+            #create an empty group in the same space as the joint
+            group = cmds.group(empty = True, name = joint + "_anim_grp")
+            constraint = cmds.parentConstraint(joint, group)[0]
+            cmds.delete(constraint)
+
+            #create an additional layer of group that has zeroed attrs
+            offsetGroup = cmds.group(empty = True, name = joint + "_anim_offset_grp")
+            constraint = cmds.parentConstraint(joint, offsetGroup)[0]
+            cmds.delete(constraint)
+            cmds.parent(offsetGroup, group)
+
+            #create a control object in the same space as the joint
+            control = utils.createControl("circle", 45, joint + "_anim")
+            tempDupe = cmds.duplicate(control)[0]
+            constraint = cmds.parentConstraint(joint, control)[0]
+            cmds.delete(constraint)
+            fkControls.append(control)
+
+
+            #parent the control object to the group
+            cmds.parent(control, offsetGroup)
+            constraint = cmds.orientConstraint(tempDupe, control, skip = ["x", "z"])[0]
+            cmds.delete(constraint)
+            cmds.makeIdentity(control, t = 1, r = 1, s = 1, apply = True)
+
+
+            #setup hierarchy
+            if parent != None:
+                cmds.parent(group, parent, absolute = True)
+
+            else:
+                cmds.parent(group, spaceSwitcher)
+                cmds.parent(spaceSwitcherFollow, "body_anim")
+
+
+            #set the parent to be the current spine control
+            parent = control
+
+            #clean up
+            cmds.delete(tempDupe)
+
+            for attr in [".v"]:
+                cmds.setAttr(control + attr, lock = True, keyable = False)
+
+            #set the control's color
+            cmds.setAttr(control + ".overrideEnabled", 1)
+            cmds.setAttr(control + ".overrideColor", 18)
+
+
+        return fkControls
+
+
+
+    # # # # # # # # # # # # # # # # # # # # # # # # # # # # # # # # # # # # # # # # # # # # # # # # # # # # # # # # # # # # # # # 
+    # # # # # # # # # # # # # # # # # # # # # # # # # # # # # # # # # # # # # # # # # # # # # # # # # # # # # # # # # # # # # # # 
+    # # # # # # # # # # # # # # # # # # # # # # # # # # # # # # # # # # # # # # # # # # # # # # # # # # # # # # # # # # # # # # # 
+    # # # # # # # # # # # # # # # # # # # # # # # # # # # # # # # # # # # # # # # # # # # # # # # # # # # # # # # # # # # # # # # 
+    # # # # # # # # # # # # # # # # # # # # # # # # # # # # # # # # # # # # # # # # # # # # # # # # # # # # # # # # # # # # # # # 
+    def getSpineJoints(self):
+
+        if cmds.objExists("Skeleton_Settings"):
+            numSpineBones = int(cmds.getAttr("Skeleton_Settings.numSpineBones"))
+            spineJoints = []
+
+            for i in range(int(numSpineBones)):
+
+                if i < 10:
+                    spineJoint = "spine_0" + str(i + 1)
+
+                else:
+                    spineJoint = "spine_" + (i + 1)
+
+                spineJoints.append(spineJoint)
+
+            return spineJoints
+
+        else:
+            cmds.confirmDialog(title = "Error", icon = "critical", message = "Needed node, Skeleton_Settings could not be found.")
+            return
diff --git a/Engine/Extras/Maya_AnimationRiggingTools/ARTv1/General/Scripts/Modules/ART_UpperArmIKTwist.py b/Engine/Extras/Maya_AnimationRiggingTools/ARTv1/General/Scripts/Modules/ART_UpperArmIKTwist.py
new file mode 100644
--- /dev/null
+++ b/Engine/Extras/Maya_AnimationRiggingTools/ARTv1/General/Scripts/Modules/ART_UpperArmIKTwist.py
@@ -0,0 +1,136 @@
+import maya.cmds as cmds
+import Modules.ART_rigUtils as utils
+reload(utils)
+
+#----------------------------------------------------------------------------------------
+# Currently this script is written to work with multiple twists, but if you have more than one the last one is the only one that will work.
+def upperIKTwist(color, suffix, name, prefix, upperJnt, lowerJnt, sysGroup):
+    print "START Building the "+suffix+" "+upperJnt+" UPPER IK twists---------------------------------------"
+
+    #create a nice name
+    name = prefix + name + suffix
+    if name.find("_") == 0:
+        name = name.partition("_")[2]
+        
+        # define the joints for this rig.
+        upperjnt = prefix + upperJnt + suffix
+        lowerjnt = prefix + lowerJnt + suffix
+        #OLD - upperjnt_parent = "driver_" + prefix + "clavicle" + suffix
+        upperjnt_parent = cmds.listRelatives(upperjnt, p=True)[0]
+        driver_upperjnt = "driver_" + prefix + upperJnt + suffix
+        driver_lowerjnt = "driver_" + prefix + lowerJnt + suffix
+
+        numRolls = 0
+        for joint in ["_twist_01", "_twist_02", "_twist_03", "_twist_04", "_twist_05", "_twist_06"]:
+            if cmds.objExists("driver_" + prefix + upperJnt + joint + suffix):
+                numRolls = numRolls + 1
+
+        print "...There are a total of " + str(numRolls) + " to build."
+
+        for i in range(int(numRolls)):
+            print "...Building upper twist_0" + str(i + 1)
+
+            upperjnt_twist = prefix+upperJnt+"_twist_0"+str(i+1)+suffix
+            driver_upperjnt_twist = "driver_"+prefix+upperJnt+"_twist_0"+str(i+1)+suffix
+
+            ####################################
+            # Create a master group for the roll controls to be parented to.
+            rollGrp = cmds.group(empty=True, name=upperjnt+"_roll_grp_0"+str(i+1))
+            cmds.parent(rollGrp, sysGroup)
+            cmds.parentConstraint(upperjnt_parent, rollGrp, mo=True)
+
+            # Create the manual twist control
+            if i == 0:
+                twistCtrl = utils.createControl("circle", 15, upperjnt + "_twist_anim")
+            else:
+                twistCtrl = utils.createControl("circle", 15, upperjnt + "_twist"+str(i+1)+"_anim")
+            cmds.setAttr(twistCtrl + ".ry", -90)
+            cmds.setAttr(twistCtrl + ".sx", 0.8)
+            cmds.setAttr(twistCtrl + ".sy", 0.8)
+            cmds.setAttr(twistCtrl + ".sz", 0.8)
+            cmds.makeIdentity(twistCtrl, r = 1, s = 1, apply =True)
+
+            # move the manual control to the correct location
+            constraint = cmds.parentConstraint(driver_upperjnt_twist, twistCtrl)[0]
+            cmds.delete(constraint)
+
+            # create a group for the manual control and parent the twist to it.
+            twistCtrlGrp = cmds.group(empty = True, name=twistCtrl+"_grp")
+            constraint = cmds.parentConstraint(driver_upperjnt_twist, twistCtrlGrp, mo=False)[0]
+            cmds.delete(constraint)
+            twistCtrlDeltoidOffsetGrp = cmds.group(empty = True, name=twistCtrl+"_deltOffset_grp")
+            constraint = cmds.parentConstraint(driver_upperjnt_twist, twistCtrlDeltoidOffsetGrp, mo=False)[0]
+            cmds.delete(constraint)
+            cmds.parent(twistCtrl, twistCtrlDeltoidOffsetGrp)
+            cmds.parent(twistCtrlDeltoidOffsetGrp, twistCtrlGrp)
+            cmds.parent(twistCtrlGrp, rollGrp)
+            cmds.makeIdentity(twistCtrl, t = 1, r = 1, s = 1, apply = True)
+
+            # set the manual controls visibility settings
+            cmds.setAttr(twistCtrl + ".overrideEnabled", 1)
+            cmds.setAttr(twistCtrl + ".overrideColor", color)
+            for attr in [".sx", ".sy", ".sz"]:
+                cmds.setAttr(twistCtrl + attr, lock = True, keyable = False)
+            cmds.setAttr(twistCtrl + ".v", keyable = False)	
+
+            # add attr on rig settings for manual twist control visibility
+            cmds.select("Rig_Settings")
+            if i == 0:
+                cmds.addAttr(longName=(upperjnt + "_twistCtrlVisUpper"), at = 'bool', dv = 0, keyable = True)
+            cmds.connectAttr("Rig_Settings." + upperjnt + "_twistCtrlVisUpper", twistCtrl + ".v")
+
+            # add attr to rig settings for the twist ammount values
+            cmds.select("Rig_Settings")
+            cmds.addAttr(longName= (upperjnt+"_Twist"+str(i+1)+"Amount"), defaultValue=0.5, minValue=0, maxValue=1, keyable = True)
+            for u in range(int(i+1)):
+                cmds.setAttr("Rig_Settings."+upperjnt+"_Twist"+str(u+1)+"Amount", (1.0/(i+2.0)*((2.0-u)+(i-1.0))) )
+
+            ####################################
+            # Create the twist rig on the shoulder without any extra controls.  Those are added later.
+            cmds.delete(upperjnt_twist+"_orientConstraint1")
+
+            twistIKHandle = cmds.ikHandle(name="twistIKHandle_0"+str(i+1)+"_"+upperjnt, sol="ikRPsolver", sj=upperjnt, ee=upperjnt_twist)[0]
+            print "TWIST IK HANDLE IS: "+twistIKHandle
+            cmds.setAttr(twistIKHandle+".poleVectorX", 0)
+            cmds.setAttr(twistIKHandle+".poleVectorY", 0)
+            cmds.setAttr(twistIKHandle+".poleVectorZ", 0)
+            cmds.parent(twistIKHandle, driver_upperjnt)
+
+            driver_upperjnt_offset = cmds.duplicate(driver_upperjnt, po=True, name="driver_"+upperJnt+"_twistOffset_l")[0]
+            cmds.parent(driver_upperjnt_offset, driver_upperjnt)
+
+            twistOrntConst = cmds.orientConstraint(driver_upperjnt, upperjnt, driver_upperjnt_offset, mo=True)[0]
+            cmds.setAttr(twistOrntConst+".interpType", 2)
+
+            twistMDNode = cmds.shadingNode("multiplyDivide", asUtility=True, name=driver_upperjnt_twist+"_MDnode")
+            cmds.setAttr(twistMDNode+".input2X", 1)
+            cmds.setAttr(twistMDNode+".input2Y", -1)
+
+            cmds.connectAttr(driver_upperjnt_offset+".rotateX", twistMDNode+".input1X", force=True)
+            cmds.connectAttr(driver_upperjnt+".rotateX", twistMDNode+".input1Y", force=True)
+            cmds.connectAttr(twistMDNode+".outputX", driver_upperjnt_twist+".rx", force=True)            
+            cmds.connectAttr(twistMDNode+".outputY", twistCtrlDeltoidOffsetGrp+".rx", force=True)            
+
+            cmds.pointConstraint(driver_upperjnt_twist, twistCtrlGrp, mo=True)
+            cmds.orientConstraint(driver_upperjnt_twist, twistCtrlGrp, mo=True)
+
+            cmds.orientConstraint(twistCtrl, upperjnt_twist, mo=True)
+            
+            if cmds.objExists("deltoid_anim"+suffix):
+                if suffix == "_l":
+                    cmds.connectAttr("deltoid_anim"+suffix+".rotateX", twistIKHandle+".twist", force=True)
+                else:
+                    deltMDNode = cmds.shadingNode("multiplyDivide", asUtility=True, name=twistIKHandle+"_MDNode")
+                    cmds.setAttr(deltMDNode+".input2X", -1)
+                    cmds.connectAttr("deltoid_anim"+suffix+".rotateX", deltMDNode+".input1X", force=True)
+                    cmds.connectAttr(deltMDNode+".outputX", twistIKHandle+".twist", force=True)
+                    
+            reverseMDNode = cmds.shadingNode("reverse", asUtility=True, name=driver_upperjnt_offset+"_Reversenode")
+            cmds.connectAttr("Rig_Settings."+upperjnt+"_Twist"+str(i+1)+"Amount", reverseMDNode+".inputX", force=True)
+            cmds.connectAttr(reverseMDNode+".outputX", twistOrntConst+"."+upperjnt+"W1", force=True)            
+            cmds.connectAttr("Rig_Settings."+upperjnt+"_Twist"+str(i+1)+"Amount", twistOrntConst+"."+driver_upperjnt+"W0", force=True)
+                          
+
+    print ".END Building the "+suffix+" "+upperJnt+" UPPER twists---------------------------------------"            
+
+#upperIKTwist(6, "_l", "", "", "upperarm", "lowerarm", "arm_sys_grp")
diff --git a/Engine/Extras/Maya_AnimationRiggingTools/ARTv1/General/Scripts/Modules/ART_UpperArmTwist.py b/Engine/Extras/Maya_AnimationRiggingTools/ARTv1/General/Scripts/Modules/ART_UpperArmTwist.py
new file mode 100644
--- /dev/null
+++ b/Engine/Extras/Maya_AnimationRiggingTools/ARTv1/General/Scripts/Modules/ART_UpperArmTwist.py
@@ -0,0 +1,176 @@
+import maya.cmds as cmds
+import maya.mel as mel
+import os
+import ART_rigUtils as utils
+reload(utils)
+
+#----------------------------------------------------------------------------------------
+# Currently this script is written to work with multiple twists, but if you have more than one the last one is the only one that will work.
+def upperArmTwist(color, rollGrpParent, fkArmJoint, ikArmJoint, suffix, name, prefix, upperArm, lowerArm):
+    print "START Building the "+ suffix +" UPPER arm twists---------------------------------------"
+
+    #create a nice name
+    name = prefix + name + suffix
+    if name.find("_") == 0:
+        name = name.partition("_")[2]
+        
+        # define the joints for this rig.
+        upperarm = prefix + upperArm + suffix
+        lowerarm = prefix + lowerArm + suffix
+        driver_clavicle = "driver_" + prefix + "clavicle" + suffix
+        driver_upperarm = "driver_deltoid_" + prefix + upperArm + suffix
+        driver_lowerarm = "driver_" + prefix + lowerArm + suffix
+
+        numRolls = 0
+        for joint in ["_twist_01", "_twist_02", "_twist_03", "_twist_04", "_twist_05", "_twist_06"]:
+            if cmds.objExists("driver_" + prefix + upperArm + joint + suffix):
+                numRolls = numRolls + 1
+
+        print "...There are a total of " + str(numRolls) + " to build."
+            
+        for i in range(int(numRolls)):
+            print "...Building upper arm twist_0" + str(i + 1)
+            driver_upperarm_twist = "driver_"+prefix+upperArm+"_twist_0"+str(i+1)+suffix
+
+            rollGrp = cmds.group(empty=True, name=upperarm+"_roll_grp_0"+str(i+1))
+            cmds.parent(rollGrp, "arm_sys_grp")
+
+            # Move the driver_upperarm_twist joint to its correct position between the driver_upperarm and its driver_lowerarm using a point constraint.
+            const = cmds.parentConstraint(driver_upperarm, driver_upperarm_twist, n=driver_upperarm+"_temp_parentConst", weight=1)
+            cmds.delete(const)
+            const = cmds.pointConstraint(driver_upperarm, driver_lowerarm, driver_upperarm_twist, n=driver_upperarm+"_temp_PointConst", weight=1)
+            cmds.delete(const)
+
+            # Duplicate the driver_upperarm, driver_lowerarm, and driver_upperarm_twist joints
+            blendJnts=[]
+            blend_driver_upperarmParent = cmds.duplicate(driver_upperarm, n="Blend_"+driver_upperarm+"_parent_ignore_0"+str(i+1), ic=True, po=True)[0]
+            blendJnts.append(blend_driver_upperarmParent)
+
+            blend_driver_lowerarm = cmds.duplicate(driver_lowerarm, n="Blend_"+driver_lowerarm+"_ignore_0"+str(i+1), ic=True, po =True)[0]
+            blendJnts.append(blend_driver_lowerarm)
+
+            blend_driver_upperarm = cmds.duplicate(driver_upperarm, n="Blend_"+driver_upperarm+"_ignore_0"+str(i+1), po=True)[0]
+            blendJnts.append(blend_driver_upperarm)
+
+            blend_driver_upperarm_twist = cmds.duplicate(driver_upperarm_twist, n="Blend_"+driver_upperarm_twist+"_ignore_0"+str(i+1), ic=True, po= True)[0]
+            blendJnts.append(blend_driver_upperarm_twist)
+
+            cmds.parent(blend_driver_lowerarm, blend_driver_upperarm_twist, blend_driver_upperarmParent)
+            cmds.parent(blend_driver_upperarm, blend_driver_upperarm_twist)
+
+            # Set the parent bone to have zyx rotation order
+            cmds.setAttr(blend_driver_upperarmParent+".rotateOrder", 2)
+
+            originalJnts = [driver_upperarm, driver_lowerarm]
+            #print originalJnts
+            
+            # Now we are going to disconnect the constraints that were driving the original driver_upperarm 
+            # and driver_lowerarm joints and re-connect them to the parent_ignore and ignore joints that were duped from them.
+            # This is so that we can create new constraints on the driver_upperarm and driver_lowerarm joints for the rig.
+            j = 0
+            a = 0
+
+            #while a < len(originalJnts):
+            for a in range(0,len(originalJnts)):
+                # Go through the inputs of the original joints and disconnect them
+                origConn = cmds.listConnections(originalJnts[a], d=False, c=+True, p=True)
+
+                for j in range(0,len(origConn),2):
+                    cmds.disconnectAttr(origConn[j+1], origConn[j])
+                    nameList = origConn[j+1].split(".")[0]
+                    destConn = cmds.listConnections(nameList, s=False)
+                    # Find out if the node that was an input to the original joint is also an output
+                    if originalJnts[a] in destConn:
+                        sourceConn = cmds.listConnections(nameList, d=False, c=True, p=True)
+                        for k in range(0,len(sourceConn),2):
+                            # Get the input connections to the node that are connected to the original joint
+                            nameConnList = sourceConn[k+1].split(".")
+                            if (nameConnList[0] == originalJnts[a]):
+                                # Disconnect from the original joint and connect to the blend joint
+                                cmds.disconnectAttr(sourceConn[k+1], sourceConn[k])
+                                nameConnList[0] = blendJnts[a]
+                                cmds.connectAttr(nameConnList[0] + "." + nameConnList[-1], sourceConn[k])
+
+            # create the manual twist control
+            if i == 0:
+                twistCtrl = utils.createControl("circle", 15, upperarm + "_twist_anim")
+            else:
+                twistCtrl = utils.createControl("circle", 15, upperarm + "_twist"+str(i+1)+"_anim")
+            cmds.setAttr(twistCtrl + ".ry", -90)
+            cmds.setAttr(twistCtrl + ".sx", 0.8)
+            cmds.setAttr(twistCtrl + ".sy", 0.8)
+            cmds.setAttr(twistCtrl + ".sz", 0.8)
+            cmds.makeIdentity(twistCtrl, r = 1, s = 1, apply =True)
+
+            # move the manual control to the correct location
+            constraint = cmds.parentConstraint(blend_driver_upperarm_twist, twistCtrl)[0]
+            cmds.delete(constraint)
+
+            # create a group for the manual control and parent the twist to it.
+            twistCtrlGrp = cmds.group(empty = True, name = twistCtrl + "_grp")
+            constraint = cmds.parentConstraint(blend_driver_upperarm_twist, twistCtrlGrp)[0]
+            cmds.delete(constraint)
+            cmds.parent(twistCtrl, twistCtrlGrp)
+            cmds.parent(twistCtrlGrp, rollGrp)
+            cmds.makeIdentity(twistCtrl, t = 1, r = 1, s = 1, apply = True)
+            cmds.parentConstraint(blend_driver_upperarm_twist, twistCtrlGrp)
+
+            # set the manual controls visibility settings
+            cmds.setAttr(twistCtrl + ".overrideEnabled", 1)
+            cmds.setAttr(twistCtrl + ".overrideColor", color)
+            for attr in [".sx", ".sy", ".sz"]:
+                cmds.setAttr(twistCtrl + attr, lock = True, keyable = False)
+            cmds.setAttr(twistCtrl + ".v", keyable = False)	
+            
+            # add attr on rig settings for manual twist control visibility
+            cmds.select("Rig_Settings")
+            if i == 0:
+                cmds.addAttr(longName=(name + "twistCtrlVisUpper"), at = 'bool', dv = 0, keyable = True)
+            cmds.connectAttr("Rig_Settings." + name + "twistCtrlVisUpper", twistCtrl + ".v")
+                        
+            # add attr to rig settings for the twist ammount values
+            cmds.select("Rig_Settings")
+            cmds.addAttr(longName= (name+"UpperarmTwist"+str(i+1)+"Amount"), defaultValue=0.5, minValue=0, maxValue=1, keyable = True)
+            for u in range(int(i+1)):
+                cmds.setAttr("Rig_Settings."+name+"UpperarmTwist"+str(u+1)+"Amount", (1.0/(i+2.0)*((2.0-u)+(i-1.0))) )
+
+            cmds.parent(blend_driver_upperarmParent, rollGrp)
+
+            # Constrain the original joints to the blend joints
+            cmds.orientConstraint(blend_driver_upperarm, driver_upperarm, n=driver_upperarm+"_0"+str(i+1)+"_OrntCnst", mo=True, w=1)
+            cmds.pointConstraint(blend_driver_upperarm, driver_upperarm, n=driver_upperarm+"_0"+str(i+1)+"_PtCnst", mo=True, w=1)
+
+            cmds.orientConstraint(blend_driver_lowerarm, driver_lowerarm, n=driver_lowerarm+"_0"+str(i+1)+"_OrntCnst", mo=True, w=1 )
+            cmds.pointConstraint(blend_driver_lowerarm, driver_lowerarm, n=driver_lowerarm+"_0"+str(i+1)+"_PtCnst", mo=True, w=1)
+
+            cmds.orientConstraint(twistCtrl, driver_upperarm_twist, n=driver_upperarm_twist+"_0"+str(i+1)+"_OrntCnst", mo=True, w=1)
+            cmds.pointConstraint(twistCtrl, driver_upperarm_twist, n=driver_upperarm_twist+"_0"+str(i+1)+"_PtCnst", mo=True, w=1)
+
+            # Create a driver_upperarm_twist multiplier to multiply the driver_upperarm_twist values by -.5 to counter rotate the driver_upperarm_twist and driver_upperarm back from the overall limb's driver_upperarm_twist amount.
+            twistMult = cmds.shadingNode("multiplyDivide", asUtility=True, name=blend_driver_upperarm_twist+"_0"+str(i+1)+"_MDnode")
+            twistMultDrive = cmds.shadingNode("multiplyDivide", asUtility=True, name=blend_driver_upperarm_twist+"_0"+str(i+1)+"_Drive_MDnode")
+            twistMultDriveInv = cmds.shadingNode("multiplyDivide", asUtility=True, name=blend_driver_upperarm_twist+"_0"+str(i+1)+"_DriveInv_MDnode")
+            twistMultMultiplier = cmds.shadingNode("multiplyDivide", asUtility=True, name=blend_driver_upperarm_twist+"_0"+str(i+1)+"_Multiplier_MDnode")
+            twistMultMultiplierInv = cmds.shadingNode("multiplyDivide", asUtility=True, name=blend_driver_upperarm_twist+"_0"+str(i+1)+"_MultiplierInv_MDnode")
+            twistReverse = cmds.shadingNode("reverse", asUtility=True, name=blend_driver_upperarm_twist+"_0"+str(i+1)+"_REVnode")
+
+            cmds.connectAttr(blend_driver_upperarmParent+".rotateX", twistMult+".input1X", force=True)
+            cmds.connectAttr("Rig_Settings."+name+"UpperarmTwist"+str(i+1)+"Amount", twistReverse+".inputX")
+            cmds.connectAttr("Rig_Settings."+name+"UpperarmTwist"+str(i+1)+"Amount", twistMultMultiplierInv+".input1X")
+            cmds.connectAttr(twistReverse+".outputX", twistMultMultiplier+".input1X", force=True)
+            cmds.connectAttr(twistMult+".outputX", twistMultDrive+".input1X", force=True)
+            cmds.connectAttr(twistMultMultiplier+".outputX", twistMultDrive+".input2X", force=True)
+            cmds.connectAttr(twistMult+".outputX", twistMultDriveInv+".input1X", force=True)
+            cmds.connectAttr(twistMultMultiplierInv+".outputX", twistMultDriveInv+".input2X", force=True)
+            cmds.connectAttr(twistMultDrive+".outputX", blend_driver_upperarm_twist+".rotateX", force=True)
+            cmds.connectAttr(twistMultDriveInv+".outputX", blend_driver_upperarm+".rotateX", force=True)
+            cmds.setAttr(twistMult+".input2X", -0.5)
+            cmds.setAttr(twistMult+".input2X", l=True)
+
+            cmds.setAttr(twistMultMultiplier+".input2X", 2)
+            cmds.setAttr(twistMultMultiplierInv+".input2X", 2)
+            
+            cmds.select(blend_driver_upperarmParent)
+            
+            cmds.parentConstraint(driver_clavicle, rollGrp, mo=True)
+    print ".END Building the "+ suffix +" UPPER arm twists---------------------------------------"
\ No newline at end of file
diff --git a/Engine/Extras/Maya_AnimationRiggingTools/ARTv1/General/Scripts/Modules/ART_UpperLegIKTwist.py b/Engine/Extras/Maya_AnimationRiggingTools/ARTv1/General/Scripts/Modules/ART_UpperLegIKTwist.py
new file mode 100644
--- /dev/null
+++ b/Engine/Extras/Maya_AnimationRiggingTools/ARTv1/General/Scripts/Modules/ART_UpperLegIKTwist.py
@@ -0,0 +1,129 @@
+import maya.cmds as cmds
+import Modules.ART_rigUtils as utils
+reload(utils)
+
+#----------------------------------------------------------------------------------------
+# Currently this script is written to work with multiple twists, but if you have more than one the last one is the only one that will work.
+def upperIKTwist(color, suffix, name, prefix, upperJnt, lowerJnt, sysGroup):
+    print "START Building the "+suffix+" "+upperJnt+" UPPER IK twists---------------------------------------"
+
+    #create a nice name
+    name = prefix + name + suffix
+    if name.find("_") == 0:
+        name = name.partition("_")[2]
+        
+        # define the joints for this rig.
+        upperjnt = prefix + upperJnt + suffix
+        lowerjnt = prefix + lowerJnt + suffix
+        #OLD - upperjnt_parent = "driver_" + prefix + "clavicle" + suffix
+        upperjnt_parent = cmds.listRelatives(upperjnt, p=True)[0]
+        driver_upperjnt = "driver_" + prefix + upperJnt + suffix
+        driver_lowerjnt = "driver_" + prefix + lowerJnt + suffix
+
+        numRolls = 0
+        for joint in ["_twist_01", "_twist_02", "_twist_03", "_twist_04", "_twist_05", "_twist_06"]:
+            if cmds.objExists("driver_" + prefix + upperJnt + joint + suffix):
+                numRolls = numRolls + 1
+
+        print "...There are a total of " + str(numRolls) + " to build."
+
+        for i in range(int(numRolls)):
+            print "...Building upper twist_0" + str(i + 1)
+
+            upperjnt_twist = prefix+upperJnt+"_twist_0"+str(i+1)+suffix
+            driver_upperjnt_twist = "driver_"+prefix+upperJnt+"_twist_0"+str(i+1)+suffix
+
+            ####################################
+            # Create a master group for the roll controls to be parented to.
+            rollGrp = cmds.group(empty=True, name=upperjnt+"_roll_grp_0"+str(i+1))
+            cmds.parent(rollGrp, sysGroup)
+            cmds.parentConstraint(upperjnt_parent, rollGrp, mo=True)
+
+            # Create the manual twist control
+            if i == 0:
+                twistCtrl = utils.createControl("circle", 15, upperjnt + "_twist_anim")
+            else:
+                twistCtrl = utils.createControl("circle", 15, upperjnt + "_twist"+str(i+1)+"_anim")
+            cmds.setAttr(twistCtrl + ".ry", -90)
+            cmds.setAttr(twistCtrl + ".sx", 0.8)
+            cmds.setAttr(twistCtrl + ".sy", 0.8)
+            cmds.setAttr(twistCtrl + ".sz", 0.8)
+            cmds.makeIdentity(twistCtrl, r = 1, s = 1, apply =True)
+
+            # move the manual control to the correct location
+            constraint = cmds.parentConstraint(driver_upperjnt_twist, twistCtrl)[0]
+            cmds.delete(constraint)
+
+            # create a group for the manual control and parent the twist to it.
+            twistCtrlGrp = cmds.group(empty = True, name=twistCtrl+"_grp")
+            constraint = cmds.parentConstraint(driver_upperjnt_twist, twistCtrlGrp)[0]
+            cmds.delete(constraint)
+            cmds.parent(twistCtrl, twistCtrlGrp)
+            cmds.parent(twistCtrlGrp, rollGrp)
+            cmds.makeIdentity(twistCtrl, t = 1, r = 1, s = 1, apply = True)
+
+            # set the manual controls visibility settings
+            cmds.setAttr(twistCtrl + ".overrideEnabled", 1)
+            cmds.setAttr(twistCtrl + ".overrideColor", color)
+            for attr in [".sx", ".sy", ".sz"]:
+                cmds.setAttr(twistCtrl + attr, lock = True, keyable = False)
+            cmds.setAttr(twistCtrl + ".v", keyable = False)	
+
+            # add attr on rig settings for manual twist control visibility
+            cmds.select("Rig_Settings")
+            if i == 0:
+                cmds.addAttr(longName=(upperjnt + "_twistCtrlVisUpper"), at = 'bool', dv = 0, keyable = True)
+            cmds.connectAttr("Rig_Settings." + upperjnt + "_twistCtrlVisUpper", twistCtrl + ".v")
+
+            # add attr to rig settings for the twist ammount values
+            cmds.select("Rig_Settings")
+            cmds.addAttr(longName= (upperjnt+"_Twist"+str(i+1)+"Amount"), defaultValue=0.5, minValue=0, maxValue=1, keyable = True)
+            for u in range(int(i+1)):
+                cmds.setAttr("Rig_Settings."+upperjnt+"_Twist"+str(u+1)+"Amount", (1.0/(i+2.0)*((2.0-u)+(i-1.0))) )
+
+            ####################################
+            # Create the twist rig on the shoulder without any extra controls.  Those are added later.
+            cmds.delete(upperjnt_twist+"_orientConstraint1")
+
+            twistIKHandle = cmds.ikHandle(name="twistIKHandle_0"+str(i+1)+"_"+upperjnt, sol="ikRPsolver", sj=upperjnt, ee=upperjnt_twist)[0]
+            print "TWIST IK HANDLE IS: "+twistIKHandle
+            cmds.setAttr(twistIKHandle+".poleVectorX", 0)
+            cmds.setAttr(twistIKHandle+".poleVectorY", 0)
+            cmds.setAttr(twistIKHandle+".poleVectorZ", 0)
+            cmds.parent(twistIKHandle, driver_upperjnt)
+
+            driver_upperjnt_offset = cmds.duplicate(driver_upperjnt, po=True, name="driver_"+upperJnt+"_twistOffset_l")[0]
+            cmds.parent(driver_upperjnt_offset, driver_upperjnt)
+
+            twistOrntConst = cmds.orientConstraint(driver_upperjnt, upperjnt, driver_upperjnt_offset, mo=True)[0]
+            cmds.setAttr(twistOrntConst+".interpType", 2)
+
+            twistMDNode = cmds.shadingNode("multiplyDivide", asUtility=True, name=driver_upperjnt_twist+"_MDnode")
+            cmds.setAttr(twistMDNode+".input2X", 1)
+
+            cmds.connectAttr(driver_upperjnt_offset+".rotateX", twistMDNode+".input1X", force=True)
+            cmds.connectAttr(twistMDNode+".outputX", driver_upperjnt_twist+".rx", force=True)            
+
+            cmds.pointConstraint(driver_upperjnt_twist, twistCtrlGrp, mo=True)
+            cmds.orientConstraint(driver_upperjnt_twist, twistCtrlGrp, mo=True)
+
+            cmds.orientConstraint(twistCtrl, upperjnt_twist, mo=True)
+            
+            if cmds.objExists("deltoid_anim"+suffix):
+                if suffix == "_l":
+                    cmds.connectAttr("deltoid_anim"+suffix+".rotateX", twistIKHandle+".twist", force=True)
+                else:
+                    deltMDNode = cmds.shadingNode("multiplyDivide", asUtility=True, name=twistIKHandle+"_MDNode")
+                    cmds.setAttr(deltMDNode+".input2X", -1)
+                    cmds.connectAttr("deltoid_anim"+suffix+".rotateX", deltMDNode+".input1X", force=True)
+                    cmds.connectAttr(deltMDNode+".outputX", twistIKHandle+".twist", force=True)
+
+            reverseMDNode = cmds.shadingNode("reverse", asUtility=True, name=driver_upperjnt_offset+"_Reversenode")
+            cmds.connectAttr("Rig_Settings."+upperjnt+"_Twist"+str(i+1)+"Amount", reverseMDNode+".inputX", force=True)
+            cmds.connectAttr(reverseMDNode+".outputX", twistOrntConst+"."+upperjnt+"W1", force=True)            
+            cmds.connectAttr("Rig_Settings."+upperjnt+"_Twist"+str(i+1)+"Amount", twistOrntConst+"."+driver_upperjnt+"W0", force=True)
+
+
+    print ".END Building the "+suffix+" "+upperJnt+" UPPER twists---------------------------------------"            
+
+#upperIKTwist(6, "_l", "", "", "upperarm", "lowerarm", "arm_sys_grp")
diff --git a/Engine/Extras/Maya_AnimationRiggingTools/ARTv1/General/Scripts/Modules/ART_rigUtils.py b/Engine/Extras/Maya_AnimationRiggingTools/ARTv1/General/Scripts/Modules/ART_rigUtils.py
new file mode 100644
--- /dev/null
+++ b/Engine/Extras/Maya_AnimationRiggingTools/ARTv1/General/Scripts/Modules/ART_rigUtils.py
@@ -0,0 +1,154 @@
+import maya.cmds as cmds
+import maya.mel as mel
+import os
+
+#rigging utilities
+
+
+# # # # # # # # # # # # # # # # # # # # # # # # # # # # # # # # # # # # # # # # # # # # # # # # # # # # # # # # # # # # # # # 
+# # # # # # # # # # # # # # # # # # # # # # # # # # # # # # # # # # # # # # # # # # # # # # # # # # # # # # # # # # # # # # # 
+# # # # # # # # # # # # # # # # # # # # # # # # # # # # # # # # # # # # # # # # # # # # # # # # # # # # # # # # # # # # # # # 
+# # # # # # # # # # # # # # # # # # # # # # # # # # # # # # # # # # # # # # # # # # # # # # # # # # # # # # # # # # # # # # # 
+# # # # # # # # # # # # # # # # # # # # # # # # # # # # # # # # # # # # # # # # # # # # # # # # # # # # # # # # # # # # # # # 
+def createControl(controlType, size, name):
+    
+    scale = getScaleFactor()
+            
+    if controlType == "circle":
+        control = cmds.circle(c = (0,0,0), sw = 360, r = size * scale, d = 3, name = name)[0]
+        
+    if controlType == "circleSpecial":
+        control = cmds.circle(c = (0,0,0), sw = 360, r = 1, d = 3, name = name)[0]
+        
+        side = name.rpartition("_")[2]
+        
+        if side == "l":
+            cmds.xform(control, piv = (0, -1, 0))
+            
+        else:
+            cmds.xform(control, piv = (0, 1, 0))
+            
+        cmds.setAttr(control + ".scaleX", size * scale)
+        cmds.setAttr(control + ".scaleY", size * scale)
+        cmds.setAttr(control + ".scaleZ", size * scale)
+
+    if controlType == "square":
+        control = cmds.circle(c = (0,0,0), s = 4, sw = 360, r = size * scale, d = 1, name = name)[0]
+        cmds.setAttr(control + ".rz", 45)
+        
+    if controlType == "foot":
+        control = cmds.curve(name = name, d = 3, p = [(0, 40, 0), (-3.42, 39, 0), (-10.2, 37, 0), (-13, 22, 0), (-15.7, 13.2, 0), (-20, -14, 0), (-18.1, -25.6, 0), (-15, -44.8, 0), (1.1, -41.2, 0), (4.8, -41.7, 0), (15.5, -31.9, 0), (16.9, -22.7, 0), (18.6, -15.2, 0), (16.5, -.5, 0), (11.2, 29.2, 0), (10.7, 39.7, 0), (3.6, 39.9, 0), (0, 40, 0)])
+        
+        footLoc = cmds.spaceLocator(name = (name + "_end_loc"))[0]
+        cmds.parent(footLoc, control)
+        cmds.setAttr(footLoc + ".ty", -40)
+        cmds.setAttr(footLoc + ".v", 0)
+                     
+        cmds.setAttr(control + ".scaleX", size * scale)
+        cmds.setAttr(control + ".scaleY", size * scale)
+        cmds.setAttr(control + ".scaleZ", size * scale)
+        
+    
+    if controlType == "arrow":
+        control = cmds.curve(name = name, d = 1, p = [(0, -45, 0), (5, -45, 0), (5, -62, 0), (10, -62, 0), (0, -72, 0), (-10, -62, 0), (-5, -62, 0), (-5, -45, 0), (0, -45, 0)])
+        cmds.xform(control, cp = True)
+        cmds.setAttr(control + ".ty", 58.5)
+        cmds.makeIdentity(control, t = 1, apply = True)
+        cmds.xform(control, piv = (0, 13.5, 0))
+        
+        cmds.setAttr(control + ".scaleX", size * scale)
+        cmds.setAttr(control + ".scaleY", size * scale)
+        cmds.setAttr(control + ".scaleZ", size * scale)
+        
+    if controlType == "arrowOnBall":
+        
+        control = cmds.curve(name = name, d = 1, p = [(0.80718, 0.830576, 8.022739), (0.80718, 4.219206, 7.146586 ), (0.80718, 6.317059, 5.70073), (2.830981, 6.317059, 5.70073), (0, 8.422749, 2.94335), (-2.830981, 6.317059, 5.70073), (-0.80718, 6.317059, 5.70073), (-0.80718, 4.219352, 7.146486), (-0.80718, 0.830576, 8.022739), (-4.187851, 0.830576, 7.158003), (-6.310271, 0.830576, 5.705409), (-6.317059, 2.830981, 5.7007), (-8.422749, 0, 2.94335), (-6.317059, -2.830981, 5.70073), (-6.317059, -0.830576, 5.70073), (-4.225134, -0.830576, 7.142501), (-0.827872, -0.830576, 8.017446), (-0.80718, -4.176512, 7.160965), (-0.80718, -6.317059, 5.70073), (-2.830981, -6.317059, 5.70073), (0, -8.422749, 2.94335), (2.830981, -6.317059, 5.70073), (0.80718, -6.317059, 5.70073), (0.80718, -4.21137, 7.151987), (0.80718, -0.830576, 8.022739), (4.183345, -0.830576, 7.159155), (6.317059, -0.830576, 5.70073), (6.317059, -2.830981, 5.70073), (8.422749, 0, 2.94335), (6.317059, 2.830981, 5.70073), (6.317059, 0.830576, 5.70073), (4.263245, 0.830576, 7.116234), (0.80718, 0.830576, 8.022739)])
+        
+        cmds.setAttr(control + ".scaleX", size * scale)
+        cmds.setAttr(control + ".scaleY", size * scale)
+        cmds.setAttr(control + ".scaleZ", size * scale)
+        
+        
+    if controlType == "semiCircle":
+        
+        control = cmds.curve(name = name, d = 3, p = [(0,0,0), (7, 0, 0), (8, 0, 0), (5, 4, 0), (0, 5, 0), (-5, 4, 0), (-8, 0, 0), (-7, 0, 0), (0,0,0)])
+        cmds.xform(control, ws = True, t = (0, 5, 0))
+        cmds.xform(control, ws = True, piv = (0,0,0))
+        cmds.makeIdentity(control, t = 1, apply = True)
+        
+        
+        cmds.setAttr(control + ".scaleX", size * scale)
+        cmds.setAttr(control + ".scaleY", size * scale)
+        cmds.setAttr(control + ".scaleZ", size * scale)
+        
+        
+    if controlType == "pin":
+        control = cmds.curve(name = name, d = 1, p = [(12,0,0), (0, 0, 0), (-12, -12, 0), (-12, 12, 0), (0, 0, 0)])
+        cmds.xform(control, ws = True, piv = [12,0,0])
+        cmds.setAttr(control + ".scaleY", .5)
+        cmds.makeIdentity(control, t = 1, apply = True)
+        
+        
+        cmds.setAttr(control + ".scaleX", size * scale)
+        cmds.setAttr(control + ".scaleY", size * scale)
+        cmds.setAttr(control + ".scaleZ", size * scale)
+        
+    
+    if controlType == "sphere":
+
+        points = [(0, 0, 1), (0, 0.5, 0.866), (0, 0.866025, 0.5), (0, 1, 0), (0, 0.866025, -0.5), (0, 0.5, -0.866025), (0, 0, -1), (0, -0.5, -0.866025), (0, -0.866025, -0.5), (0, -1, 0), (0, -0.866025, 0.5), (0, -0.5, 0.866025), (0, 0, 1), (0.707107, 0, 0.707107), (1, 0, 0), (0.707107, 0, -0.707107), (0, 0, -1), (-0.707107, 0, -0.707107), (-1, 0, 0), (-0.866025, 0.5, 0), (-0.5, 0.866025, 0), (0, 1, 0), (0.5, 0.866025, 0), (0.866025, 0.5, 0), (1, 0, 0), (0.866025, -0.5, 0), (0.5, -0.866025, 0), (0, -1, 0), (-0.5, -0.866025, 0), (-0.866025, -0.5, 0), (-1, 0, 0), (-0.707107, 0, 0.707107), (0, 0, 1)]
+        control = cmds.curve(name = name, d = 1, p = points)
+        
+        cmds.setAttr(control + ".scaleX", size * scale)
+        cmds.setAttr(control + ".scaleY", size * scale)
+        cmds.setAttr(control + ".scaleZ", size * scale)
+        
+
+    return control    
+
+
+# # # # # # # # # # # # # # # # # # # # # # # # # # # # # # # # # # # # # # # # # # # # # # # # # # # # # # # # # # # # # # # 
+# # # # # # # # # # # # # # # # # # # # # # # # # # # # # # # # # # # # # # # # # # # # # # # # # # # # # # # # # # # # # # # 
+# # # # # # # # # # # # # # # # # # # # # # # # # # # # # # # # # # # # # # # # # # # # # # # # # # # # # # # # # # # # # # # 
+# # # # # # # # # # # # # # # # # # # # # # # # # # # # # # # # # # # # # # # # # # # # # # # # # # # # # # # # # # # # # # # 
+# # # # # # # # # # # # # # # # # # # # # # # # # # # # # # # # # # # # # # # # # # # # # # # # # # # # # # # # # # # # # # # 
+def getScaleFactor():
+    
+    headLoc = cmds.spaceLocator(name = "headLoc")[0]
+    cmds.parentConstraint("driver_head", headLoc)
+    
+    height = cmds.getAttr(headLoc + ".tz")
+    
+    defaultHeight = 400
+    scaleFactor = height/defaultHeight
+    cmds.delete(headLoc)
+    
+    return scaleFactor    
+
+# # # # # # # # # # # # # # # # # # # # # # # # # # # # # # # # # # # # # # # # # # # # # # # # # # # # # # # # # # # # # # # 
+# # # # # # # # # # # # # # # # # # # # # # # # # # # # # # # # # # # # # # # # # # # # # # # # # # # # # # # # # # # # # # # 
+# # # # # # # # # # # # # # # # # # # # # # # # # # # # # # # # # # # # # # # # # # # # # # # # # # # # # # # # # # # # # # # 
+# # # # # # # # # # # # # # # # # # # # # # # # # # # # # # # # # # # # # # # # # # # # # # # # # # # # # # # # # # # # # # # 
+# # # # # # # # # # # # # # # # # # # # # # # # # # # # # # # # # # # # # # # # # # # # # # # # # # # # # # # # # # # # # # # 
+def getUpAxis(obj):
+    
+    cmds.xform(obj, ws = True, relative = True, t = [0, 0, 10])
+    translate = cmds.getAttr(obj + ".translate")[0]
+    newTuple = (abs(translate[0]), abs(translate[1]), abs(translate[2]))
+    cmds.xform(obj, ws = True, relative = True, t = [0, 0, -10])
+    
+    highestVal = max(newTuple)
+    axis = newTuple.index(highestVal)
+    upAxis = None
+    
+    if axis == 0:
+        upAxis = "X"
+        
+    if axis == 1:
+        upAxis = "Y"
+        
+    if axis == 2:
+        upAxis = "Z"
+        
+    return upAxis
+    
\ No newline at end of file
diff --git a/Engine/Extras/Maya_AnimationRiggingTools/ARTv1/General/Scripts/Modules/facial/face.py b/Engine/Extras/Maya_AnimationRiggingTools/ARTv1/General/Scripts/Modules/facial/face.py
new file mode 100644
--- /dev/null
+++ b/Engine/Extras/Maya_AnimationRiggingTools/ARTv1/General/Scripts/Modules/facial/face.py
@@ -0,0 +1,695 @@
+'''
+FACIAL MODULE
+
+'''
+
+import os, sys
+import traceback
+import json
+
+from PySide import QtGui, QtCore
+reload(QtGui)
+reload(QtCore)
+
+import maya.cmds as cmds
+
+# internal libs
+from Modules.facial import utils
+
+##################
+# BASE
+##################
+
+
+class BaseFace(object):
+    '''
+    This class is the base to be inherited and does not serialize to anything in Maya
+    '''
+    def __init__(self, version=None, create=True, uType='uMeta'):
+        pass
+
+    # FUNCTIONS
+    ##############
+
+    # PROPERTIES
+    ###############
+
+    #get and set controllers
+    @property
+    def controls(self):
+        return cmds.listConnections(self.node + '.controls')
+    @controls.setter
+    def controls(self, controls):
+        for c in controls:
+            utils.msgConnect(self.node + '.controls', c + '.facialNode')
+
+    #get and set version
+    @property
+    def version(self):
+        return cmds.getAttr(self.node+'.version')
+    @version.setter
+    def version(self, ver):
+        cmds.setAttr(self.node+'.version', ver, type='string')
+
+    #get and set uType identifier
+    @property
+    def uType(self):
+        return cmds.getAttr(self.node+'.uType')
+    @uType.setter
+    def uType(self, type):
+        cmds.setAttr(self.node+'.uType', type, type='string')
+
+    #Build attributes on the node if it's being built
+    def buildBase(self):
+        cmds.addAttr(self.node, longName='version', dataType='string')
+        cmds.addAttr(self.node, longName='controls', attributeType='message')
+        cmds.addAttr(self.node, longName='uType', dataType='string')
+
+
+
+##################
+## MASK
+##################
+class FaceMask(BaseFace):
+    '''
+    This class serializes to a network node
+    '''
+    def __init__(self, maskGeoNode=None, faceModuleNode=None, maskNode=None, **kwargs):
+        super(FaceMask, self).__init__(**kwargs)
+        self.maskGeoNode = maskGeoNode
+        self.faceModuleNode = faceModuleNode
+
+        if maskNode:
+            self.node = maskNode
+            self.initialize()
+
+        elif 'create' in kwargs:
+            self.buildNode()
+
+    ## FUNCTIONS
+    ##############
+
+    def bakeMovers(self):
+        pass
+
+    def zeroMovers(self):
+        for mover in self.faceModule.jointMovers:
+            cmds.xform()
+
+    def snapMovers(self, axis, mesh, debug=1):
+        '''
+        for this system, in the tool UI itself that snaps the movers, I always bidirectionally cast against Y
+        '''
+        try:
+            cmds.undoInfo(openChunk=True)
+            #iterate through the *active* facial joint movers and snap them to the surface
+            if debug: print 'activeJointMovers:', self.faceModule.activeJointMovers
+            for mover in self.faceModule.activeJointMovers:
+                if debug: print 'snapping', mover
+                #do not snap manually placed movers
+                if not cmds.getAttr(mover + '.alignmentType'):
+
+                    #set the axis vector
+                    axisVector = None
+                    if axis == 'x': axisVector = utils.getAxisWorldVectors(mover)[0]
+                    elif axis == 'y': axisVector = utils.getAxisWorldVectors(mover)[1]
+                    elif axis == 'z': axisVector = utils.getAxisWorldVectors(mover)[2]
+
+                    pos = cmds.xform(mover, q=1, rp=1, ws=1)
+
+                    #raycast in and out and find the closest hit of the facial surface
+                    pos1 = utils.rayMeshIntersect(mesh, pos, axisVector)
+                    pos2 = utils.rayMeshIntersect(mesh, pos, -axisVector)
+                    dist1 = utils.square_distance(pos, pos1)
+                    dist2 = utils.square_distance(pos, pos2)
+                    newPos = None
+                    if dist1 <= dist2:
+                        newPos = pos1
+                    else:
+                        newPos = pos2
+
+                    cmds.xform(mover, t=newPos, ws=1)
+                    if debug: print mover, 'snapped.'
+        except Exception as e:
+            print e
+        finally:
+            cmds.undoInfo(closeChunk=True)
+
+    def deleteSdkParentConstraints(self):
+        try:
+            cmds.undoInfo(openChunk=True)
+            #update facial SDK locations to current joint mover placement
+            print '>>>deleteSdkParentConstraints'
+            for mover in self.faceModule.jointMovers:
+                sdk = cmds.listConnections(mover + '.sdk')[0]
+                sdkParent = cmds.listRelatives(sdk, parent = True)[0]
+                utils.deleteConstraints(sdkParent)
+        except Exception as e:
+            print(traceback.format_exc())
+        finally:
+            cmds.undoInfo(closeChunk=True)
+
+    def buildSdkParentConstraints(self):
+        try:
+            cmds.undoInfo(openChunk=True)
+            for mover in self.faceModule.jointMovers:
+                sdk = cmds.listConnections(mover + '.sdk')[0]
+                lra = cmds.listConnections(mover + '.lra')[0]
+                sdkParent = cmds.listRelatives(sdk, parent = True)[0]
+                #TODO: Maybe don't maintain offset later (?)
+                # there are some relationships that have offsets (eyes) that should be maintained
+                cmds.parentConstraint(lra, sdkParent, mo=1)
+        except Exception as e:
+            print(traceback.format_exc())
+        finally:
+            cmds.undoInfo(closeChunk=True)
+
+    def toggleMaskActive(self, mode):
+        try:
+            cmds.undoInfo(openChunk=True)
+            if mode:
+                cmds.showHidden(self.controls)
+                cmds.showHidden(self.maskGeo)
+
+                #hide all movers
+                cmds.hide(self.faceModule.jointMovers)
+                #TODO: Need to check if the constraints were deleted and re-add them
+                '''for mover in self.faceModule.jointMovers:
+                    sdk = cmds.listConnections(mover + '.sdk')[0]
+                    sdkParent = cmds.listRelatives(sdk, parent = True)[0]
+                    lra = cmds.listConnections(mover + '.lra')[0]
+                    cmds.parentConstraint(lra, sdkParent)'''
+                print 'MASK: ACTIVATED'
+
+            else:
+                #show joint movers
+                cmds.showHidden(self.faceModule.jointMovers)
+
+                #hide the mask controls and mesh
+                cmds.hide(self.controls)
+                cmds.hide(self.maskGeo)
+                print 'MASK: DEACTIVATED'
+        except Exception as e:
+            print(traceback.format_exc())
+        finally:
+            cmds.undoInfo(closeChunk=True)
+
+
+
+    ## BUILD
+    def buildNode(self):
+        self.node = cmds.createNode('network', name='maskNode')
+        self.buildBase()
+
+        #maskGeo attr and connection if it was passed in
+        cmds.addAttr(self.node, longName='maskGeo', attributeType='message')
+        if self.maskGeoNode:
+            self.maskGeo = self.maskGeoNode
+
+        #FaceModule attr and connection if it was passed in
+        cmds.addAttr(self.node, longName='faceModule', attributeType='message')
+        if self.faceModuleNode:
+            self.faceModule = self.faceModuleNode
+
+        #Set other metadata
+        self.version = '0.00'
+        self.uType = 'faceMask'
+
+    ##SERIALIZATION
+    def initialize(self):
+        pass
+
+
+    ## PROPERTIES
+    ###############
+    @property
+    def maskGeo(self):
+        return cmds.listConnections(self.node + '.maskGeo')[0]
+    @maskGeo.setter
+    def maskGeo(self, maskG):
+        utils.msgConnect(self.node + '.maskGeo', maskG + '.maskNode')
+
+    @property
+    def faceModule(self):
+        return FaceModule(faceNode=cmds.listConnections(self.node + '.faceModule')[0])
+    @faceModule.setter
+    def faceModule(self, faceM):
+        utils.msgConnect(faceM + '.maskNode', self.node + '.faceModule')
+
+    @property
+    def attachLocations(self):
+        return cmds.listConnections(self.node + '.attachLocations')
+    @attachLocations.setter
+    def attachLocations(self, locs):
+        for loc in locs:
+            utils.msgConnect(self.node + '.attachLocations', loc + '.maskNode')
+
+    @property
+    def active(self):
+        return cmds.listConnections(self.node + '.active')
+    @active.setter
+    def active(self, mode):
+        if not utils.attrExists(self.node + '.active'):
+            cmds.addAttr(self.node, longName='active', at='bool')
+        cmds.setAttr(self.node+'.active', mode)
+
+        self.toggleMaskActive(mode)
+
+
+####################
+### FACE
+####################
+class FaceModule(BaseFace):
+    '''
+    This class serializes to a network node
+    '''
+    def __init__(self, renderMeshNodes=None, faceNode=None, **kwargs):
+        super(FaceModule, self).__init__(**kwargs)
+
+        if faceNode:
+            self.node = faceNode
+            self.initialize()
+
+        elif 'create' in kwargs:
+            self.buildNode()
+        else:
+            cmds.error('No facial node set through faceNode, and no create flag used!\nLOCALS: ' + str(locals()))
+
+
+    ## FUNCTIONS
+    ##############
+
+    def initialize(self):
+        self.mask = self.faceMask
+
+    def readDrivenDict(self, node):
+        if utils.attrExists(node + '.drivenDict'):
+            dict = cmds.getAttr(node + '.drivenDict')
+            if dict:
+                return json.loads(dict)
+        else:
+            cmds.warning('No drivenDict found on object: ' + node)
+            return {}
+
+    def updateConstraints(self, node):
+        '''
+        This is basically made for the method 'snapEyeSdkPivots' below.
+        It is a general function, except the part where is pops the sdkNode connection.
+        '''
+        parent = cmds.listRelatives(node, parent=1)[0]
+        listConnectionsSucks = cmds.listConnections(parent, type='parentConstraint')
+        if listConnectionsSucks:
+            for const in set(listConnectionsSucks):
+                conns = list(set(cmds.listConnections(const, type='transform')))
+                conns.pop(conns.index(const))
+                conns.pop(conns.index(parent))
+                if utils.attrExists(const + '.sdkNode'):
+                    conns.pop(conns.index(cmds.listConnections(const + '.sdkNode')[0]))
+                if len(conns) == 1:
+                    cmds.parentConstraint(conns[0], parent, e=1, mo=1)
+
+    def snapEyeSdkPivots(self):
+        L_eyeSDKs = [u'l_eye_lid_upper_mid_SDK', u'l_eye_lid_upper_inner_SDK', u'l_eye_lid_upper_outer_SDK', u'l_eye_lid_lower_outer_SDK', u'l_eye_lid_lower_mid_SDK', u'l_eye_lid_lower_inner_SDK']
+        R_eyeSDKs = [u'r_eye_lid_upper_outer_SDK', u'r_eye_lid_upper_mid_SDK', u'r_eye_lid_upper_inner_SDK', u'r_eye_lid_lower_inner_SDK', u'r_eye_lid_lower_mid_SDK', u'r_eye_lid_lower_outer_SDK']
+        for node in L_eyeSDKs:
+            utils.snapPivotWithZeroParent(node, 'l_eye_SDK')
+            self.updateConstraints(node)
+        for node in R_eyeSDKs:
+            utils.snapPivotWithZeroParent(node, 'r_eye_SDK')
+            self.updateConstraints(node)
+
+    ## BUILD
+    def buildNode(self):
+        #create and pass the 'node' to the base
+        self.node = cmds.createNode('network', name='faceNode')
+        #call the base build on that node
+        self.buildBase()
+
+        #Set other useful metadata
+        self.version = '0.00'
+        self.uType = 'faceModule'
+
+        cmds.addAttr(self.node, longName='jointMovers', attributeType='message')
+        cmds.addAttr(self.node, longName='sdks', attributeType='message')
+
+    def stampDeltas(self):
+        try:
+            cmds.undoInfo(openChunk=True)
+            bb = self.blendBoard
+            if bb:
+                #build attr
+                if not utils.attrExists(bb + '.initialPoseDict'):
+                    cmds.addAttr(bb, longName='initialPoseDict', dt='string')
+                    cmds.setAttr(bb + '.initialPoseDict', '{}', type='string')
+
+                initialPoseDict={}
+
+                drivenDict = self.readDrivenDict(bb)
+                for drivingAttr in drivenDict:
+                    attrDict = {}
+                    for driven in drivenDict[drivingAttr]:
+                        cmds.setAttr(drivingAttr, 1.0)
+                        attrDict[driven] = cmds.getAttr(driven)
+                    initialPoseDict[drivingAttr] = attrDict
+                    cmds.setAttr(drivingAttr, 0.0)
+
+                cmds.setAttr(bb + '.initialPoseDict', json.dumps(initialPoseDict), type='string')
+
+            else:
+                cmds.warning('No blendBoard/poses found!')
+        except Exception as e:
+            print(traceback.format_exc())
+            cmds.undoInfo(closeChunk=True)
+
+
+    def buildBlendBoard(self, fpath):
+        try:
+            cmds.undoInfo(openChunk=True)
+
+            #TODO: rebuild and keep all connections and attrs
+            if self.blendBoard:
+                cmds.delete(self.blendBoard)
+
+            lines = []
+            if os.path.exists(fpath):
+                f = open(fpath, 'r')
+                lines = f.readlines()
+                f.close()
+            self.blendBoard = cmds.spaceLocator(name='blendBoard')[0]
+            cmds.addAttr(self.blendBoard, longName='uType', dataType='string')
+            cmds.setAttr(self.blendBoard + '.uType', 'blendBoard', type='string')
+            utils.msgConnect(self.node + '.blendBoard', self.blendBoard + '.faceModule')
+
+            #store input file path
+            cmds.addAttr(self.blendBoard, longName='inputFile', dataType='string')
+            cmds.setAttr(self.blendBoard + '.inputFile', fpath, type='string')
+
+            #make purdy
+            cmds.setAttr("blendBoard.tx", l=1, k=0, channelBox=0)
+            cmds.setAttr("blendBoard.ty", l=1, k=0, channelBox=0)
+            cmds.setAttr("blendBoard.tz", l=1, k=0, channelBox=0)
+            cmds.setAttr("blendBoard.rx", l=1, k=0, channelBox=0)
+            cmds.setAttr("blendBoard.ry", l=1, k=0, channelBox=0)
+            cmds.setAttr("blendBoard.rz", l=1, k=0, channelBox=0)
+            cmds.setAttr("blendBoard.sx", l=1, k=0, channelBox=0)
+            cmds.setAttr("blendBoard.sy", l=1, k=0, channelBox=0)
+            cmds.setAttr("blendBoard.sz", l=1, k=0, channelBox=0)
+            cmds.setAttr("blendBoard.v", l=1, k=0, channelBox=0)
+
+            #make pose attrs
+            for line in lines:
+                line = line.strip()
+                if line != '':
+                    print line
+                    if '##' in line:
+                        cmds.addAttr(self.blendBoard,enumName='#####', longName=line.replace('##',''), at='enum', k=1)
+                        cmds.setAttr(self.blendBoard + '.' + line.replace('##',''), lock=1)
+                    else:
+                        cmds.addAttr(self.blendBoard, longName=line, at='float', min=0, max=1, k=1)
+
+            cmds.undoInfo(closeChunk=True)
+        except Exception as e:
+            exc_type, exc_value, exc_tb = sys.exc_info()
+            for line in traceback.format_exception(exc_type, exc_value, exc_tb):
+                print line
+            cmds.undoInfo(closeChunk=True)
+
+    def createJointMover(self, name=None, pos=(0,0,0), matrix=None, scale=1, toroidRatio=0.02, ss=0, facialNode=None):
+        try:
+            cmds.undoInfo(openChunk=True)
+            if not name:
+                name = 'test'
+            radius = 1
+
+            #create mask attachment node
+            maskAttach = cmds.spaceLocator(name=name+'_mover_maskAttach')[0]
+
+            #create jeremy-joint-mover-stack
+            topNode = cmds.group(name=name+'_mover_grp', em=1)
+            mover = utils.createCircleControl(name=name+'_mover', radius=2, type='toroid', heightRatio=toroidRatio)
+            #offset = utils.createCircleControl(name=name+'_mover_offset', radius=1.5, color='xyz', type='toroid', heightRatio=toroidRatio)
+            offset = cmds.group(name=name+'_mover_offset', em=1)
+            lra_grp = cmds.group(name=name+'_lra_grp', em=1)
+            lra = utils.createLRA(name=name + '_LRA', length=2, thickness=0.05, arrowRadiusMult=3)
+            cmds.parent(mover, topNode)
+            cmds.parent(offset, mover)
+            cmds.parent(lra_grp, offset)
+            cmds.parent(lra, lra_grp)
+            cmds.displaySmoothness(mover, f=1)
+            if scale is not 1:
+                cmds.scale(scale, scale, scale, topNode)
+                cmds.makeIdentity(topNode, apply=1, t=0, r=0, s=1, n=0, pn=1)
+            if not ss:
+                cmds.select(topNode)
+            cmds.parentConstraint(maskAttach, topNode)
+
+            #create SDK pose stack
+            sdkOffset = cmds.group(name=name+'_SDK_grp', em=1)
+            sdk = utils.createLRA(name=name + '_SDK', length=2, color=False, thickness=0.05, arrowRadiusMult=3)
+            tweak = cmds.group(name=name+'_TWEAK', em=1)
+            if scale is not 1:
+                cmds.scale(scale, scale, scale, sdk)
+                cmds.makeIdentity(sdk, apply=1, t=0, r=0, s=1, n=0, pn=1)
+            cmds.parent(tweak, sdk)
+            cmds.parent(sdk, sdkOffset)
+            cmds.parentConstraint(lra, sdkOffset)
+            utils.toggleShapeSelection(sdk, False)
+
+            #cleanup outliner
+            if cmds.objExists('maskAttachNodes'):
+                cmds.parent(maskAttach, 'maskAttachNodes')
+            if cmds.objExists('sdkNodes'):
+                cmds.parent(sdkOffset, 'sdkNodes')
+            if cmds.objExists('jointMovers'):
+                cmds.parent(topNode, 'jointMovers')
+
+            #make internal attrs/connections
+            utils.msgConnect(mover + '.lra', lra + '.jointMover')
+            utils.msgConnect(mover + '.attachNode', maskAttach + '.jointMover')
+            utils.msgConnect(mover + '.sdk', sdk + '.jointMover')
+
+            #make external attrs/connections
+            if facialNode:
+                locs = None
+
+                #mark attachLocations
+                try:
+                    locs = self.mask.attachLocations
+                    locs.extend(maskAttach)
+                    self.mask.attachLocations = locs
+                except Exception as e:
+                    print(traceback.format_exc())
+
+                #hang the jm off the facial node
+                utils.msgConnect(self.node + '.jointMovers', mover + '.facialNode')
+
+                #mark jointMovers
+                jm = self.jointMovers
+                jm.append(mover)
+                self.jointMovers = jm
+
+                #mark sdks
+                s = self.sdks
+                s.append(sdk)
+                self.sdks = s
+
+            #stamp with alignment type
+            cmds.addAttr(sdk, ln='alignmentType', enumName='mask surface:manual',  k=1, at='enum')
+
+            cmds.undoInfo(closeChunk=True)
+            cmds.select(maskAttach)
+            return [topNode, lra, offset, maskAttach]
+        except Exception as e:
+            cmds.undoInfo(closeChunk=True)
+            import sys
+            exc_type, exc_obj, exc_tb = sys.exc_info()
+            fname = os.path.split(exc_tb.tb_frame.f_code.co_filename)[1]
+            print(e, fname, exc_tb.tb_lineno)
+
+    def findMirrorJointMoverByName(self, jm):
+        mirrorNode = None
+        if not jm.lower().startswith('c_'):
+            if jm.lower().startswith('r_'):
+                mirrorNode = 'l' + jm[1:]
+            else:
+                mirrorNode = 'r' + jm[1:]
+            if mirrorNode:
+                if cmds.objExists(mirrorNode):
+                    return mirrorNode
+            else:
+                cmds.warning('Cannot find l_ or r_ prefix in: ' + jm)
+                return False
+
+    def mirrorFacialXform(self, jm, axis='x', debug=0):
+        mirror = self.findMirrorJointMoverByName(jm)
+        if debug:
+            print mirror
+        if mirror:
+            pos = cmds.xform(jm, q=1, ws=1, rp=1)
+            rot = cmds.xform(jm, q=1, ws=1, ro=1)
+            #scale = cmds.xform(jm, q=1, r=1, s=1)
+
+            cmds.select(mirror)
+            if axis == 'x':
+                cmds.move(-pos[0], pos[1], pos[2], mirror,  ws=1, rpr=1 )
+                cmds.rotate(rot[0], -rot[1], -rot[2], mirror,  ws=1, a=1 )
+            else:
+                cmds.warning('Currently works for faces created looking down y axis.')
+
+    def connectControlBoard(self, controlBoard):
+        bb = self.blendBoard
+        utils.msgConnect(self.node + '.controlBoard', controlBoard + '.faceModule')
+        if bb:
+            for attr in cmds.listAttr(bb, k=1):
+                bbAttr = bb + '.' + attr
+                if cmds.listConnections(bbAttr):
+                    cbAttr = controlBoard + '.' + attr
+                    if utils.attrExists(cbAttr):
+                        cmds.connectAttr(cbAttr, bbAttr, f=1)
+
+    def connectSdksToFinalSkeleton(self):
+        jmLocked = cmds.lockNode('JointMover', lock=1, q=1)[0]
+        if jmLocked:
+            cmds.lockNode('JointMover', lock=0)
+        for mover in self.activeJointMovers:
+            #TODO: this won't work with multiple heads
+            #hate to use name string here, but there's only one skeleton
+            sdk = cmds.listConnections(mover + '.sdk')[0]
+            joint = mover.replace('_mover','')
+            tweak = cmds.listRelatives(sdk, children=1, typ='transform')[0]
+            cmds.parentConstraint(tweak, joint, mo=1)
+            cmds.scaleConstraint(tweak, joint)
+
+            #to get away from sting names later
+            #TODO: detect locks and only unlock/relock if the item was locked
+            moverLocked = cmds.lockNode(mover, lock=1, q=1)[0]
+            if moverLocked:
+                cmds.lockNode(mover, lock=0)
+            utils.msgConnect(mover + '.deformingJoint', joint + '.mover')
+            if moverLocked:
+                cmds.lockNode(mover, lock=1)
+        if jmLocked:
+            cmds.lockNode('JointMover', lock=1)
+
+
+    ## PROPERTIES
+    ###############
+    @property
+    def renderMeshes(self):
+        return cmds.listConnections(self.node + '.renderMeshes')
+    @renderMeshes.setter
+    def renderMeshes(self, meshes):
+        for mesh in meshes:
+            utils.msgConnect(self.node + '.renderMeshes', mesh + '.faceNode')
+
+    @property
+    def faceMask(self):
+        return FaceMask(maskNode=cmds.listConnections(self.node + '.maskNode')[0])
+    @faceMask.setter
+    def faceMask(self, maskNode):
+        utils.msgConnect(self.node + '.faceMask', maskNode + '.faceNode')
+
+    @property
+    def blendBoard(self):
+        bb = cmds.listConnections(self.node + '.blendBoard')
+        if bb:
+            return cmds.listConnections(self.node + '.blendBoard')[0]
+        else:
+            return []
+    @blendBoard.setter
+    def blendBoard(self, bBoard):
+        utils.msgConnect(self.node + '.blendBoard', bBoard + '.faceNode')
+
+    @property
+    def jointMovers(self):
+        conns = cmds.listConnections(self.node + '.jointMovers')
+        if conns:
+            return conns
+        else: return []
+    @jointMovers.setter
+    def jointMovers(self, movers):
+        for m in movers:
+            utils.msgConnect(self.node + '.jointMovers', m + '.facialNode')
+
+    @property
+    def sdks(self):
+        conns = cmds.listConnections(self.node + '.sdks')
+        if conns:
+            return conns
+        else: return []
+    @sdks.setter
+    def sdks(self, sdks):
+        for sdk in sdks:
+            utils.msgConnect(self.node + '.sdks', sdk + '.facialNode')
+
+    @property
+    def blendBoardConnections(self):
+        if self.blendBoard:
+            poses = []
+            bb = self.blendBoard
+            for att in cmds.listAttr(bb, k=1):
+                conns = cmds.listConnections(bb + '.' + att)
+                if conns:
+                    poses.append(att)
+            return poses
+        else:
+            cmds.warning('No blendBoard connected!')
+            return None
+    @blendBoardConnections.setter
+    def blendBoardConnections(self, *args):
+        print 'Attribute cannot be set: connections are derived from the blendBoard'
+
+    @property
+    def poses(self):
+        bb = self.blendBoard
+        return self.readDrivenDict(bb)
+    @poses.setter
+    def poses(self, poses):
+        cmds.warning('Attribute cannot be set: poses are derived from the blendBoard drivenDict attr')
+
+    @property
+    def activeJointMovers(self, debug=0):
+        conns = cmds.listConnections(self.node + '.jointMovers')
+        activeCons = []
+        if conns:
+            for con in conns:
+                if utils.shapesHidden(con) == False:
+                    if debug: print con, 'is not hidden and active'
+                    activeCons.append(con)
+                else:
+                    if debug: print con, 'shapes hidden! not added.'
+            return activeCons
+        else:
+            return []
+    @activeJointMovers.setter
+    def activeJointMovers(self, *args):
+        print 'Attribute cannot be set: connections are derived from the facial joint mover visibility'
+        return False
+
+
+####################
+### TESTING
+####################
+def autoTest(build=True):
+    #build
+    if build:
+        face = FaceModule(create=1)
+        mask = FaceMask(create=1, maskGeoNode='maskGeo', faceModuleNode=face.node)
+        mask.controls = [u'jaw', u'lips_CTRL_MASK', u'R_nose', u'nose_CTRL_MASK', u'L_nose', u'l_lip_corner_CTRL_MASK', u'r_lip_corner_CTRL_MASK', u'R_mask_CTRL_MASK', u'brow_ridge_CTRL_MASK', u'R_eye_LRA', u'L_eye_LRA', u'L_mask_CTRL_MASK', u'top_mask_CTRL_MASK']
+        face.renderMeshes = ['steele']
+        '''
+        print face.faceMask
+        print face.version
+        print mask.maskGeo
+        print mask.faceModule
+        '''
+    else:
+        face = FaceModule(faceNode='faceNode')
+        mask = FaceMask(faceModuleNode='faceNode')
+
+
+#autoTest()
+
+
+if __name__ == '__main__':
+    show()
diff --git a/Engine/Extras/Maya_AnimationRiggingTools/ARTv1/General/Scripts/Modules/facial/facialMaskTool.py b/Engine/Extras/Maya_AnimationRiggingTools/ARTv1/General/Scripts/Modules/facial/facialMaskTool.py
new file mode 100644
--- /dev/null
+++ b/Engine/Extras/Maya_AnimationRiggingTools/ARTv1/General/Scripts/Modules/facial/facialMaskTool.py
@@ -0,0 +1,249 @@
+'''
+FACIAL TOOL
+'''
+
+import os
+
+#uic related
+import shiboken
+from PySide import QtGui, QtCore
+reload(QtGui)
+reload(QtCore)
+from cStringIO import StringIO
+import xml.etree.ElementTree as xml
+import pysideuic
+
+import maya.cmds as cmds
+import maya.OpenMayaUI as openMayaUI
+import maya.mel as mel
+
+#internal libs
+from Modules.facial import utils
+reload(utils)
+from Modules.facial import jointMover_UI as jm_ui
+reload(jm_ui)
+from Modules.facial import poseEditor_UI as pose_ui
+reload(pose_ui)
+from Modules.facial import face
+reload(face)
+from Modules.facial import mask_align_UI as ma_ui
+reload(jm_ui)
+
+########
+## UIC
+########
+def show():
+    global facialMaskToolWindow
+    try:
+        facialMaskToolWindow.close()
+    except:
+        pass
+    facialMaskToolWindow = facialMaskTool()
+    facialMaskToolWindow.show()
+    return facialMaskToolWindow
+
+def loadUiType(uiFile):
+    """
+    Pyside lacks the "loadUiType" command, so we have to convert the ui file to py code in-memory first
+    and then execute it in a special frame to retrieve the form_class.
+    http://tech-artists.org/forum/showthread.php?3035-PySide-in-Maya-2013
+    """
+    parsed = xml.parse(uiFile)
+    widget_class = parsed.find('widget').get('class')
+    form_class = parsed.find('class').text
+
+    with open(uiFile, 'r') as f:
+        o = StringIO()
+        frame = {}
+
+        pysideuic.compileUi(f, o, indent=0)
+        pyc = compile(o.getvalue(), '<string>', 'exec')
+        exec pyc in frame
+
+        #Fetch the base_class and form class based on their type in the xml from designer
+        form_class = frame['Ui_%s'%form_class]
+        base_class = eval('QtGui.%s'%widget_class)
+    return form_class, base_class
+
+def getMayaWindow():
+    ptr = openMayaUI.MQtUtil.mainWindow()
+    if ptr is not None:
+        return shiboken.wrapInstance(long(ptr), QtGui.QWidget)
+
+uiFile = None	
+try:
+    selfDirectory = os.path.dirname(__file__)
+    uiFile = selfDirectory + '/facialMaskTool.ui'
+except:
+    uiFile = 'D:\\Build\\usr\\jeremy_ernst\\MayaTools\\General\\Scripts\\Modules\\facial\\facialMaskTool.ui'
+if os.path.isfile(uiFile):
+    form_class, base_class = loadUiType(uiFile)
+else:
+    cmds.error('Cannot find UI file: ' + uiFile)
+
+
+
+##############
+## MASK TOOL
+##############
+class facialMaskTool(base_class, form_class):	  
+    def __init__(self, parent=getMayaWindow()):
+        self.closeExistingWindow()
+        super(facialMaskTool, self).__init__(parent)
+        
+        self.setupUi(self)
+        
+        wName = openMayaUI.MQtUtil.fullName(long(shiboken.getCppPointer(self)[0]))
+        
+        self.currentFace = None
+        
+        ## Connect UI
+        ########################################################################
+        self.createJointMoverBTN.clicked.connect(self.createJointMoverFn)
+        self.showHideAttachBTN.clicked.connect(self.showHideAttachFn)
+        
+        self.openJointMoverDlgBTN.clicked.connect(self.openJointMoverDlgFn)
+        self.openPoseDlgBTN.clicked.connect(self.openPoseDlgFn)
+        self.openMaskAlignmentDlgBTN.clicked.connect(self.openMaskAlignmentDlgFn)
+        
+        self.selectFaceBTN.clicked.connect(self.selectFaceFn)
+        self.selectMaskBTN.clicked.connect(self.selectMaskFn)
+        self.maskModeBTN.clicked.connect(self.maskModeFn)
+        self.constrainToMaskBTN.pressed.connect(self.constrainToMaskFn)
+        self.buildBlendBoardBTN.pressed.connect(self.buildBlendBoardFn)
+        self.stampDeltasBTN.pressed.connect(self.stampDeltasFn)
+        self.mirrorPlacementBTN.pressed.connect(self.mirrorPlacementFn)
+        self.deleteSdkParentConstraintsBTN.pressed.connect(self.deleteSdkParentConstraintsFn)
+
+        self.refreshUI()
+    
+    ##UI
+    ######
+    
+    def stampDeltasFn(self):
+        self.currentFace.stampDeltas()
+    
+    def createJointMoverFn(self):
+        inputName, ok = QtGui.QInputDialog.getText(self, 'Create Facial Joint Mover', 'Base Name:')
+        if ok:
+            self.currentFace.createJointMover(name=inputName, scale=0.3,toroidRatio=0.05, facialNode=self.currentFace.node)
+        self.refreshUI()
+    
+    def showHideAttachFn(self):
+        if self.showHideAttachBTN.isChecked():
+            print 'hiding'
+            self.showHideAttachBTN.setText('Show Mask Attach Nodes')
+        else:
+            print 'showing'
+            self.showHideAttachBTN.setText('Hide Mask Attach Nodes')
+            
+    def closeExistingWindow(self):
+        for qt in QtGui.qApp.topLevelWidgets():
+            try:
+                if qt.__class__.__name__ == self.__class__.__name__:
+                    qt.deleteLater()
+                    print 'uExport: Closed ' + str(qt.__class__.__name__)
+            except:
+                pass
+    
+    def buildBlendBoardFn(self):
+        fName, ok = QtGui.QFileDialog.getOpenFileName(self, "Select Poses file")
+        if ok:
+            self.currentFace.buildBlendBoard(fName)
+    
+    def constrainToMaskFn(self):
+        ma = cmds.ls(sl=1)[0]
+        if '_maskAttach' in ma:
+            utils.constrainToMesh(self.currentFace.faceMask.maskGeo, ma)
+        else:
+            cmds.warning('constrainToMask: Node does not appear to be a maskAttach nmode: ' + ma)
+    
+    def openJointMoverDlgFn(self):
+        global jointMoverWin
+        try:
+            jointMoverWin.close()
+        except:
+            pass
+        
+        jointMoverWin = jm_ui.jointMoverWin()
+        jointMoverWin.show()
+    
+    def openPoseDlgFn(self):
+        global poseWin
+        try:
+            poseWin.close()
+        except:
+            pass
+        
+        poseWin = pose_ui.poseEditorWin()
+        poseWin.show()
+
+    def openMaskAlignmentDlgFn(selfself):
+        reload(jm_ui)
+        global maskAlignWin
+        try:
+            maskAlignWin.close()
+        except:
+            pass
+
+        maskAlignWin = ma_ui.maskAlignWin()
+        maskAlignWin.show()
+
+    def selectFaceFn(self):
+        cmds.select(self.currentFace.node)
+        
+    def selectMaskFn(self):
+        cmds.select(self.currentFace.faceMask)
+        
+    def maskModeFn(self):
+        if self.maskModeBTN.isChecked():
+            self.currentFace.faceMask.active = True
+            self.maskModeBTN.setText('Mask Mode: Active')
+            print 'MASK MODE: ACTIVE'
+        else:
+            self.currentFace.faceMask.active = False
+            self.maskModeBTN.setText('Mask Mode: Disabled')
+            print 'MASK MODE: DISABLED'
+    
+    def mirrorPlacementFn(self):
+        try:
+            cmds.undoInfo(openChunk=True)
+            for node in cmds.ls(sl=1):
+                prefix = node.lower()[:1]
+                if prefix == 'l':
+                    mirrorNode = 'r_' + node[2:]
+                    if not cmds.objExists(mirrorNode):
+                        print 'Creating mirrorNode:', mirrorNode
+                        inputName = mirrorNode.split('_mover_')[0]
+                        nodes = self.currentFace.createJointMover(name=inputName, scale=0.3,toroidRatio=0.05, facialNode=self.currentFace.node)
+                        
+                    self.currentFace.mirrorFacialXform(node, axis='x', debug=0)
+            cmds.undoInfo(closeChunk=True)
+        except Exception as e:
+            print e
+            cmds.undoInfo(closeChunk=True)
+
+    def deleteSdkParentConstraintsFn(self):
+        if self.deleteSdkParentConstraintsBTN.isChecked():
+            print '>>>buildSdkParentConstraints'
+            self.currentFace.mask.buildSdkParentConstraints()
+            self.deleteSdkParentConstraintsBTN.setText('deleteSdkConstraints')
+        else:
+            print '>>>deleteSdkParentConstraints'
+            self.currentFace.mask.deleteSdkParentConstraints()
+            self.deleteSdkParentConstraintsBTN.setText('buildSdkConstraints')
+            
+    def refreshUI(self):
+        for f in utils.getUType('faceModule'):
+            self.currentFacialNodeCMB.clear()
+            self.currentFacialNodeCMB.addItem(f)
+            self.currentFace = face.FaceModule(faceNode=f)
+            self.maskControlsLBL.setText('Mask Controls: ' + str(len(self.currentFace.faceMask.controls)))
+            if self.currentFace.jointMovers:
+                self.facialMoversLBL.setText('Facial Joint Movers: ' + str(len(self.currentFace.jointMovers)))
+            self.maskGeoLBL.setText('Mask Geo: ' + self.currentFace.faceMask.maskGeo)
+            if self.currentFace.sdks:
+                self.drivenNodesLBL.setText('Driven Nodes: ' + str(len(self.currentFace.sdks)))
+
+if __name__ == '__main__':
+    show()
diff --git a/Engine/Extras/Maya_AnimationRiggingTools/ARTv1/General/Scripts/Modules/facial/jointMover_UI.py b/Engine/Extras/Maya_AnimationRiggingTools/ARTv1/General/Scripts/Modules/facial/jointMover_UI.py
new file mode 100644
--- /dev/null
+++ b/Engine/Extras/Maya_AnimationRiggingTools/ARTv1/General/Scripts/Modules/facial/jointMover_UI.py
@@ -0,0 +1,182 @@
+'''
+from Modules.facial import jointMover_UI as jm_ui
+reload(jm_ui)
+
+global jointMoverWin
+try:
+    jointMoverWin.close()
+except:
+    pass
+jointMoverWin = jm_ui.jointMoverWin()
+jointMoverWin.show()
+'''
+
+from PySide import QtCore, QtGui
+
+import maya.cmds as cmds
+import maya.OpenMayaUI as mui
+import shiboken
+
+import os
+
+from Modules.facial import face
+from Modules.facial import utils
+
+from Modules.facial import mask_align_UI as ma_ui
+reload(ma_ui)
+
+
+def getMayaWindow():
+    ptr = mui.MQtUtil.mainWindow()
+    if ptr is not None:
+        return shiboken.wrapInstance(long(ptr), QtGui.QWidget)
+
+class jointMoverWin(QtGui.QMainWindow):
+    def __init__(self, parent=getMayaWindow(), debug=0):
+        QtGui.QMainWindow.__init__(self, parent)
+        
+        self.currentFace = None
+        
+        #facial folder
+        facialDir = os.path.dirname(__file__)
+
+        for f in utils.getUType('faceModule'):
+            print 'jointMover_UI: Face initialized from:', f
+            self.currentFace = face.FaceModule(faceNode=f)
+        
+        ### BUILD USER INTERFACE
+        ##############################
+        
+        #create/set a central widget
+        wid = QtGui.QWidget()
+        self.setCentralWidget(wid)
+        
+        #setArt v2 stylesheet
+        utils.setArtV2StyleSheet(wid, imageDirectory='/art2images')
+        
+        self.resize(315, 514)
+        
+        #top face module picker
+        self.label = QtGui.QLabel(self)
+        self.label.setText("FacialNode:")
+        self.horizontalLayout = QtGui.QHBoxLayout()
+        self.horizontalLayout.setSpacing(2)
+        self.label.setMaximumSize(QtCore.QSize(78, 16777215))
+        self.horizontalLayout.addWidget(self.label)
+        self.currentFacialNodeCMB = QtGui.QComboBox(self)
+        self.horizontalLayout.addWidget(self.currentFacialNodeCMB)
+        
+        #attachment joint
+        self.h2 = QtGui.QHBoxLayout()
+        self.attachJointLBL = QtGui.QLabel(self)
+        self.attachJointLBL.setText("Face Attachment Joint:")
+        self.h2.setSpacing(2)
+        self.h2.addWidget(self.attachJointLBL)
+        self.attachJointLINE = QtGui.QLineEdit(self)
+        self.h2.addWidget(self.attachJointLINE)
+        
+        #main vertical layout
+        self.verticalLayout = QtGui.QVBoxLayout(wid)
+        self.verticalLayout.addLayout(self.horizontalLayout)
+        self.verticalLayout.addLayout(self.h2)
+        
+        #joint mover tree
+        self.moverTree = QtGui.QTreeWidget(self)
+        self.moverTree.setObjectName("treeWidget")
+        self.moverTree.headerItem().setText(0, "1")
+        self.verticalLayout.addWidget(self.moverTree)
+        self.moverTree.setHeaderHidden(True)
+        
+        QtCore.QMetaObject.connectSlotsByName(self)
+        
+        self.setWindowTitle("Select Joint Movers")
+
+        
+        #snap button
+        self.continueBTN = QtGui.QPushButton(self)
+        self.continueBTN.setText('>> Continue >>')
+        self.continueBTN.setObjectName("blueButton")
+        font = self.continueBTN.font()
+        font.setBold(True)
+        font.setPointSize(12)
+        self.continueBTN.setFont(font)
+        
+        #load button
+        self.loadBTN = QtGui.QPushButton(self)
+        self.loadBTN.setText('Template...')
+        font = self.loadBTN.font()
+        font.setBold(True)
+        font.setPointSize(12)
+        self.loadBTN.setFont(font)
+        self.loadBTN.setMaximumSize(QtCore.QSize(75, 16777215))
+        
+        #button layout on bottom
+        self.h3 = QtGui.QHBoxLayout()
+        self.verticalLayout.addLayout(self.h3)
+        self.h3.setSpacing(2)
+        self.h3.addWidget(self.loadBTN)
+        self.h3.addWidget(self.continueBTN)
+        
+        #connections
+        self.moverTree.itemClicked.connect(self.check_status)
+        self.continueBTN.pressed.connect(self.continueFn)
+        
+        #set the mask
+        self.mask = self.currentFace.faceMask
+        self.mask.active = True
+        cmds.showHidden(self.currentFace.sdks)
+        
+        ### FILL UI
+        ##############################
+        
+        self.attachJointLINE.setText('head')
+        self.refreshUI()
+
+    def check_status(self, item, column):
+        for i in range(0, self.moverTree.topLevelItemCount()):
+            if self.moverTree.topLevelItem(i).checkState(0) == QtCore.Qt.Unchecked:
+                utils.hideShapes(self.moverTree.topLevelItem(i).mover)
+            else:
+                utils.unhideShapes(self.moverTree.topLevelItem(i).mover)
+        cmds.select(item.text(0))
+
+
+    def continueFn(self):
+        global maskAlignWin
+        try:
+            maskAlignWin.close()
+        except:
+            pass
+
+        maskAlignWin = ma_ui.maskAlignWin()
+        maskAlignWin.show()
+        self.close()
+    
+    def refreshUI(self):
+        for f in utils.getUType('faceModule'):
+            self.currentFacialNodeCMB.addItem(f)
+            self.currentFace = face.FaceModule(faceNode=f)
+            self.buildExportTree(self.currentFace.jointMovers)
+            
+    def buildExportTree(self, jointMovers):
+        for jm in jointMovers:
+
+            #top level
+            wid1 = QtGui.QTreeWidgetItem()
+            font = wid1.font(0)
+            font.setPointSize(12)
+            wid1.setFont(0, font)
+            
+            wid1.mover = jm.replace('_SDK', '_SDK_offset')
+
+            wid1.setText(0,jm)
+
+            self.moverTree.addTopLevelItem(wid1)
+            
+            wid1.setFlags(QtCore.Qt.ItemIsEnabled | QtCore.Qt.ItemIsEditable | QtCore.Qt.ItemIsUserCheckable)
+
+            #check if enabled
+            if not utils.shapesHidden(jm):
+                wid1.setCheckState(0, QtCore.Qt.Checked)
+            else:
+                wid1.setCheckState(0, QtCore.Qt.Unchecked)
diff --git a/Engine/Extras/Maya_AnimationRiggingTools/ARTv1/General/Scripts/Modules/facial/mask_align_UI.py b/Engine/Extras/Maya_AnimationRiggingTools/ARTv1/General/Scripts/Modules/facial/mask_align_UI.py
new file mode 100644
--- /dev/null
+++ b/Engine/Extras/Maya_AnimationRiggingTools/ARTv1/General/Scripts/Modules/facial/mask_align_UI.py
@@ -0,0 +1,249 @@
+
+from PySide import QtCore, QtGui
+
+import maya.cmds as cmds
+import maya.OpenMayaUI as mui
+import shiboken
+
+from Modules.facial import utils
+reload(utils)
+from Modules.facial import face
+reload(face)
+
+
+import os
+
+def getMayaWindow():
+    ptr = mui.MQtUtil.mainWindow()
+    if ptr is not None:
+        return shiboken.wrapInstance(long(ptr), QtGui.QWidget)
+
+class maskAlignWin(QtGui.QMainWindow):
+    def __init__(self, parent=getMayaWindow(), debug=0):
+        QtGui.QMainWindow.__init__(self, parent)
+
+        #facial folder
+        #facialDir = os.path.dirname(__file__)
+
+        ### BUILD USER INTERFACE
+        ##############################
+
+        #create/set a central widget
+        wid = QtGui.QWidget()
+        self.setCentralWidget(wid)
+
+        #setArt v2 stylesheet
+        utils.setArtV2StyleSheet(wid, imageDirectory='/art2images')
+        
+        self.verticalLayout = QtGui.QVBoxLayout(wid)
+        self.verticalLayout.setSpacing(2)
+        
+        #top face module picker
+        self.label = QtGui.QLabel(self)
+        self.label.setText("FacialNode:")
+        self.horizontalLayout = QtGui.QHBoxLayout()
+        self.horizontalLayout.addWidget(self.label)
+        self.currentFacialNodeCMB = QtGui.QComboBox(self)
+        self.horizontalLayout.addWidget(self.currentFacialNodeCMB)
+        
+        self.resize(530, 390)
+
+        self.verticalLayout.addLayout(self.horizontalLayout)
+        
+        self.page1 = QtGui.QFrame()
+        self.page1.setMinimumSize(QtCore.QSize( 530, 380 ))
+        self.page1.setMaximumSize(QtCore.QSize( 530, 380 ))
+        
+        self.page2 = QtGui.QFrame()
+        self.page2.setMinimumSize(QtCore.QSize( 530, 380 ))
+        self.page2.setMaximumSize(QtCore.QSize( 530, 380 ))
+
+        self.page3 = QtGui.QFrame()
+        self.page3.setMinimumSize(QtCore.QSize( 530, 380 ))
+        self.page3.setMaximumSize(QtCore.QSize( 530, 380 ))
+
+        self.page4 = QtGui.QFrame()
+        self.page4.setMinimumSize(QtCore.QSize( 530, 380 ))
+        self.page4.setMaximumSize(QtCore.QSize( 530, 380 ))
+        
+        self.stackedWidget = QtGui.QStackedWidget()
+        self.stackedWidget.addWidget(self.page1)
+        self.stackedWidget.addWidget(self.page2)
+        self.stackedWidget.addWidget(self.page3)
+        self.stackedWidget.addWidget(self.page4)
+        self.verticalLayout.addWidget(self.stackedWidget)
+        
+        #set the image backdrops
+        image = (os.path.dirname(__file__) + '\\process_splash_align_mask.png').replace('\\','/')
+        self.page1.setStyleSheet("background-image: url(" + image + ");")
+        
+        image = (os.path.dirname(__file__) + '\\process_splash_align_eyes.png').replace('\\','/')
+        self.page2.setStyleSheet("background-image: url(" + image + ");")
+
+        image = (os.path.dirname(__file__) + '\\process_splash_snap.png').replace('\\','/')
+        self.page3.setStyleSheet("background-image: url(" + image + ");")
+
+        image = (os.path.dirname(__file__) + '\\process_splash_goodbye.png').replace('\\','/')
+        self.page4.setStyleSheet("background-image: url(" + image + ");")
+
+        #Add next button
+        self.nextBTN = QtGui.QPushButton(self)
+        self.nextBTN.setText('NEXT')
+        font = QtGui.QFont()
+        font.setPointSize(12)
+        font.setBold(True)
+        self.nextBTN.setFont(font)
+        self.nextBTN.setGeometry(420, 300, 90, 25)
+        
+        #Add back button
+        self.backBTN = QtGui.QPushButton(self)
+        self.backBTN.setText('BACK')
+        font = QtGui.QFont()
+        font.setPointSize(12)
+        font.setBold(True)
+        self.backBTN.setFont(font)
+        self.backBTN.setGeometry(420, 330, 90, 25)
+        self.backBTN.setHidden(True)
+        
+        self.stackedWidget.setCurrentIndex(0)
+
+        #build and hide the 'mesh chooser'
+        self.meshBTN = QtGui.QPushButton(self)
+        self.meshBTN.setText('<<')
+        font = QtGui.QFont()
+        font.setPointSize(12)
+        font.setBold(True)
+        self.meshBTN.setFont(font)
+        self.meshBTN.setGeometry(0, 0, 25, 25)
+
+        #build and hide the 'mesh chooser'
+        self.snapBTN = QtGui.QPushButton(self)
+        self.snapBTN.setText('SNAP')
+        font = QtGui.QFont()
+        font.setPointSize(12)
+        font.setBold(True)
+        self.snapBTN.setFont(font)
+        self.snapBTN.setGeometry(300, 330, 90, 25)
+
+        self.meshLBL = QtGui.QLabel(self)
+        self.meshLBL.setText("Facial Mesh:")
+        self.meshEdit = QtGui.QLineEdit(self)
+        self.hMeshLayout = QtGui.QHBoxLayout()
+        self.hMeshLayout.addWidget(self.meshLBL)
+        self.hMeshLayout.addWidget(self.meshEdit)
+        self.hMeshLayout.addWidget(self.meshBTN)
+        self.hMeshLayout.setGeometry(QtCore.QRect(100, 300, 300, 25))
+
+        #close button
+        self.closeBTN = QtGui.QPushButton(self)
+        self.closeBTN.setText('Thank you Masky,\nI will alsways remember you..')
+        font = QtGui.QFont()
+        font.setPointSize(16)
+        font.setBold(True)
+        self.closeBTN.setFont(font)
+        self.closeBTN.setGeometry(25, 250, 500, 70)
+        self.closeBTN.setHidden(True)
+
+        #button connects
+        self.nextBTN.clicked.connect(self.nextFn)
+        self.backBTN.clicked.connect(self.backFn)
+        self.meshBTN.clicked.connect(self.meshFn)
+        self.snapBTN.clicked.connect(self.snapFn)
+        self.closeBTN.clicked.connect(self.close)
+
+        self.snapVis(hide=True)
+
+        self.refreshUI()
+
+        #show joint movers and sdks
+        cmds.showHidden(self.currentFace.jointMovers)
+        cmds.showHidden(self.currentFace.sdks)
+
+    def nextFn(self):
+        index = self.stackedWidget.currentIndex()+1
+        self.stackedWidget.setCurrentIndex(index)
+        if index > 0:
+            #disable / hide mask
+            if self.currentMask.active:
+                self.currentMask.active = False
+
+            self.backBTN.setHidden(False)
+            if index == 2:
+                self.snapVis(hide=False)
+            if index == 3:
+                self.snapVis(hide=True)
+                self.closeBTN.setHidden(False)
+                self.nextBTN.setHidden(True)
+                self.backBTN.setHidden(True)
+
+        #show sdks
+        if index != 1:
+            cmds.showHidden(self.currentFace.sdks)
+            cmds.hide(self.currentMask.maskGeo)
+            cmds.hide(self.currentMask.controls)
+        else:
+            cmds.hide(self.currentFace.sdks)
+
+    def backFn(self):
+        index = self.stackedWidget.currentIndex()-1
+        self.stackedWidget.setCurrentIndex(index)
+        if index == 0:
+            self.backBTN.setHidden(True)
+            self.currentMask.active = True
+            #show joint movers and sdks
+            cmds.showHidden(self.currentFace.jointMovers)
+            cmds.showHidden(self.currentFace.sdks)
+        if index != 2:
+            self.snapVis(hide=True)
+        else:
+            self.snapVis(hide=False)
+            cmds.hide(self.currentMask.maskGeo)
+            cmds.hide(self.currentMask.controls)
+
+        #show sdks
+        if index != 1:
+            cmds.showHidden(self.currentFace.sdks)
+        else:
+            cmds.hide(self.currentFace.sdks)
+
+        self.nextBTN.setHidden(False)
+        self.closeBTN.setHidden(True)
+
+    def snapVis(self, hide=True):
+        if hide:
+            self.snapBTN.setHidden(True)
+            self.meshLBL.setHidden(True)
+            self.meshBTN.setHidden(True)
+            self.meshEdit.setHidden(True)
+
+        else:
+            self.snapBTN.setHidden(False)
+            self.meshLBL.setHidden(False)
+            self.meshBTN.setHidden(False)
+            self.meshEdit.setHidden(False)
+
+    def meshFn(self):
+        mesh = cmds.ls(sl=1)[0]
+        self.meshEdit.setText(mesh)
+        self.currentFace.renderMeshes = [mesh]
+
+    def snapFn(self):
+        self.currentFace.faceMask.active = False
+        mesh = self.meshEdit.text()
+        self.currentFace.faceMask.snapMovers('y', mesh)
+        #self.currentFace.snapEyeSdkPivots()
+
+    def updateFaceAndMask(self):
+        currentModule = self.currentFacialNodeCMB.currentText()
+        self.currentFace = face.FaceModule(faceNode=currentModule)
+        self.currentMask = self.currentFace.faceMask
+
+    def refreshUI(self):
+        oldIndex = self.currentFacialNodeCMB.currentIndex()
+        for f in utils.getUType('faceModule'):
+            self.currentFacialNodeCMB.addItem(f)
+        if oldIndex != -1:
+            self.currentFacialNodeCMB.setCurrentIndex(oldIndex)
+        self.updateFaceAndMask()
+        if self.currentFace.renderMeshes:
+            self.meshEdit.setText(self.currentFace.renderMeshes[0])
\ No newline at end of file
