Branching

diff --git a/Engine/Extras/Maya_AnimationRiggingTools/ARTv1/General/Scripts/Modules/facial/poseEditor_UI.py b/Engine/Extras/Maya_AnimationRiggingTools/ARTv1/General/Scripts/Modules/facial/poseEditor_UI.py
new file mode 100644
--- /dev/null
+++ b/Engine/Extras/Maya_AnimationRiggingTools/ARTv1/General/Scripts/Modules/facial/poseEditor_UI.py
@@ -0,0 +1,413 @@
+from PySide import QtCore, QtGui
+
+import maya.cmds as cmds
+import maya.OpenMayaUI as mui
+import shiboken
+
+from Modules.facial import face
+reload(face)
+from Modules.facial import utils
+reload(utils)
+
+import functools
+import json
+
+def getMayaWindow():
+    ptr = mui.MQtUtil.mainWindow()
+    if ptr is not None:
+        return shiboken.wrapInstance(long(ptr), QtGui.QWidget)
+
+class poseWidget(QtGui.QWidget):
+    def __init__(self, edit=0):
+        super(poseWidget, self).__init__()
+
+        self.edit = edit
+        self.pose = None
+        self.driven = None
+        self.verticalLayout = QtGui.QVBoxLayout(self)
+
+        #horiz layout for top half
+        self.topHorizLayout = QtGui.QHBoxLayout()
+
+
+        #pose label
+        self.poseLBL = QtGui.QLabel(self)
+        font = QtGui.QFont()
+        font.setPointSize(12)
+        font.setBold(True)
+        self.poseLBL.setFont(font)
+        self.topHorizLayout.addWidget(self.poseLBL)
+
+        #edit label
+        self.editLBL = QtGui.QLabel(self)
+        self.editLBL.setText('<font color=\'#CC5200\'>POSE EDITED</font>')
+        self.editLBL.setHidden(True)
+        font = QtGui.QFont()
+        font.setPointSize(8)
+        font.setBold(True)
+        self.editLBL.setFont(font)
+        self.topHorizLayout.insertStretch(1)
+        self.topHorizLayout.addWidget(self.editLBL)
+
+        self.verticalLayout.addLayout(self.topHorizLayout)
+
+        #horiz layout for bottom half
+        self.horizontalLayout = QtGui.QHBoxLayout()
+
+        #slider
+        self.poseSlider = QtGui.QSlider(self)
+        self.poseSlider.setMaximum(100)
+        self.poseSlider.setOrientation(QtCore.Qt.Horizontal)
+        self.horizontalLayout.addWidget(self.poseSlider)
+        self.sliderValueLBL = QtGui.QLabel(self)
+        self.horizontalLayout.addWidget(self.sliderValueLBL)
+        self.sliderValueLBL.setText('0.0')
+
+        if self.edit:
+            #update button
+            self.updateBTN = QtGui.QPushButton(self)
+            self.updateBTN.setMinimumSize(QtCore.QSize(92, 0))
+            self.updateBTN.setText('UPDATE POSE')
+            font = QtGui.QFont()
+            font.setPointSize(12)
+            font.setBold(True)
+            self.updateBTN.setFont(font)
+            self.horizontalLayout.addWidget(self.updateBTN)
+            #reset button
+            self.resetBTN = QtGui.QPushButton(self)
+            self.resetBTN.setMinimumSize(QtCore.QSize(50, 0))
+            self.resetBTN.setText('RESET')
+            font = QtGui.QFont()
+            font.setPointSize(12)
+            font.setBold(True)
+            self.resetBTN.setFont(font)
+            self.horizontalLayout.addWidget(self.resetBTN)
+            self.resetBTN.setHidden(True)
+            self.resetBTN.setObjectName("orangeButton")
+        self.verticalLayout.addLayout(self.horizontalLayout)
+
+        #margins
+        self.horizontalLayout.setContentsMargins(2,2,2,2)
+        self.horizontalLayout.setSpacing(3)
+        self.verticalLayout.setContentsMargins(2,2,2,2)
+        self.verticalLayout.setSpacing(1)
+
+    def mousePressEvent(self, event):
+        nodes = list(set([x.split('.')[0] for x in self.driven]))
+        cmds.select(nodes)
+
+
+class poseEditorWin(QtGui.QMainWindow):
+    def __init__(self, parent=getMayaWindow(), mode='choose', edit=0):
+        QtGui.QMainWindow.__init__(self, parent)
+
+        self.edit = edit
+        self.currentFace = None
+
+        #TODO this won't work with multiple faces
+        for f in utils.getUType('faceModule'):
+            print 'poseChooser_UI: Face initialized from:', f
+            self.currentFace = face.FaceModule(faceNode=f)
+
+        ### BUILD USER INTERFACE
+        ##############################
+
+        #create/set a central widget
+        wid = QtGui.QWidget()
+        self.setCentralWidget(wid)
+
+        #setArt v2 stylesheet
+        utils.setArtV2StyleSheet(wid, imageDirectory='/art2images')
+
+        self.resize(400, 514)
+
+        #top face module picker
+        self.label = QtGui.QLabel(self)
+        self.label.setText("FacialNode:")
+        self.horizontalLayout = QtGui.QHBoxLayout()
+        self.horizontalLayout.setSpacing(2)
+        self.label.setMaximumSize(QtCore.QSize(78, 16777215))
+        self.horizontalLayout.addWidget(self.label)
+        self.currentFacialNodeCMB = QtGui.QComboBox(self)
+        self.horizontalLayout.addWidget(self.currentFacialNodeCMB)
+
+        #header and instruction
+        self.h2 = QtGui.QHBoxLayout()
+        self.attachJointLBL = QtGui.QLabel(self)
+
+        if self.edit:
+            self.attachJointLBL.setText("Edit Face Poses")
+        else:
+            self.attachJointLBL.setText("Available Face Poses")
+
+        font = self.attachJointLBL.font()
+        font.setBold(True)
+        font.setPointSize(13)
+        self.attachJointLBL.setFont(font)
+
+        #options horizontal layout
+        self.optionsLayout = QtGui.QHBoxLayout()
+
+        if self.edit:
+            #zero out button
+            self.zeroBTN = QtGui.QPushButton(self)
+            self.zeroBTN.setText('ZERO ALL')
+            self.zeroBTN.setMaximumSize(QtCore.QSize(75, 50))
+
+            #mirror edit out button
+            self.mirrorPoseBTN = QtGui.QPushButton(self)
+            self.mirrorPoseBTN.setText('MIRROR SELECTED POSE')
+
+            #checkboxes
+            self.moveMirrorsCHK = QtGui.QCheckBox(self)
+            self.moveMirrorsCHK.setText('Sync Movement of Mirror Poses')
+
+        self.info = QtGui.QLabel(self)
+        if self.edit:
+            self.info.setText("Check and edit the facial poses to make them better match your character")
+        else:
+            self.info.setText("These poses are available for your character's facial rig, based on the Facial Joint Movers you have selected")
+        self.info.setWordWrap(1)
+
+        self.h2.setSpacing(2)
+        self.h2.addWidget(self.attachJointLBL)
+
+
+        #main vertical layout
+        self.verticalLayout = QtGui.QVBoxLayout(wid)
+        self.verticalLayout.addLayout(self.horizontalLayout)
+        self.verticalLayout.addLayout(self.h2)
+        self.verticalLayout.addWidget(self.info)
+
+        #pose list
+        self.poseList = QtGui.QListWidget(self)
+        self.verticalLayout.addWidget(self.poseList)
+
+        QtCore.QMetaObject.connectSlotsByName(self)
+
+        self.setWindowTitle("Available Facial Poses")
+
+        if self.edit:
+            #options area
+            self.verticalLayout.addWidget(self.moveMirrorsCHK)
+            self.verticalLayout.addLayout(self.optionsLayout)
+            self.optionsLayout.addWidget(self.zeroBTN)
+            self.optionsLayout.addWidget(self.mirrorPoseBTN)
+
+
+        #continue button
+        self.continueBTN = QtGui.QPushButton(self)
+        self.continueBTN.setText('>> Continue >>')
+        self.continueBTN.setObjectName("blueButton")
+        font = self.continueBTN.font()
+        font.setBold(True)
+        font.setPointSize(12)
+        self.continueBTN.setFont(font)
+
+        #back button
+        self.backBTN = QtGui.QPushButton(self)
+        self.backBTN.setText('<< Back <<')
+        self.backBTN.setObjectName("blueButton")
+        font = self.backBTN.font()
+        font.setBold(True)
+        font.setPointSize(12)
+        self.backBTN.setFont(font)
+
+        #button layout on bottom
+        self.h3 = QtGui.QHBoxLayout()
+        self.verticalLayout.addLayout(self.h3)
+        self.h3.setSpacing(2)
+        self.h3.addWidget(self.backBTN)
+        self.h3.addWidget(self.continueBTN)
+
+        #self.verticalLayout.setContentsMargins(2,2,2,2)
+        self.verticalLayout.setSpacing(3)
+
+        #Make SDK transforms editable
+        #hide facial mask controls and mask
+
+        #show special eye and jaw movers
+
+        #unhide jointMovers, make LRAs unselectable
+
+        #connections
+        self.currentFacialNodeCMB.currentIndexChanged.connect(self.faceChangedFn)
+        #self.moverTree.itemClicked.connect(self.check_status)
+        #self.continueBTN.pressed.connect(self.continueFn)
+
+        #set the mask
+        self.mask = face.FaceMask
+        self.mask.active = True
+
+        ### CONNECTS
+        if self.edit:
+            self.zeroBTN.pressed.connect(self.zeroAllFn)
+
+        ### FILL UI
+        ##############################
+
+        self.refreshUI()
+
+    ### UI HOOKS
+    ##############################
+
+    def refreshUI(self):
+        for f in utils.getUType('faceModule'):
+            self.currentFacialNodeCMB.addItem(f)
+            self.currentFace = face.FaceModule(faceNode=f)
+            self.buildPoseList()
+
+    def updateFn(self, wid):
+        #check that the pose is fully on
+        if wid.poseSlider.value() != 100:
+            cmds.warning('Pose slider is not fully on, must be 1.0 in able to save edit.')
+            return False
+
+        cmds.setDrivenKeyframe(wid.driven, cd=wid.pose)
+        self.checkForEditedPoses()
+
+    def resetFn(self, wid):
+        print 'resetting', wid
+        try:
+            cmds.undoInfo(openChunk=True)
+            #get original pose from driven dict
+            initialPoseDict = json.loads(cmds.getAttr(self.currentFace.blendBoard + '.initialPoseDict'))
+            initialPoseAttrs = initialPoseDict[wid.pose]
+
+            #store current slider/attr val
+            preVal = cmds.getAttr(wid.pose)
+
+            #iterate over the attrs and set them bvack to original
+            #TODO: rotate order?
+            cmds.setAttr(wid.pose, 1.0)
+            for att in wid.driven:
+                #print att, initialPoseAttrs[att]
+                cmds.setAttr(att, initialPoseAttrs[att])
+
+            #set the old driven keyframe
+            cmds.setDrivenKeyframe(wid.driven, cd=wid.pose)
+
+            #reset the slider/attr val
+            cmds.setAttr(wid.pose, preVal)
+            self.checkForEditedPoses()
+
+        except Exception, err:
+            import traceback
+            print(traceback.format_exc())
+        finally:
+            cmds.undoInfo(closeChunk=True)
+
+    def sliderChanged(self, wid, value):
+        val = float(value)/100
+        floatVal = str(val)
+        wid.sliderValueLBL.setText(floatVal)
+        cmds.setAttr(wid.pose, val)
+
+    def zeroAllFn(self):
+        for i in range(0, self.poseList.count()):
+            wid = self.poseList.itemWidget(self.poseList.item(i))
+            wid.poseSlider.setValue(0)
+
+    def getAllDrivenAttrs(self):
+        ret = {}
+        for i in range(0, self.poseList.count()):
+            wid = self.poseList.itemWidget(self.poseList.item(i))
+            ret[wid.pose] = cmds.getAttr(wid.pose)
+        return ret
+
+    def setAllDrivenAttrs(self, dict):
+        for pose in dict:
+            cmds.setAttr(pose, dict[pose])
+
+    def zeroAllAttrs(self):
+        for i in range(0, self.poseList.count()):
+            wid = self.poseList.itemWidget(self.poseList.item(i))
+            cmds.setAttr(wid.pose, 0.0)
+
+    #not used atm
+    def getSliders(self):
+        ret = []
+        for i in range(0, self.poseList.count()):
+            wid = self.poseList.itemWidget(self.poseList.item(i))
+            ret.append(wid.poseSlider.value())
+        return ret
+
+    #not used atm
+    def setSliders(self, sliders):
+        for i in range(0, len(sliders)):
+            wid = self.poseList.itemWidget(self.poseList.item(i))
+            wid.poseSlider.setValue(sliders[i])
+
+    def checkForEditedPoses(self, threshold=0.001):
+        initialPoseDict = json.loads(cmds.getAttr(self.currentFace.blendBoard + '.initialPoseDict'))
+
+        for i in range(0, self.poseList.count()):
+            edited = False
+            wid = self.poseList.itemWidget(self.poseList.item(i))
+            preVals = self.getAllDrivenAttrs()
+            self.zeroAllAttrs()
+            cmds.setAttr(wid.pose, 1.0)
+
+            #iterate and check driven att values
+            initialPoseAttrs = initialPoseDict[wid.pose]
+            for att in wid.driven:
+                error = abs(initialPoseAttrs[att] - cmds.getAttr(att))
+                if error > threshold:
+                    edited = True
+            if edited:
+                wid.editLBL.setHidden(False)
+                wid.resetBTN.setHidden(False)
+            else:
+                wid.editLBL.setHidden(True)
+                wid.resetBTN.setHidden(True)
+            self.setAllDrivenAttrs(preVals)
+
+    def addPose(self, pose, driven):
+        wid = poseWidget(edit=self.edit)
+        wid.pose = pose
+        wid.driven = driven
+        currentVal = cmds.getAttr(pose)
+        wid.poseLBL.setText(pose.split('.')[-1])
+        wid2 = QtGui.QListWidgetItem()
+        self.poseList.addItem(wid2)
+        wid2.setSizeHint(QtCore.QSize(100, 60))
+        self.poseList.setItemWidget(wid2, wid)
+
+        #hookup slider
+        wid.poseSlider.valueChanged.connect(functools.partial(self.sliderChanged, wid))
+        wid.poseSlider.setValue(currentVal*100)
+
+        if self.edit:
+            #hookup buttons
+            wid.updateBTN.pressed.connect(functools.partial(self.updateFn, wid))
+            wid.resetBTN.pressed.connect(functools.partial(self.resetFn, wid))
+
+        #store driven for checking later
+        wid.driven = driven
+
+    def buildPoseList(self):
+        self.poseList.clear()
+        for pose in self.currentFace.poses:
+            self.addPose(pose, self.currentFace.poses[pose])
+        if self.edit:
+            self.checkForEditedPoses()
+
+    def faceChangedFn(self):
+        self.currentFace = face.FaceModule(faceNode=self.currentFacialNodeCMB.currentText())
+        self.buildPoseList()
+
+if __name__ == '__main__':
+    global poseWin
+    try:
+        poseWin.close()
+    except:
+        pass
+
+
+    poseWin = poseEditorWin(edit=1)
+    poseWin.show()
+
+
+
+
+
diff --git a/Engine/Extras/Maya_AnimationRiggingTools/ARTv1/General/Scripts/Modules/facial/utils.py b/Engine/Extras/Maya_AnimationRiggingTools/ARTv1/General/Scripts/Modules/facial/utils.py
new file mode 100644
--- /dev/null
+++ b/Engine/Extras/Maya_AnimationRiggingTools/ARTv1/General/Scripts/Modules/facial/utils.py
@@ -0,0 +1,457 @@
+import maya.api.OpenMaya as om2
+import maya.cmds as cmds
+
+import os
+import traceback
+
+## GENERAL
+#######################################################################
+
+def dagPathFromName(obj):
+    '''
+    Given a maya object name, returns an openMaya2 dagPath object, dagPath, and a full dagPath in a string
+    '''
+    sel = om2.MSelectionList()
+    sel.add(obj)
+    dag = sel.getDagPath(0)
+    dagStr = sel.getDagPath(0).fullPathName()
+    mobj = sel.getDependNode(0)
+    return mobj, dag, dagStr
+
+
+
+### ATTRIBUTE MANAGEMENT
+#######################################################################
+def attrExists(attr):
+    if '.' in attr:
+        node, att = attr.split('.')
+        return cmds.attributeQuery(att, node=node, ex=1)
+    else:
+        cmds.warning('attrExists: No attr passed in: ' + attr)
+        return False
+
+def msgConnect(attribFrom, attribTo, debug=0):
+    #TODO needs a mode to dump all current connections (overwrite/force)
+    objFrom, attFrom = attribFrom.split('.')
+    objTo, attTo = attribTo.split('.')
+    if debug: print 'msgConnect>>> Locals:', locals()
+    if not attrExists(attribFrom):
+        cmds.addAttr(objFrom, longName=attFrom, attributeType='message')
+    if not attrExists(attribTo):
+        cmds.addAttr(objTo, longName=attTo, attributeType='message')
+    
+    #check that both atts, if existing are msg atts
+    for a in (attribTo, attribFrom):
+        if cmds.getAttr(a, type=1) != 'message':
+            cmds.warning('msgConnect: Attr, ' + a + ' is not a message attribute. CONNECTION ABORTED.')
+            return False
+    
+    try:
+        return cmds.connectAttr(attribFrom, attribTo)
+    except Exception as e:
+        print e
+        return False
+
+def snapMeshPivot(mesh1, mesh2):
+	'''
+	move mesh2 pivot to mesh1 pivot
+	@author=chrise
+	'''
+
+	dupe = cmds.duplicate(mesh2)[0]
+	mesh1Xform = cmds.getAttr(mesh1+'.worldMatrix')
+	cmds.xform(mesh2, m=mesh1Xform)
+
+	cmds.undoInfo(openChunk=True)
+	try:
+		for v in range(0, cmds.polyEvaluate(mesh2, v=1)):
+			m2Vtx = mesh2+'.vtx['+str(v)+']'
+			dupeVtx = dupe+'.vtx['+str(v)+']'
+			cmds.xform(m2Vtx, m=cmds.xform(dupeVtx, m=1, q=1, ws=1), ws=1)
+		cmds.delete(dupe)
+	finally:
+		cmds.undoInfo(closeChunk=True)
+	'''
+	#TODO: here's some sample code for working with one nurbs shape
+	#I stopped because it didn't work with multiple shape nodes
+	surf = cmds.ls(sl=1)[0]
+    if cmds.nodeType(surf) == 'transform':
+        #get all cvs
+        cmds.select(surf + '.cv[*]')
+        cvs = cmds.ls(sl=1, flatten=1)
+        cmds.select(cl=1)
+        for cv in cvs:
+            print cv
+            print cmds.pointPosition(cv)
+    '''
+
+def snapRotatePivotPosition(nodeFrom, nodeTo):
+    '''
+    sel = cmds.ls(sl=1)
+    snapRotatePivotPosition(sel[0], sel[1])
+    '''
+    p = cmds.xform(nodeTo, q=1, ws=1, rp=1)
+    cmds.xform(nodeFrom, ws=1, piv=p)
+
+def snapPivotWithZeroParent(fromNode, toNode):
+    '''
+    Mainly used for snapping the pivot of a controller that is in a topnode GRP's space
+    '''
+    #TODO: This won't work if there is more than one child
+    try:
+        parent = cmds.listRelatives(fromNode, parent=1)
+
+        cmds.parent(fromNode, w=1)
+        toNodeMx = cmds.xform(toNode, q=1, ws=1, m=1)
+        cmds.xform(parent, m=toNodeMx, ws=1)
+
+        toNodePivot = cmds.xform(toNode, ws=1, q=1, t=1)
+        cmds.xform(fromNode, ws=1, a=1, piv=toNodePivot)
+        cmds.parent(fromNode, parent)
+        cmds.makeIdentity(fromNode, t=1, r=1, a=1)
+    except Exception as e:
+        print(traceback.format_exc())
+
+def getUType(type):
+    '''
+    Returns nodes of the requested uType
+    example: getUType('faceModule')
+    '''
+    #TODO: Check with namespaces
+    return [node.split('.')[0] for node in cmds.ls('*.uType', r=1) if cmds.getAttr(node) == type]
+
+def hideShapes(node):
+    for shape in cmds.listRelatives(node, shapes=1, f=1):
+        cmds.hide(shape)
+
+def unhideShapes(node):
+    for shape in cmds.listRelatives(node, shapes=1, f=1):
+        cmds.showHidden(shape)
+
+def shapesHidden(node):
+    for shape in cmds.listRelatives(node, shapes=1, f=1):
+        if cmds.getAttr(shape + '.visibility') == 1:
+            return False
+    return True
+
+### CONTROLLER MANAGEMENT
+#######################################################################
+def toggleShapeSelection(node, selectable):
+    '''
+    example: makeShapesUnselectable(cmds.ls(sl=1)[0], False)
+    '''
+    for shape in cmds.listRelatives(node, fullPath=1, shapes=1):
+        if selectable:
+            cmds.setAttr(shape + '.overrideEnabled', 0)
+            cmds.setAttr(shape + '.overrideDisplayType', 0)
+        else:
+            cmds.setAttr(shape + '.overrideEnabled', 1)    
+            cmds.setAttr(shape + '.overrideDisplayType', 2)
+
+def getAxisWorldVectors(node):
+    matrix = om2.MGlobal.getSelectionListByName(node).getDagPath(0).inclusiveMatrix()
+    x = ((om2.MVector.kXaxisVector) * matrix).normal()
+    y = ((om2.MVector.kYaxisVector) * matrix).normal()
+    z = ((om2.MVector.kZaxisVector) * matrix).normal()
+    return [x,y,z]
+
+def createArrow(name=None, thickness=0.1, length=2, vector=[1,0,0], point=[0,0,0]):
+    '''
+    Creates an arrow in the direction of the vector
+    Example:
+    arrow = createArrow(vector=[0,1,0], length=4)
+    '''
+    if not name:
+        name = 'arrow_CTRL'
+    #calc length for thickness
+    ratio = length/thickness
+    cyl = cmds.cylinder(radius=thickness, sections=4, heightRatio=ratio, pivot=[length/2, 0, 0], ch=0)[0]
+    cone = cmds.cone(radius=thickness*2, sections=4, ch=0, pivot=(length*1.0005,0,0))[0]
+    xform = cmds.createNode('transform', ss=1, name=name)
+    
+    shapes = []
+    transforms = []
+    for node in [cone,cyl]:
+        shapes.append(cmds.listRelatives(node, fullPath=1, shapes=1)[0])
+        transforms.append(node)
+    cmds.parent(shapes, xform, r=1, s=1)
+    rotateBy = cmds.angleBetween(euler=True, v1=(1,0,0), v2=vector)
+    cmds.rotate(rotateBy[0], rotateBy[1], rotateBy[2], xform)
+    cmds.xform(xform, t=point)
+    cmds.delete(transforms)
+    return xform
+
+
+def colorControl(node, name=None, color=None, useExisting=1, shapes=0):
+    '''
+    Colors an object by creating and assigning a shader with the color you input
+    Example:
+    colorControl('arrow', color=(1,0,0))
+    '''
+    if not name:
+        name = 'colorControl_%s_%s_%s' % (color[0], color[1], color[2])
+    
+    shader = None
+    shading_group = None
+    
+    if cmds.objExists(name):
+        if 'shadingDependNode' in cmds.nodeType(name, i=1):
+            if useExisting:
+                shader = name
+                shading_group = cmds.listConnections(name + '.outColor')[0]
+    else:
+        if len(color) == 3:
+            shader=cmds.shadingNode('blinn',asShader=True, name = name)
+            cmds.setAttr(shader + '.color', color[0], color[1], color[2], type='double3')
+            shading_group = cmds.sets(renderable=True,noSurfaceShader=True,empty=True, name=name + '_colorControlSG')
+            cmds.connectAttr('%s.outColor' %shader ,'%s.surfaceShader' %shading_group)
+        else:
+            cmds.warning('Color input is not a 3 element list or tuple. COLOR: ' + str(color))
+            return False
+    if shapes:
+        for shape in cmds.listRelatives(node, fullPath=1, shapes=1):
+            cmds.sets(shape, e=1, forceElement=shading_group)
+    else:
+        cmds.sets(node, e=1, forceElement=shading_group)
+    cmds.setAttr(shader + '.diffuse', 1)
+    cmds.setAttr(shader + '.ambientColor', 1,1,1, type='double3')
+
+    #Need to check if it's a blinn before setting these
+    #cmds.setAttr(shader + '.eccentricity', 0)
+    #cmds.setAttr(shader + '.specularColor', 0,0,0, type='double3')
+    return shader, shading_group
+    
+def createLRA(node=None, matrix=None, name=None, color=True, thickness=0.1, length=2, vector=[1,0,0], point=[0,0,0], arrowRadiusMult=2):
+    '''
+    Creates an LRA at the origin, or at a node, or a mmx
+    Example:
+    createLRA(length=2, thickness=0.05, arrowRadiusMult=3)
+    '''
+    if not name:
+        name = 'arrow_CTRL'
+    nodes = []
+    x,y,z = [],[],[]
+    #calc length for thickness
+    ratio = length/thickness
+    xform = cmds.createNode('transform', ss=1, name=name)
+
+    x.append(cmds.cylinder(radius=thickness, sections=4, heightRatio=ratio, pivot=[length/2, 0, 0], ch=0, axis=[1,0,0])[0])
+    x.append(cmds.cone(radius=thickness*arrowRadiusMult, sections=4, ch=0, pivot=(length*1.0005,0,0))[0])
+    
+    y.append(cmds.cylinder(radius=thickness, sections=4, heightRatio=ratio, pivot=[0, length/2, 0], ch=0, axis=[0,1,0])[0])
+    y.append(cmds.cone(radius=thickness*arrowRadiusMult, sections=4, ch=0, pivot=(0,length*1.0005,0), axis=[0,1,0])[0])
+    
+    z.append(cmds.cylinder(radius=thickness, sections=4, heightRatio=ratio, pivot=[0, 0, length/2], ch=0, axis=[0,0,1])[0])
+    z.append(cmds.cone(radius=thickness*arrowRadiusMult, sections=4, ch=0, pivot=(0,0,length*1.0005), axis=[0,0,1])[0])
+    
+    nodes.extend(x)
+    nodes.extend(y)
+    nodes.extend(z)
+    
+    if color:
+        for node in x: colorControl(node, name='red_m', color=(1,0,0))
+        for node in y: colorControl(node, name='green_m', color=(0,1,0))
+        for node in z: colorControl(node, name='blue_m', color=(0,0,1))
+    
+    shapes = []
+    transforms = []
+    for node in nodes:
+        shapes.append(cmds.listRelatives(node, fullPath=1, shapes=1)[0])
+        transforms.append(node)
+    cmds.parent(shapes, xform, r=1, s=1)
+    rotateBy = cmds.angleBetween(euler=True, v1=(1,0,0), v2=vector)
+    cmds.rotate(rotateBy[0], rotateBy[1], rotateBy[2], xform)
+    cmds.xform(xform, t=point)
+    cmds.delete(transforms)
+    return xform
+    
+
+def createCircleControl(name=None, x=1, y=1, z=1, radius=1, color=None, type='strip', heightRatio=0.007):
+    if not name:
+        name = 'circle_CTRL'
+    xform = cmds.group(name=name, em=1)
+    circleShapes = []
+    circleTransforms = []
+    vectors = getAxisWorldVectors(xform)
+    
+    xyz = [x,y,z]
+    for i in range(0, len(xyz)):
+        if xyz[i]:
+            circ = None
+            if type == 'curve':
+                circ = cmds.circle(radius=radius, name='jointMoverAxisCtrl', normal=vectors[i], sections=1, ch=0)
+                if color:
+                    cmds.color(circ, ud=color)
+            if type == 'strip' or type == 'toroid':
+                circ = None
+                if type == 'strip':
+                    circ = cmds.cylinder(radius=radius, name='jointMoverAxisCtrl', axis=vectors[i], sections=7, ch=0, heightRatio=heightRatio)
+                elif type == 'toroid':
+                    circ = cmds.torus(radius=radius, name='jointMoverAxisCtrl', axis=vectors[i], sections=7, ch=0, heightRatio=heightRatio)
+                if color:
+                    if color == 'xyz':
+                        if i == 0: colorControl(circ, name='red_m', color=(1,0,0))
+                        if i == 1: colorControl(circ, name='green_m', color=(0,1,0))
+                        if i == 2: colorControl(circ, name='blue_m', color=(0,0,1))
+                    else:
+                        if len(color) == 3: colorControl(circ, color=color)
+            circleShapes.append(cmds.listRelatives(circ[0], fullPath=1, shapes=1)[0])
+            circleTransforms.append(circ[0])
+    cmds.parent(circleShapes, xform, r=1, s=1)
+    cmds.delete(circleTransforms)
+    return xform
+    
+### ART V2 WRANGLING
+#######################################################################
+
+def setArtV2StyleSheet(widget, sheetTemplate=None, imageDirectory=None):
+    
+    css, lines = None, None
+    if not sheetTemplate:
+        sheetTemplate = (os.path.dirname(__file__) + '\\artV2_css_template.txt').replace('\\','/')
+        imageDirectory = (os.path.dirname(__file__) + imageDirectory).replace('\\','/')
+    if os.path.isfile(sheetTemplate):
+        f = open(sheetTemplate, 'r')
+        cssTemplate = f.read()
+        f.close()
+        css = cssTemplate.replace('___PATH___', imageDirectory)
+    else:
+        print('Cannot find CSS template: ' + sheetTemplate)
+        return False
+    
+    #check to see if all images found
+    for line in css.split('\n'):
+        if '.png' in line.lower():
+            image = line.split('(')[1].split(')')[0]
+            if not os.path.isfile(image):
+                print 'ERROR:setArtV2StyleSheet: Image not found [' + image + ']'
+    
+    widget.setStyleSheet(css)
+
+## SURFACE CONSTRAINING
+#######################################################################
+
+def square_distance(pointA, pointB):
+    '''
+    Both points must have the same dimensions
+    '''
+    if len(pointA) != len(pointB):
+        print 'SQUARED_DISTANCE: incoming points do not have the same dimensions:', pointA, pointB
+        return False
+    distance = 0
+    dimensions = len(pointA)
+    for dimension in range(dimensions):
+        distance += (pointA[dimension] - pointB[dimension])**2
+    return distance
+
+def closestVertOnMesh(mesh, node):
+    '''
+    Returns closest vertex component on mesh to transform/pivot of node
+    '''
+    closestVert = None
+    minLength = None
+    verts = cmds.getAttr(mesh+".vrts", multiIndices=True)
+    nodePos = cmds.xform(node, q=1, translation=1, ws=1)
+    for v in verts:
+        vPos = cmds.xform(mesh + '.vtx['+str(v)+']', q=1, translation=1, ws=1)
+        length = om2.MVector(om2.MVector(nodePos) - om2.MVector(vPos)).length()
+        if minLength is None or length < minLength:
+            minLength = length
+            closestVert = v
+    return (mesh + '.vtx[' + str(closestVert) + ']')
+
+def closestPointonMesh(mesh, node, createLoc=False):
+    pos = cmds.xform(node, q=1, rp=1, ws=1)
+    nodePos = om2.MPoint(pos[0], pos[1], pos[2],1.0)
+    print 'point', pos[0], pos[1], pos[2]
+
+    obj,dag,stringPath = dagPathFromName(mesh)
+    meshFn = om2.MFnMesh(dag)
+
+    resultPoint = meshFn.getClosestPoint(nodePos,om2.MSpace.kWorld)[0]
+    
+    if createLoc:
+        loc = cmds.spaceLocator()[0]
+        cmds.setAttr(loc+'.t', resultPoint[0], resultPoint[1], resultPoint[2])
+        return loc
+    else:
+        return (resultPoint[0], resultPoint[1], resultPoint[2])
+
+def constrainToMesh(mesh, node, mo=True, type='follicle'):
+    '''
+    Attaches node to closest vertex of mesh
+    '''
+    
+    #get closest vertex
+    vtx = closestVertOnMesh(mesh, node)
+    
+    surfaceXform = None
+    #build/place follicle
+    if type == 'follicle':
+        surfaceXform = cmds.createNode("transform", n=(mesh+"_follicle"))
+        follicle = cmds.createNode("follicle", n=(mesh+"_follicleShape"), p=surfaceXform)
+        cmds.connectAttr((follicle+".outTranslate"), (surfaceXform+".translate"), f=1)
+        cmds.connectAttr((follicle+".outRotate"), (surfaceXform+".rotate"), f=1)
+        cmds.connectAttr((mesh+".outMesh"), (follicle+".inm"), f=1)
+        cmds.connectAttr((mesh+".worldMatrix[0]"), (follicle+".inputWorldMatrix"), f=1)
+        uv = cmds.polyListComponentConversion(vtx, fv=True, tuv=True)
+        uvs = cmds.polyEditUV(uv ,q=1, u=1, v=1)
+        u = uvs[0]
+        v = uvs[1]
+        cmds.setAttr((follicle+".parameterU"), u)
+        cmds.setAttr((follicle+".parameterV"), v)
+    cmds.parentConstraint(surfaceXform, node, mo=mo)
+
+def rayMeshIntersect(mesh, pos, vector):
+    '''
+    Example:
+    mover = cmds.ls(sl=1)[0]
+    mesh = 'steel'
+    pos = cmds.xform(mover, q=1, rp=1, ws=1)
+    newPos = rayMeshIntersect(mesh, pos, -utils.getAxisWorldVectors(mover)[1])
+    cmds.xform(mover, t=newPos, ws=1)
+    '''
+    nodePos = om2.MPoint(pos[0], pos[1], pos[2],1.0)
+
+    vector = om2.MVector(vector)
+    obj,dag,stringPath = dagPathFromName(mesh)
+    meshFn = om2.MFnMesh(dag)
+
+    selectionList = om2.MSelectionList()
+    selectionList.add(mesh)
+    dagPath = selectionList.getDagPath(0)
+    fnMesh = om2.MFnMesh(dagPath)
+    intersection = fnMesh.closestIntersection(
+    om2.MFloatPoint(pos),
+    om2.MFloatVector(vector),
+    om2.MSpace.kWorld, 99999, False)
+    if intersection:
+        hitPoint, hitRayParam, hitFace, hitTriangle, hitBary1, hitBary2 = intersection
+        x,y,z,_ = hitPoint
+        return (x,y,z)
+
+def deleteConstraints(node):
+    ch = cmds.listRelatives(node, children=1)
+    if ch:
+        for item in ch:
+            if 'Constraint' in cmds.nodeType(item):
+                cmds.delete(item)
+
+def findNamingConflicts(ignoreShapes=1):
+    #TODO: listRelatives to find parent, check relationship to make sure theyre shapes instead of namechecking
+    nameDict = {}
+    for node in cmds.ls('*', l=1):
+        if ignoreShapes:
+            if 'shape' in node.lower():
+                continue
+        shortName = node.split('|')[-1]
+        if shortName not in nameDict.keys():
+            nameDict[shortName] = [node]
+        else:
+            cmds.warning('Found duplicate node names: \n' + node + '\n' + str(nameDict[shortName]))
+            nameDict[shortName].append(node)
+    returnDict = {}
+    for name in nameDict:
+        if len(nameDict[name]) > 1:
+            returnDict[name] = nameDict[name]
+    return returnDict
\ No newline at end of file
diff --git a/Engine/Extras/Maya_AnimationRiggingTools/ARTv1/General/Scripts/Modules/p4python/EpicP4.py b/Engine/Extras/Maya_AnimationRiggingTools/ARTv1/General/Scripts/Modules/p4python/EpicP4.py
new file mode 100644
--- /dev/null
+++ b/Engine/Extras/Maya_AnimationRiggingTools/ARTv1/General/Scripts/Modules/p4python/EpicP4.py
@@ -0,0 +1,86 @@
+from Modules.p4python.P4 import P4, P4Exception
+import maya.cmds as cmds
+
+def findFileP4(asset, debug=0, root='//depot/ArtSource/'):
+    p4 = P4()
+    try: p4.connect()
+    except:
+        print 'Cannot connect!'
+        return False
+
+    try:
+        file =  p4.run_files(root + '...' + asset)
+        depotLoc = file[0]['depotFile']
+        describe = p4.run_describe(file[0]['change'])
+        return [describe[0]['user'], depotLoc, describe[0]['desc'], file[0]['change']]
+    except Exception as e:
+        print "Cannot find file.", asset
+        if debug: print e
+        return False
+    finally:
+        p4.disconnect()
+
+def rePathFileNodesP4(debug=0, searchPath='//depot/ArtSource/'):
+    nodes = cmds.ls(type='file')
+
+    for node in nodes:
+        path = cmds.getAttr((node + '.fileTextureName'))
+        newPath = None
+
+        #check if the file exists on disk
+        if cmds.file(path, exists=1, q=1):
+            pass
+
+        #if the file is not on disk, continue and get it from P4
+        else:
+            if debug:
+                print 'rePathFileNodesP4>>>> Local file not found', path, 'searching p4 database.'
+
+            #connect to the perforce server
+            p4 = P4()
+            try: p4.connect()
+            except:
+                print '>>rePathFileNodesP4: Cannot connect to perforce!'
+                return False
+
+            try:
+                #convert the path to the same direction slashes
+                floppedPath = path.replace('\\', '/')
+                fName = floppedPath.split('/')[-1]
+
+                #search p4 for the file
+                f = None
+                try:
+                    f = p4.run_files(searchPath + '...' + fName)
+                except:
+                    pass
+                if f:
+                    if len(f) > 1:
+                        print 'rePathFileNodesP4>>>> Multiple files [', len(f), '] found in depot search path with the name', f
+                        for i in range(0, len(f)):
+                            try:
+                                print 'rePathFileNodesP4>>>> SYNCING: ', f[i]['depotFile']
+                                p4.run_sync(f[i]['depotFile'])
+                            except Exception as e:
+                                if debug:
+                                    print e
+                    else:
+                        print 'rePathFileNodesP4>>>> File found on depot: ', f[0]['depotFile']
+                        try:
+                            p4.run_sync(f[0]['depotFile'])
+                        except Exception as e:
+                                if debug:
+                                    print e
+                else:
+                    print 'rePathFileNodesP4>>>> FILE NOT FOUND IN PERFORCE SEARCH PATH - fName:', fName, 'searchPath:', searchPath
+
+                fileFstat = p4.run('fstat', f[0]['depotFile'])[0]
+                newPath = fileFstat['clientFile']
+
+                if debug: print 'rePathFileNodesP4>>>> NEW PATH>>', newPath, '\n\n'
+                cmds.setAttr((node + '.fileTextureName'), newPath, type='string')
+            except Exception as e:
+                print "Cannot find file.", fName,e
+                if debug:
+                    print 'rePathFileNodesP4>>>> EXCEPTION', e
+    p4.disconnect()
\ No newline at end of file
diff --git a/Engine/Extras/Maya_AnimationRiggingTools/ARTv1/General/Scripts/Modules/p4python/P4.py b/Engine/Extras/Maya_AnimationRiggingTools/ARTv1/General/Scripts/Modules/p4python/P4.py
new file mode 100644
--- /dev/null
+++ b/Engine/Extras/Maya_AnimationRiggingTools/ARTv1/General/Scripts/Modules/p4python/P4.py
@@ -0,0 +1,1036 @@
+#!/usr/local/bin/python
+
+from __future__ import print_function
+
+"""A Python version of the Perforce "p4" client.
+    
+    This uses the Python type P4API.P4Adapter, which is a wrapper for the
+    Perforce ClientApi object.
+    
+    $Id: //depot/r15.1/p4-python/P4.py#1 $
+    
+    #*******************************************************************************
+    # Copyright (c) 2007-2010, Perforce Software, Inc.  All rights reserved.
+    # Portions Copyright (c) 1999, Mike Meyer. All rights reserved.
+    # Portions Copyright (c) 2004-2007, Robert Cowham. All rights reserved.
+    #
+    # Redistribution and use in source and binary forms, with or without
+    # modification, are permitted provided that the following conditions are met:
+    #
+    # 1.  Redistributions of source code must retain the above copyright
+    # notice, this list of conditions and the following disclaimer.
+    #
+    # 2.  Redistributions in binary form must reproduce the above copyright
+    # notice, this list of conditions and the following disclaimer in the
+    # documentation and/or other materials provided with the distribution.
+    #
+    # THIS SOFTWARE IS PROVIDED BY THE COPYRIGHT HOLDERS AND CONTRIBUTORS "AS IS"
+    # AND ANY EXPRESS OR IMPLIED WARRANTIES, INCLUDING, BUT NOT LIMITED TO, THE
+    # IMPLIED WARRANTIES OF MERCHANTABILITY AND FITNESS FOR A PARTICULAR PURPOSE
+    # ARE DISCLAIMED. IN NO EVENT SHALL PERFORCE SOFTWARE, INC. BE LIABLE FOR ANY
+    # DIRECT, INDIRECT, INCIDENTAL, SPECIAL, EXEMPLARY, OR CONSEQUENTIAL DAMAGES
+    # (INCLUDING, BUT NOT LIMITED TO, PROCUREMENT OF SUBSTITUTE GOODS OR SERVICES;
+    # LOSS OF USE, DATA, OR PROFITS; OR BUSINESS INTERRUPTION) HOWEVER CAUSED AND
+    # ON ANY THEORY OF LIABILITY, WHETHER IN CONTRACT, STRICT LIABILITY, OR TORT
+    # (INCLUDING NEGLIGENCE OR OTHERWISE) ARISING IN ANY WAY OUT OF THE USE OF THIS
+    # SOFTWARE, EVEN IF ADVISED OF THE POSSIBILITY OF SUCH DAMAGE.
+    #*******************************************************************************
+    
+    Copyright 2007-2010 Perforce Software, Inc. All rights reserved
+    
+    
+    License:
+    See accompanying LICENSE.txt including for redistribution permission.
+    """
+
+import sys, datetime
+import re
+import shutil
+from contextlib import contextmanager
+import uuid, tempfile
+import os, os.path, platform
+import subprocess
+
+# P4Exception - some sort of error occurred
+class P4Exception(Exception):
+    """Exception thrown by P4 in case of Perforce errors or warnings"""
+    
+    def __init__(self, value):
+        Exception.__init__(self)
+        
+        if isinstance(value, (list, tuple)) and len(value) > 2:
+            self.value = value[0]
+            self.errors = value[1]
+            self.warnings = value[2]
+        else:
+            self.value = value
+    def __str__(self):
+        return str(self.value)
+
+class Spec(dict):
+    """Subclass of dict, representing the fields of a spec definition.
+        
+        Attributes can be accessed either with the conventional dict format,
+        spec['attribute'] or with shorthand spec._attribute.
+        
+        Instances of this class will preventing any unknown keys.
+        """
+    def __init__(self, fieldmap=None):
+        self.__dict__['_Spec__fields'] = fieldmap
+    
+    def permitted_fields(self):
+        return self.__fields
+    
+    def __setitem__(self, key, value):
+        if not isinstance(value, str) and not isinstance(value, list):
+            raise P4Exception("Illegal value of type %s, must be string or list" % value.__class__)
+        if key in self or self.__fields == None:
+            dict.__setitem__(self, key, value)
+        elif str(key).lower() in self.__fields:
+            dict.__setitem__(self, self.__fields[key.lower()], value)
+        else:
+            raise P4Exception("Illegal field '%s'" % str(key))
+    
+    def __getattr__(self, attr):
+        key = str(attr).lower()
+        if key[0] != '_':
+            raise AttributeError(attr)
+        key = key[1:]
+        if key in self:
+            return self[key]
+        elif key in self.__fields:
+            return self[self.__fields[key]]
+    
+    def __setattr__(self, attr, value):
+        if attr == 'comment':
+            self.__dict__[attr] = value
+        else:
+            key = str(attr).lower()
+            if key[0] != '_':
+                raise AttributeError(attr)
+            key = key[1:]
+            self[key] = value
+
+#
+# P4Integration objects hold details about the integrations that have
+# been performed on a particular revision. Used primarily with the
+# P4Revision class
+#
+class Integration:
+    def __init__( self, how, file, srev, erev ):
+        self.how = how
+        self.file = file
+        self.srev = srev
+        self.erev = erev
+    
+    def __repr__(self):
+        return "Integration (how = %s file = %s srev = %s erev = %s)" \
+            % (self.how, self.file, self.srev, self.erev)
+#
+# Each P4Revision object holds details about a particular revision
+# of a file. It may also contain the history of any integrations
+# to/from the file
+#
+
+class Revision:
+    def __init__( self, depotFile ):
+        self.depotFile = depotFile
+        self.integrations = []
+        self.rev = None
+        self.change = None
+        self.action = None
+        self.type = None
+        self.time = None
+        self.user = None
+        self.client = None
+        self.desc = None
+        self.digest = None
+        self.fileSize = None
+    
+    def integration( self, how, file, srev, erev ):
+        rec = Integration( how, file, srev, erev )
+        self.integrations.append( rec )
+        return rec
+    
+    # iterator over the collection calling a provided function
+    # Python's poor man version of the Ruby block
+    def each_integration(self):
+        for i in self.integrations:
+            yield i
+    
+    def __repr__(self):
+        return "Revision (depotFile = %s rev = %s change = %s action = %s type = %s time = %s user = %s client = %s)" % \
+            (self.depotFile, self.rev, self.change, self.action, self.type, self.time, self.user, self.client)
+#
+# Each DepotFile entry contains details about one depot file.
+#
+class DepotFile:
+    def __init__( self, name ):
+        self.depotFile = name
+        self.revisions = []
+    
+    def new_revision(self):
+        r = Revision( self.depotFile )
+        self.revisions.append( r )
+        return r
+    
+    def each_revision(self):
+        for r in self.revisions:
+            yield r
+
+    def str_revision(self, rev, revFormat, changeFormat):
+        result = "... #{rev:<{rf}} change {change:{cf}} {action:9} on {date}  by {user}@{client} ({type}) '{desc}'". \
+            format(rev = rev.rev, 
+                   rf = revFormat,
+                   change = rev.change,
+                   cf = changeFormat,
+                   action = rev.action, 
+                   date = rev.time, 
+                   user = rev.user, 
+                   client = rev.client, 
+                   type = rev.type,
+                   desc = rev.desc)
+        return result
+ 
+    def str_integration(self, integ):
+        result = "... ... {how} {file}#{srev},{erev}". \
+            format(how = integ.how, file = integ.file, srev = integ.srev, erev = integ.erev)
+        return result
+    
+    def __str__(self):
+        result = "{}".format(self.depotFile)
+        
+        revFormat = len(str(self.revisions[0].rev))
+        changeFormat = len(str(self.revisions[0].change))
+        
+        for rev in self.revisions:
+            result += "\n{}".format(self.str_revision(rev, revFormat, changeFormat))
+            for integ in rev.integrations:
+                result += "\n{}".format(self.str_integration(integ))
+        return result
+    
+    def __repr__(self):
+        return "DepotFile (depotFile = %s, %s revisions)" % ( self.depotFile, len( self.revisions ) )
+
+#
+# Resolver class used in p4.run_resolve()
+#
+# The default simply checks that p4.input is set to sensible value
+# This class is meant to be subclassed for a custom resolver and
+# Resolver.resolve() overriden
+#
+
+class Resolver:
+    def __init__(self):
+        pass
+    
+    def resolve(self, mergeInfo):
+        if mergeInfo.merge_hint == "e":
+            print("Standard resolver encountered merge conflict, skipping resolve")
+            return "s"
+        else:
+            return mergeInfo.merge_hint
+
+    def actionResolve(self, mergeInfo):
+        return mergeInfo.merge_hint
+
+#
+# OutputHandler base class
+#
+# Extend this class if you want to use the handler interface
+#
+
+class OutputHandler:
+    REPORT  = 0
+    HANDLED = 1
+    CANCEL  = 2
+    
+    def __init__(self):
+        pass
+    
+    def outputText(self, s):
+        return OutputHandler.REPORT
+    
+    def outputBinary(self, b):
+        return OutputHandler.REPORT
+    
+    def outputStat(self, h):
+        return OutputHandler.REPORT
+    
+    def outputInfo(self, i):
+        return OutputHandler.REPORT
+    
+    def outputMessage(self, e):
+        return OutputHandler.REPORT
+
+class ReportHandler( OutputHandler ):
+    def __init__(self):
+        OutputHandler.__init__(self)
+    
+    def outputText(self, s):
+        print( "text: ", s)
+        return OutputHandler.HANDLED
+    
+    def outputBinary(self, b):
+        print( "binary: ", b)
+        return OutputHandler.HANDLED
+    
+    def outputStat(self, h):
+        print( "stat:", h)
+        return OutputHandler.HANDLED
+    
+    def outputInfo(self, i):
+        print( "info: ", i)
+        return OutputHandler.HANDLED
+    
+    def outputMessage(self, e):
+        print( "error:", e)
+        return OutputHandler.HANDLED
+
+class Progress:
+    TYPE_SENDFILE = 1
+    TYPE_RECEIVEFILE = 2
+    TYPE_TRANSFER = 3
+    TYPE_COMPUTATION = 4
+    
+    UNIT_PERCENT = 1
+    UNIT_FILES = 2
+    UNIT_KBYTES = 3
+    UNIT_MBYTES = 4
+    
+    def __init__(self):
+        pass
+    
+    def init(self, type):
+        self.type = type
+    
+    def setDescription( self, description, units ):
+        self.description = description
+        self.units = units
+    
+    def setTotal( self, total ):
+        self.total = total
+    
+    def update( self, position ):
+        self.position = position
+    
+    def done( self, fail ):
+        pass
+
+class TextProgress(Progress):
+    TYPES = [ "Unknown", "Submit", "Sync" ]
+    UNITS = [ "Unknown", "Percent", "Files", "KBytes", "MBytes" ]
+    
+    def __init__(self):
+        Progress.__init__(self)
+    
+    def init(self, type):
+        Progress.init(self, type)
+        print( "Progress.init with '%s'" % self.TYPES[type] )
+    
+    def setDescription(self, description, units):
+        Progress.setDescription(self, description, units)
+        print( "Progress.setDescription with '%s' and units '%s'" % (description, self.UNITS[units]) )
+    
+    def setTotal( self, total ):
+        Progress.setTotal(self, total)
+        print( "Progress.setTotal with '%s' " % total )
+    
+    def update( self, position ):
+        Progress.update(self, position )
+        print( "Progress.update with '%s'" % position )
+    
+    def done( self, fail ):
+        Progress.done(self, fail)
+        print( "Progress.done with '%s"'' % fail )
+
+def processFilelog(h):
+    if "depotFile" in h:
+        df = DepotFile( h[ "depotFile" ] )
+        for n, rev in enumerate( h[ "rev" ]):
+            # Create a new revision of this file ready for populating
+            r = df.new_revision()
+            # Populate the base attributes of each revision
+            r.rev = int( rev )
+            r.change = int( h[ "change" ][ n ] )
+            r.action = h[ "action" ][ n ]
+            r.type = h[ "type" ][ n ]
+            r.time = datetime.datetime.utcfromtimestamp( int( h[ "time" ][ n ]) )
+            r.user = h[ "user" ][ n ]
+            r.client = h[ "client" ][ n ]
+            r.desc = h[ "desc" ][ n ]
+            if "digest" in h and n < len(h[ "digest" ]):
+                r.digest = h[ "digest" ][ n ]
+            if "fileSize" in h and n < len(h[ "fileSize" ]):
+                r.fileSize = h[ "fileSize" ][ n ]
+            
+            # Now if there are any integration records for this revision,
+            # add them in too
+            
+            if (not "how" in h) or (n >= len(h["how"]) or h["how"][n] == None):
+                continue
+            else:
+                for m, how in enumerate( h[ "how" ][ n ] ):
+                    file = h[ "file" ][ n ][ m ]
+                    srev = h[ "srev" ][ n ][ m ].lstrip('#')
+                    erev = h[ "erev" ][ n ][ m ].lstrip('#')
+                    
+                    if srev == "none":
+                        srev = 0
+                    else:
+                        srev = int( srev )
+                
+                    if erev == "none":
+                        erev = 0
+                    else:
+                        erev = int( erev )
+        
+                    r.integration( how, file, srev, erev )
+        return df
+    else:
+        raise Exception("Not a filelog object: " + h)
+
+class FilelogOutputHandler(OutputHandler):
+    def __init__(self):
+        OutputHandler.__init__(self)
+    
+    def outputStat(self, h):
+        df = processFilelog(h)
+        
+        return self.outputFilelog(df)
+    
+    def outputFilelog(self, f):
+        return OutputHandler.REPORT
+
+# This is where the C/C++ shared library is loaded
+# It has to be in this place because the library needs to access
+# the classes defined above. Accessing classes defined below this
+# entry would cause an endless loop
+
+import P4API
+
+class P4(P4API.P4Adapter):
+    """Use this class to communicate with a Perforce server
+        
+        Instances of P4 will use the environment settings (including P4CONFIG)
+        to determine the connection parameters such as P4CLIENT and P4PORT.
+        
+        This attributes can also be set separately before connecting.
+        
+        To run any Perforce commands, users of this class first need to run
+        the connect() method.
+        
+        It is good practice to disconnect() after the program is complete.
+        """
+    # Constants useful for exception_level
+    # RAISE_ALL:     Errors and Warnings are raised as exceptions (default)
+    # RAISE_ERROR:   Only Errors are raised as exceptions
+    # RAISE_NONE:    No exceptions are raised, instead False is returned
+    
+    RAISE_ALL    = 2
+    RAISE_ERROR  = 1
+    RAISE_ERRORS = 1
+    RAISE_NONE   = 0
+    
+    # Named values for generic error codes returned by
+    # P4API.Message.generic
+    
+    EV_NONE      = 0     # misc
+    
+    # The fault of the user
+    
+    EV_USAGE     = 0x01  # request not consistent with dox
+    EV_UNKNOWN   = 0x02  # using unknown entity
+    EV_CONTEXT   = 0x03  # using entity in wrong context
+    EV_ILLEGAL   = 0x04  # trying to do something you can't
+    EV_NOTYET    = 0x05  # something must be corrected first
+    EV_PROTECT   = 0x06  # protections prevented operation
+    
+    # No fault at all
+    
+    EV_EMPTY     = 0x11  # action returned empty results
+    
+    # not the fault of the user
+    
+    EV_FAULT     = 0x21  # inexplicable program fault
+    EV_CLIENT    = 0x22  # client side program errors
+    EV_ADMIN     = 0x23  # server administrative action required
+    EV_CONFIG    = 0x24  # client configuration inadequate
+    EV_UPGRADE   = 0x25  # client or server too old to interact
+    EV_COMM      = 0x26  # communications error
+    EV_TOOBIG    = 0x27  # not even Perforce can handle this much
+    
+    # Named values for error severities returned by
+    # P4API.Message.severity
+    E_EMPTY      = 0  # nothing yet
+    E_INFO       = 1  # something good happened
+    E_WARN       = 2  # something not good happened
+    E_FAILED     = 3  # user did something wrong
+    E_FATAL      = 4  # system broken -- nothing can continue
+    
+    # mappings for __iterate
+    # list-of-specs => ( name-of-one-spec, field-name-in-list-of-specs )
+    
+    specfields = {
+        'clients'   :   ('client', 'client'),
+        'labels'    :   ('label', 'label'),
+        'branches'  :   ('branch', 'branch'),
+        'changes'   :   ('change', 'change'),
+        'streams'   :   ('stream', 'Stream'),
+        'jobs'      :   ('job', 'Job'),
+        'users'     :   ('user', 'User'),
+        'groups'    :   ('group', 'group'),
+        'depots'    :   ('depot', 'name'),
+        'servers'   :   ('server', 'Name')
+    }
+    
+    def __init__(self, *args, **kwlist):
+        P4API.P4Adapter.__init__(self, *args, **kwlist)
+    
+    def __del__(self):
+        if self.debug > 3:
+            print("P4.__del__()", file=sys.stderr)
+    
+    def __getattr__(self, name):
+        if name.startswith("run_"):
+            cmd = name[len("run_"):]
+            return lambda *args, **kargs: self.run(cmd, *args, **kargs)
+        elif name.startswith("delete_"):
+            cmd = name[len("delete_"):]
+            return lambda *args, **kargs: self.run(cmd, "-d", *args, **kargs)
+        elif name.startswith("fetch_"):
+            cmd = name[len("fetch_"):]
+            return lambda *args, **kargs: self.__fetch(cmd, *args, **kargs)
+        elif name.startswith("save_"):
+            cmd = name[len("save_"):]
+            return lambda *args, **kargs: self.__save(cmd, *args, **kargs)
+        elif name.startswith("parse_"):
+            cmd = name[len("parse_"):]
+            return lambda *args, **kargs: self.__parse_spec(cmd, *args, **kargs)
+        elif name.startswith("format_"):
+            cmd = name[len("format_"):]
+            return lambda *args, **kargs: self.__format_spec(cmd, *args, **kargs)
+        elif name.startswith("iterate_"):
+            cmd = name[len("iterate_"):]
+            return lambda *args, **kargs: self.__iterate(cmd, *args, **kargs)
+        else:
+            raise AttributeError(name)
+    
+    def __save(self, cmd, *args, **kargs):
+        self.input = args[0]
+        return self.run(cmd, "-i", args[1:], **kargs)
+    
+    def __parse_spec(self, cmd, *args, **kargs):
+        form = args[0]
+        comments = "\n".join( [ x for x in form.split('\n') if x.startswith('#') ] ) + "\n"
+        spec = self.parse_spec(cmd, *args, **kargs)
+        spec.__dict__['comment'] = comments
+        return spec
+    
+    def __format_spec(self, cmd, *args, **kargs):
+        spec = args[0]
+        form = self.format_spec(cmd, *args, **kargs)
+        if 'comment' in spec.__dict__:
+            form = spec.__dict__['comment'] + "\n" +  form
+        return form
+    
+    def __fetch(self, cmd, *args, **kargs):
+        result = self.run(cmd, "-o", *args, **kargs)
+        for r in result:
+            if isinstance(r, tuple) or isinstance(r,dict):
+                return r 
+        return result[0]
+    
+    def __iterate(self, cmd, *args, **kargs):
+        
+        if cmd in self.specfields:
+            specs = self.run(cmd, *args, **kargs)
+            spec = self.specfields[cmd][0]
+            field = self.specfields[cmd][1]
+            
+            # Return a generators (Python iterator object)
+            # On iteration, this will retrieve one spec at a time
+            return ( self.run(spec, '-o', x[field])[0] for x in specs )
+        else:
+            raise Exception('Unknown spec list command: %s', cmd)
+    
+    def __repr__(self):
+        state = "disconnected"
+        if self.connected():
+            state = "connected"
+        
+        return "P4 [%s@%s %s] %s" % \
+            (self.user, self.client, self.port, state)
+    
+    def identify(cls):
+        return P4API.identify()
+    identify = classmethod(identify)
+    
+    def run(self, *args, **kargs):
+        "Generic run method"
+        context = {}
+        resultLogging = True
+        
+        if "resultLogging" in kargs:
+            resultLogging= False
+            del kargs["resultLogging"]
+        
+        for (k,v) in list(kargs.items()):
+            context[k] = getattr(self, k)
+            setattr(self, k, v)
+                
+        flatArgs = self.__flatten(args)
+        if self.logger:
+            self.logger.info("p4 " + " ".join(flatArgs))
+        
+        result = P4API.P4Adapter.run(self, *flatArgs)
+        
+        if resultLogging and self.logger:
+            self.logger.debug(result)
+
+        for (k,v) in list(context.items()):
+            setattr( self, k, v)
+                    
+        return result
+    
+    def run_submit(self, *args, **kargs):
+        "Simplified submit - if any arguments is a dict, assume it to be the changeform"
+        nargs = list(args)
+        form = None
+        for n, arg in enumerate(nargs):
+            if isinstance( arg, dict):
+                self.input = arg
+                nargs.pop(n)
+                nargs.append("-i")
+                break
+        return self.run("submit", *nargs, **kargs)
+    
+    def run_shelve(self, *args, **kargs):
+        "Simplified shelve - if any arguments is a dict, assume it to be the changeform"
+        nargs = list(args)
+        form = None
+        for n, arg in enumerate(nargs):
+            if isinstance( arg, dict):
+                self.input = arg
+                nargs.pop(n)
+                nargs.append("-i")
+                break
+        return self.run("shelve", *nargs, **kargs)
+    
+    def delete_shelve(self, *args, **kargs):
+        "Simplified deletion of shelves - if no -c is passed in, add it to the args"
+        nargs = list(args)
+        if '-c' not in nargs:
+            nargs = ['-c'] + nargs # prepend -c if it is not there
+        nargs = ['-d'] + nargs
+        return self.run("shelve", *nargs, **kargs)
+    
+    def run_login(self, *args):
+        "Simple interface to make login easier"
+        self.input = self.password
+        return self.run("login", *args)
+    
+    def run_password( self, oldpass, newpass ):
+        "Simple interface to allow setting of the password"
+        if( oldpass and len(oldpass) > 0 ):
+            self.input = [ oldpass, newpass, newpass ]
+        else:
+            self.input = [ newpass, newpass ]
+        
+        try:
+            return self.run( "password" )
+        except P4Exception as e:
+            if self.errors and self.errors[0] == "Passwords don't match.":
+                raise P4Exception("Password invalid.")
+
+
+    #
+    # run_filelog: convert "p4 filelog" responses into objects with useful
+    #              methods
+    #
+    # Requires tagged output to be of any real use. If tagged output it not
+    # enabled then you just get the raw data back
+    #
+    def run_filelog( self, *args, **kargs ):
+        kargs["resultLogging"] = False
+        raw = self.run( 'filelog', args, **kargs )
+        if (not self.tagged or not raw):
+            # untagged mode returns simple strings, which breaks the code below
+            # raw could be None if a handler is used
+            return raw
+        result = []
+        for h in raw:
+            df = None
+            if isinstance( h, dict ):
+                df = processFilelog( h )
+            else:
+                df = h
+            result.append( df )
+        
+        # if logger is in the kargs, P4.run has changed its context but already reset it
+        # so we need to pull out the logger again for this special case
+        # run_filelog has its own result debug output because of the post-processing step
+        
+        logger = self.logger
+        if "logger" in kargs:
+            logger = kargs["logger"]
+        
+        if logger:
+            output = "\n\n".join([ str(x) for x in result ])
+            logger.debug(output)
+        
+        return result
+
+    def run_print(self, *args, **kargs):
+        kargs["resultLogging"] = False
+        raw = self.run('print', args, **kargs)
+
+        logger = self.logger
+        if "logger" in kargs:
+            logger = kargs["logger"]
+
+        result = []
+        if raw:
+            debugResult = []
+            for line in raw:
+                if isinstance(line, dict):
+                    result.append(line)
+                    if logger:
+                        debugResult.append(line)
+                    result.append("")
+                else:
+                    # to support encoding for Python 3, we have to do a little dance
+                    # we cannot add bytes to the str "", but we expect that all lines
+                    # are either str or bytes. So if we encounter bytes, we replace the content
+                    try:
+                        result[-1] += line
+                    except TypeError:
+                        if type(line) == bytes and type(result[-1]) == str and result[-1] == "":
+                            result[-1] = line
+                        else:
+                            raise
+            if logger:
+                logger.debug(debugResult)
+            return result
+        else:
+            return []
+
+    def run_resolve(self, *args, **kargs):
+        if self.resolver:
+            myResolver = self.resolver
+        else:
+            myResolver = Resolver()
+        if "resolver" in kargs:
+            myResolver = kargs["resolver"]
+        
+        savedResolver = self.resolver
+        self.resolver = myResolver
+        result = self.run("resolve", args)
+        self.resolver = savedResolver
+        
+        return result
+
+    def run_tickets(self, *args):
+        fname = self.ticket_file
+        with open(fname) as f:
+            tickets_raw = f.readlines()
+        pattern = re.compile('([^=]*)=(.*):([^:]*)\n')
+        tickets = [ pattern.match(x).groups() for x in tickets_raw ]
+        keys = [ "Host", "User", "Ticket" ]
+        result = [ dict(zip(keys, x)) for x in tickets ]
+        return result
+    
+    def run_init(self, *args, **kargs):
+        raise Exception("Please run P4.init() instead")
+        
+    def run_clone(self, *args, **kargs):
+        raise Exception("Please run P4.clone) instead")
+    
+    def __flatten(self, args):
+        result = []
+        if isinstance(args, tuple) or isinstance(args, list):
+            for i in args:
+                result.extend(self.__flatten(i))
+        else:
+            result.append(args)
+        return tuple(result)
+
+    def __enter__( self ):
+        return self
+    
+    def __exit__( self, exc_type, exc_val, exc_tb ):
+        if self.connected():
+            self.disconnect()
+        return False
+    
+    def connect( self ):
+        P4API.P4Adapter.connect( self )
+        return self
+    
+    def is_ignored( self, path ):
+        return P4API.P4Adapter.is_ignored( self, os.path.abspath(path) )
+    
+    @contextmanager
+    def while_tagged( self, t ):
+        old = self.tagged
+        self.tagged = t
+        yield
+        self.tagged = old
+    
+    @contextmanager
+    def at_exception_level( self, e ):
+        old = self.exception_level
+        self.exception_level = e
+        yield
+        self.exception_level = old
+    
+    @contextmanager
+    def using_handler( self, c ):
+        old = self.handler
+        self.handler = c
+        yield
+        self.handler = old
+    
+    @contextmanager
+    def saved_context( self , **kargs):
+        """Saves the context of this p4 object and restores it again at the end of the block"""
+        
+        saved_context = {}
+        for attr in self.__members__:
+            saved_context[attr] = getattr(self, attr)
+        
+        for (k,v) in list(kargs.items()):
+            setattr( self, k, v)
+        
+        yield
+        
+        # now restore the context again. Ignore AttributeError exception
+        # Exception is expected because some attributes only have getters, no setters
+        
+        for (k,v) in list(saved_context.items()):
+            if k not in ("port", "track"):
+                try:
+                    setattr( self, k, v )
+                except AttributeError:
+                    pass # expected for server_level and p4config_file
+
+    @contextmanager
+    def temp_client( self, prefix, template ):
+        """Creates a temporary workspace with a temporary root. 
+            To be used with the "with" statement. Will clean up the temporary root and client
+            workspace after the block has finished.
+            The prefix is prepended to the workspace name and should be used in conjunction with
+            the SpecMap of a spec depot to avoid creating entries there.
+        """
+
+        name = "{prefix}_{id}".format(prefix=prefix, id=str(uuid.uuid1()))
+        ws = self.fetch_client('-t', template, name)
+
+        try:
+            root = tempfile.mkdtemp(prefix=prefix)
+            ws._root = root
+            self.save_client(ws)
+
+            oldName = self.client
+            self.client = name
+            
+            yield ws
+
+            self.client = oldName
+        finally:
+            self.delete_client(name)
+            shutil.rmtree(root)
+
+class Map(P4API.P4Map):
+    def __init__(self, *args):
+        P4API.P4Map.__init__(self, *args)
+        if len(args) > 0:
+            self.insert( *args )
+    
+    LEFT2RIGHT = True
+    RIGHT2LEFT = False
+    
+    def __str__( self ):
+        result = ""
+        
+        for a in self.as_array():
+            result += a + "\n"
+        
+        return result
+    
+    def is_empty(self):
+        """Returns True if this map has no entries yet, otherwise False"""
+        
+        return self.count() == 0
+    
+    def includes(self, *args):
+        return self.translate(*args) != None
+    
+    def reverse(self):
+        return Map(P4API.P4Map.reverse(self).as_array())
+    
+    def insert(self, *args):
+        """Insert an argument to the map. The argument can be:
+            
+            A String,
+            Either of the form "[+-]//lhs/... //rhs/..." or "[+-]//lhs/..."
+            for label style maps.
+            A List:
+            This is a list of strings of one of the single string formats
+            described above.
+            A pair of Strings:
+            P4.Map.insert(lhs, rhs)
+            """
+        
+        if len(args) == 1 :
+            arg = args[0]
+            if isinstance( arg, str ):
+                P4API.P4Map.insert( self, arg )
+            elif isinstance( arg, list ):
+                for s in arg:
+                    P4API.P4Map.insert( self, s )
+        
+        else: # expecting 2 args in this case: left, right
+            left = args[0].strip()
+            right = args[1].strip()
+            P4API.P4Map.insert(self, left, right )
+
+
+def init(*args, **kargs):
+    return __run_dvcs("init", args, **kargs)
+
+def clone(*args, **kargs):
+    return __run_dvcs("clone", args, **kargs)
+
+def __run_dvcs(cmd, *args, **kargs):
+    __check_paths()
+    options = []
+    
+    def add_option(options, name, opt):
+        if name in kargs:
+            options += [opt, kargs[name]]
+    
+    add_option(options, "client","-c")
+    add_option(options, "directory","-d")
+    add_option(options, "user", "-u")
+        
+    named_args = []
+    
+    if "unicode" in kargs:
+        unicode = kargs["unicode"]
+        del kargs["unicode"]
+        if unicode:
+            named_args.append("-xi")
+        else:
+            named_args.append("-n")
+    
+    if "casesensitive" in kargs:
+        casesensitive = kargs["casesensitive"]
+        del kargs["casesensitive"]
+        if casesensitive:
+            named_args.append("-C1")
+        else:
+            named_args.append("-C0")
+    
+    arguments = [ "p4" ]
+    arguments += options
+    arguments += [ cmd ]
+    arguments += named_args 
+    for a in args:
+        arguments += a
+    
+    p = subprocess.Popen(arguments, stdout=subprocess.PIPE, stderr=subprocess.PIPE)
+    (sout,serr) = (p.stdout, p.stderr)
+    results = sout.read()
+    errors = serr.read()
+    sout.close()
+    serr.close()
+    
+    if errors:
+        if type(errors) == bytes:
+            errors = errors.decode('UTF-8')
+        errors = errors.strip()
+        
+        raise Exception("Cmd '{}' raised\n{}".format(" ".join(arguments), errors))
+    
+    if type(results) == bytes:
+        results = results.decode('UTF-8')
+    
+    # this is where we analyze the outcome of the command
+    # move into the resulting directory (if specified)
+    # then create a P4 object and return that (connected??)
+    
+    results = results.strip()
+    
+    if "directory" in kargs:
+        path = os.path.abspath(kargs["directory"])
+        os.chdir(path)
+        os.environ["PWD"] = path
+        
+    # create a copy of kargs without the "directory" key - make it compatible with Python 2.6
+    # new_kargs = { x:kargs[x] for x in kargs if x != "directory" }
+    new_kargs = dict((x,kargs[x]) for x in kargs if x != "directory" )
+    
+    return P4(**new_kargs)
+    
+def __check_paths():
+    if not __exec_exists("p4"):
+        raise Exception("P4 executable not in path")
+    if not __exec_exists("p4d"):
+        raise Exception("P4D executable not in path")
+    
+def __exec_exists(name):
+    execName = name
+    if platform.system() == "Windows":
+        execName += ".exe"
+    
+    for p in os.environ["PATH"].split(os.pathsep):
+        pathToFile = os.path.join(p, execName)
+        if os.path.exists(pathToFile) and os.access(pathToFile, os.X_OK):
+            return __check_version(pathToFile)
+    return False
+    
+def __check_version(pathToFile):
+    p = subprocess.Popen([pathToFile,'-V'], stdout=subprocess.PIPE, stderr=subprocess.PIPE)
+    (sout,serr) = (p.stdout, p.stderr)
+    
+    output = sout.read()
+    error = serr.read()
+
+    sout.close()
+    serr.close()
+    
+    if type(output) == bytes: # Python 3 output is bytes, not str
+        output = output.decode('UTF-8')
+    chunks = output.split(os.linesep)
+    
+    pattern = re.compile("Rev. (?P<Program>.+)/(?P<Platform>.+)/(?P<Release>.+)/(?P<Patch>\d+) \((\d+/\d+/\d+)\).")
+    
+    for c in chunks:
+        match = pattern.match(c)
+        if match:
+            version = match.group('Release')
+            year = int(version.split('.')[0])
+            if year >= 2015:
+                return True
+            else:
+                program = match.group('Program')
+                raise Exception("{} must be at least 2015.1, not {}".format(program, version))
+
+    raise Exception("Unknown P4 output : {}".format(output) )
+
+if __name__ == "__main__":
+    p4 = P4()
+    p4.connect()
+    try:
+        ret = p4.run(sys.argv[1:])
+        for line in ret:
+            if isinstance(line, dict):
+                print("-----")
+                for k in list(line.keys()):
+                    print(k, "=", line[k])
+            else:
+                print(line)
+    except:
+        for e in p4.errors:
+            print(e)
diff --git a/Engine/Extras/Maya_AnimationRiggingTools/ARTv1/General/Scripts/Modules/uExport/__init__.py b/Engine/Extras/Maya_AnimationRiggingTools/ARTv1/General/Scripts/Modules/uExport/__init__.py
new file mode 100644
--- /dev/null
+++ b/Engine/Extras/Maya_AnimationRiggingTools/ARTv1/General/Scripts/Modules/uExport/__init__.py
@@ -0,0 +1 @@
+from uExport import *
\ No newline at end of file
diff --git a/Engine/Extras/Maya_AnimationRiggingTools/ARTv1/General/Scripts/Modules/uExport/uExport.py b/Engine/Extras/Maya_AnimationRiggingTools/ARTv1/General/Scripts/Modules/uExport/uExport.py
new file mode 100644
--- /dev/null
+++ b/Engine/Extras/Maya_AnimationRiggingTools/ARTv1/General/Scripts/Modules/uExport/uExport.py
@@ -0,0 +1,580 @@
+'''
+uExport
+Christopher Evans, Version 0.1, Feb 2014
+@author = Chris Evans
+version = 0.2
+
+Add this to a shelf:
+import uExport as ue
+uExportToolWindow = ue.show()
+
+'''
+
+import os
+
+import shiboken
+import time
+from PySide import QtGui, QtCore
+from cStringIO import StringIO
+import xml.etree.ElementTree as xml
+import pysideuic
+
+import maya.cmds as cmds
+import maya.OpenMayaUI as openMayaUI
+import maya.mel as mel
+
+def show():
+    global uExportToolWindow
+    try:
+        uExportToolWindow.close()
+    except:
+        pass
+
+    uExportToolWindow = uExportTool()
+    uExportToolWindow.show()
+    return uExportToolWindow
+
+def loadUiType(uiFile):
+    """
+    Pyside lacks the "loadUiType" command, so we have to convert the ui file to py code in-memory first
+    and then execute it in a special frame to retrieve the form_class.
+    http://tech-artists.org/forum/showthread.php?3035-PySide-in-Maya-2013
+    """
+    parsed = xml.parse(uiFile)
+    widget_class = parsed.find('widget').get('class')
+    form_class = parsed.find('class').text
+
+    with open(uiFile, 'r') as f:
+        o = StringIO()
+        frame = {}
+
+        pysideuic.compileUi(f, o, indent=0)
+        pyc = compile(o.getvalue(), '<string>', 'exec')
+        exec pyc in frame
+
+        #Fetch the base_class and form class based on their type in the xml from designer
+        form_class = frame['Ui_%s'%form_class]
+        base_class = eval('QtGui.%s'%widget_class)
+    return form_class, base_class
+
+def getMayaWindow():
+    ptr = openMayaUI.MQtUtil.mainWindow()
+    if ptr is not None:
+        return shiboken.wrapInstance(long(ptr), QtGui.QWidget)
+
+try:
+    selfDirectory = os.path.dirname(__file__)
+    uiFile = selfDirectory + '/uExport.ui'
+except:
+    uiFile = 'D:\\Build\\usr\\jeremy_ernst\\MayaTools\\General\\Scripts\\Modules\\uExport\\uExport.ui'
+    form_class, base_class = loadUiType(uiFile)
+
+if os.path.isfile(uiFile):
+    form_class, base_class = loadUiType(uiFile)
+else:
+    cmds.error('Cannot find UI file: ' + uiFile)
+
+#this is used by both classes below, this method is usually in the utils lib at Epic
+def attrExists(attr):
+    if '.' in attr:
+        node, att = attr.split('.')
+        return cmds.attributeQuery(att, node=node, ex=1)
+    else:
+        cmds.warning('attrExists: No attr passed in: ' + attr)
+        return False
+
+class uExport():
+    '''
+    Just a little basket to store things.
+    TODO: Add properties to get/set values
+    '''
+    def __init__(self, node):
+        self.meshes = cmds.listConnections(node + '.rendermesh')
+        self.root = cmds.listConnections(node + '.export_root')
+
+        #joints
+        self.joints = [self.root]
+        children = cmds.listRelatives(self.root, type='joint',allDescendents=True)
+        if children:
+            self.joints.extend(children)
+
+        self.version = cmds.getAttr(node + '.uexport_ver')
+        self.node = node
+        self.name = node.split('|')[-1]
+        self.asset_name = node
+        self.folder_path = None
+        if attrExists(node + '.asset_name'):
+            self.asset_name = cmds.getAttr(node + '.asset_name')
+        if attrExists(node + '.fbx_name'):
+            self.fbx_name = cmds.getAttr(node + '.fbx_name')
+        if attrExists(node + '.folder_path'):
+            self.folder_path = cmds.getAttr(node + '.folder_path')
+
+
+########################################################################
+## UEXPORT TOOL
+########################################################################
+
+
+class uExportTool(base_class, form_class):
+    title = 'uExportTool 0.2'
+
+    currentMesh = None
+    currentSkin = None
+    currentInf = None
+    currentVerts = None
+    currentNormalization = None
+
+    scriptJobNum = None
+    copyCache = None
+
+    jointLoc = None
+
+    iconLib = {}
+    iconPath = os.environ.get('MAYA_LOCATION', None) + '/icons/'
+    iconLib['joint'] = QtGui.QIcon(QtGui.QPixmap(iconPath + 'kinJoint.png'))
+    iconLib['ikHandle'] = QtGui.QIcon(QtGui.QPixmap(iconPath + 'kinHandle.png'))
+    iconLib['transform'] = QtGui.QIcon(QtGui.QPixmap(iconPath + 'orientJoint.png'))
+
+    def __init__(self, parent=getMayaWindow()):
+        self.closeExistingWindow()
+        super(uExportTool, self).__init__(parent)
+
+        self.setupUi(self)
+        self.setWindowTitle(self.title)
+
+        wName = openMayaUI.MQtUtil.fullName(long(shiboken.getCppPointer(self)[0]))
+
+        ## Connect UI
+        ########################################################################
+        self.connect(self.export_BTN, QtCore.SIGNAL("clicked()"), self.export_FN)
+        self.connect(self.createUexportNode_BTN, QtCore.SIGNAL("clicked()"), self.createUexportNode_FN)
+        self.connect(self.replaceUnknownNodes, QtCore.SIGNAL("clicked()"), self.replaceUnknownNodes_FN)
+        self.refreshBTN.clicked.connect(self.refreshUI)
+
+        #context menu
+        self.export_tree.setContextMenuPolicy(QtCore.Qt.CustomContextMenu)
+        self.connect(self.export_tree, QtCore.SIGNAL("customContextMenuRequested(QPoint)" ), self.openMenu)
+
+        self.export_tree.itemClicked.connect(self.check_status)
+
+        self.snapRoot_CMB.setHidden(True)
+        self.refreshUI()
+
+## GENERAL
+########################################################################
+
+    def closeExistingWindow(self):
+        for qt in QtGui.qApp.topLevelWidgets():
+            try:
+                if qt.__class__.__name__ == self.__class__.__name__:
+                    qt.deleteLater()
+                    print 'uExport: Closed ' + str(qt.__class__.__name__)
+            except:
+                pass
+
+    def convertSkelSettingsToNN(delete=1):
+        orig = 'SkeletonSettings_Cache'
+        if cmds.objExists(orig):
+            if cmds.nodeType(orig) == 'unknown':
+                print orig
+                new = cmds.createNode('network')
+                for att in cmds.listAttr(orig):
+                    if not cmds.attributeQuery(att, node=new, exists=1):
+                        typ = cmds.attributeQuery(att, node=orig, at=1)
+                        if typ == 'typed':
+                            cmds.addAttr(new, longName=att, dt='string')
+                            if cmds.getAttr(orig + '.' + att):
+                                cmds.setAttr(new + '.' + att, cmds.getAttr(orig + '.' + att), type='string')
+                        elif typ == 'enum':
+                            cmds.addAttr(new, longName=att, at='enum', enumName=cmds.attributeQuery(att, node=orig, listEnum=1)[0])
+                cmds.delete(orig)
+                cmds.rename(new, 'SkeletonSettings_Cache')
+
+    def replaceUnknownNodes_FN(self):
+        self.convertSkelSettingsToNN()
+
+## UI RELATED
+########################################################################
+
+    def openMenu(self, position):
+        menu = QtGui.QMenu()
+        rendermeshRewire = menu.addAction("Re-wire rendermesh attr to current mesh selection")
+        rootRewire = menu.addAction("Re-wire root joint attr to current selected joint")
+        pos = self.export_tree.mapToGlobal(position)
+        action = menu.exec_(pos)
+        if action == rendermeshRewire:
+            for index in self.export_tree.selectedIndexes():
+                uExportNode = self.export_tree.itemFromIndex(index).uExport.node
+                meshes = cmds.ls(sl=1)
+                #TODO: check if theyre actually meshes
+                self.connectRenderMeshes(uExportNode, meshes)
+    	elif action == rootRewire:
+            index = self.export_tree.selectedIndexes()[0]
+            uExportNode = self.export_tree.itemFromIndex(index).uExport.node
+            root = cmds.ls(sl=1)
+            if len(root) == 1:
+                self.connectRoot(uExportNode, root[0])
+            else:
+                cmds.error('Select a single joint, you == fail, bro. ' + str(root))
+        self.refreshUI()
+
+    def refreshUI(self):
+        self.export_tree.clear()
+        self.uNodes = []
+        for node in self.getExportNodes():
+            self.uNodes.append(uExport(node))
+        self.buildExportTree(self.uNodes)
+
+
+    def buildExportTree(self, uNodes):
+        for uNode in uNodes:
+
+            #top level
+            wid1 = QtGui.QTreeWidgetItem()
+            font = wid1.font(0)
+            font.setPointSize(15)
+
+            wid1.setText(0,uNode.asset_name)
+            wid1.uExport = uNode
+
+            wid1.setText(1, uNode.version)
+            self.export_tree.addTopLevelItem(wid1)
+            wid1.setExpanded(True)
+            wid1.setFont(0,font)
+
+            font = wid1.font(0)
+            font.setPointSize(10)
+
+            wid1.setFlags(QtCore.Qt.ItemIsEnabled | QtCore.Qt.ItemIsEditable | QtCore.Qt.ItemIsUserCheckable | QtCore.Qt.ItemIsSelectable)
+            wid1.setCheckState(0, QtCore.Qt.Checked)
+
+            #mesh branch
+            meshTop = QtGui.QTreeWidgetItem()
+            meshes = uNode.meshes
+            if meshes:
+                meshTop.setText(0, 'RENDER MESHES: (' + str(len(uNode.meshes)) + ')')
+            else:
+                meshTop.setText(0, 'RENDER MESHES: NONE')
+            meshTop.setFlags(QtCore.Qt.ItemIsEnabled | QtCore.Qt.ItemIsEditable | QtCore.Qt.ItemIsUserCheckable)
+            meshTop.setCheckState(0, QtCore.Qt.Checked)
+            wid1.addChild(meshTop)
+
+            #meshes
+            if uNode.meshes:
+                for mesh in uNode.meshes:
+                    meshWid = QtGui.QTreeWidgetItem()
+                    meshWid.setText(0, mesh)
+                    meshTop.addChild(meshWid)
+
+            if uNode.root:
+                #anim branch
+                animTop = QtGui.QTreeWidgetItem()
+                jnts = uNode.joints
+                if jnts:
+                    animTop.setText(0, 'ANIMATION:  (' + str(len(uNode.joints)) + ' JOINTS)')
+                else:
+                    animTop.setText(0, 'ANIMATION: NONE')
+                animTop.setFlags(QtCore.Qt.ItemIsEnabled | QtCore.Qt.ItemIsEditable | QtCore.Qt.ItemIsUserCheckable)
+                animTop.setCheckState(0, QtCore.Qt.Checked)
+                wid1.addChild(animTop)
+
+                #anims
+                animWid = QtGui.QTreeWidgetItem()
+                animWid.setText(0, '<< CURRENT TIME RANGE >>')
+                animTop.addChild(animWid)
+
+            #meta branch
+            metaTop = QtGui.QTreeWidgetItem()
+            metaTop.setText(0, 'METADATA')
+            wid1.addChild(metaTop)
+
+            #meta
+            fpathWid = QtGui.QTreeWidgetItem()
+            fpathWid.setText(0, 'FOLDER_PATH:   ' + str(uNode.folder_path))
+            metaTop.addChild(fpathWid)
+            #meta
+            fbxName = QtGui.QTreeWidgetItem()
+            fbxName.setText(0, 'FBX_FILE:   ' + str(uNode.fbx_name))
+            metaTop.addChild(fbxName)
+
+    def createUexportNode_FN(self):
+        if cmds.ls(sl=1):
+            if self.useRoot_CHK.isChecked():
+                rootName = self.rootName_CMB.currentText()
+                self.create(renderMeshes=cmds.ls(sl=1), rootJoint=rootName)
+            else:
+                #TODO: modal picker with joint filter
+                pass
+
+    def export_FN(self):
+        #get what should be export in an exportTask fn
+        exportWidgets = self.getExportNodeWidgets()
+        for wid in exportWidgets:
+            snapConst = None
+            if self.snapRoot_CHK.isChecked():
+                node = self.snapRoot_CMB.currentText()
+                try:
+                    if cmds.objExists(node):
+                        snapConst = cmds.parentConstraint(node, wid.uExport.root)
+                    else:
+                        cmds.error('Unable to find node: ' + node)
+                except:
+                    cmds.warning('Could not constrain root' + wid.uExport.root + ', is it already constrained?')
+            #export
+            mpath = cmds.file(sceneName=1, q=1)
+            fname = mpath.split('/')[-1]
+            fpath = mpath.replace(fname,'').replace('/','\\')
+            if wid.uExport.fbx_name:
+                if wid.uExport.fbx_name != '':
+                    fname = wid.uExport.fbx_name
+            if wid.uExport.folder_path:
+                if wid.uExport.folder_path != '':
+                    fpath = wid.uExport.folder_path
+
+            #prompt user for save location
+            userPath = QtGui.QFileDialog.getSaveFileName(caption = 'Export ' + wid.uExport.name + ' to:', filter='FBX Files (*.fbx)', dir=fpath + '\\' + fname)[0]
+
+            if userPath:
+                wid.uExport.fbx_name = userPath.split('/')[-1]
+                wid.uExport.folder_path = userPath.replace(wid.uExport.fbx_name,'')
+
+                #if set to save loc, save it
+                if self.rememberSaveCHK.isChecked():
+                    cmds.setAttr(wid.uExport.name + '.folder_path', wid.uExport.folder_path, type='string')
+                    cmds.setAttr(wid.uExport.name + '.fbx_name', wid.uExport.fbx_name, type='string')
+
+                meshChk = 1
+                animChk = 1
+                for c in range(0, wid.childCount()):
+                    if 'MESH' in wid.child(c).text(0):
+                        if wid.child(c).checkState(0) == QtCore.Qt.Checked:
+                            pass
+                        else:
+                            meshChk = 0
+
+                start = time.time()
+                #put export into the uExport class later
+                self.export(wid.uExport.name, path=userPath, mesh=meshChk)
+                elapsed = (time.time() - start)
+                print 'uExport>>> Exported ', wid.text(0), 'to', userPath ,'in %.2f seconds.' % elapsed
+
+            else:
+                cmds.warning('Invalid path specified: [' + str(userPath) + ']')
+            #cleanup constraint
+            if snapConst: cmds.delete(snapConst)
+
+    def getExportNodeWidgets(self):
+        nodes = []
+        for i in range(0, self.export_tree.topLevelItemCount()):
+            if self.export_tree.topLevelItem(i).checkState(0) == QtCore.Qt.Checked:
+                nodes.append(self.export_tree.topLevelItem(i))
+        return nodes
+
+
+    def check_status(self):
+        for i in range(0, self.export_tree.topLevelItemCount()):
+            if self.export_tree.topLevelItem(i).checkState(0) == QtCore.Qt.Unchecked:
+                for c in range(0, self.export_tree.topLevelItem(i).childCount()):
+                    self.export_tree.topLevelItem(i).child(c).setCheckState(0,QtCore.Qt.Unchecked)
+
+## P4 CRAP
+
+## UEXPORT NODE
+########################################################################
+    def getExportNodes(self):
+        return cmds.ls('*.uexport_ver', o=1, r=1)
+
+    def connectRenderMeshes(self, uNode, renderMeshes, rewire=1):
+        try:
+            cmds.undoInfo(openChunk=True)
+            if rewire:
+                for conn in cmds.listConnections(uNode + '.rendermesh', plugs=1, destination=1):
+                    cmds.disconnectAttr(uNode + '.rendermesh', conn)
+        
+            for node in renderMeshes:
+                if not attrExists(node+'.export'):
+                    cmds.addAttr(node, longName='export', attributeType='message')
+                cmds.connectAttr(uNode + '.rendermesh', node + '.export' )
+        except Exception as e:
+            print e
+        finally:
+            cmds.undoInfo(closeChunk=True)
+
+    def connectRoot(self, uNode, root, rewire=1):
+        try:
+            cmds.undoInfo(openChunk=True)
+            if rewire:
+                conns = cmds.listConnections(uNode + '.export_root', plugs=1, source=1)
+                if conns:
+                    for conn in conns:
+                        cmds.disconnectAttr(conn, uNode + '.export_root')
+        
+            if not attrExists(root+'.export'):
+                cmds.addAttr(root, longName='export', attributeType='message')
+            cmds.connectAttr(root + '.export', uNode + '.export_root' )
+        except Exception as e:
+            print e
+        finally:
+            cmds.undoInfo(closeChunk=True)
+        
+    def create(self, renderMeshes=None, rootJoint=None, strName='uExport'):
+        uExport = None
+        if cmds.objExists(strName):
+            #later re-hook up
+            #set uExport
+            pass
+        else:
+            uExport = cmds.group(em=1, name=strName)
+            cmds.addAttr(uExport, ln='rendermesh', at='message')
+            cmds.addAttr(uExport, ln='export_root', at='message')
+            cmds.addAttr(uExport, ln='materials', at='message')
+            cmds.addAttr(uExport, ln='uexport_ver', dt='string')
+            cmds.setAttr(uExport + '.uexport_ver', '1.0', type='string')
+            cmds.addAttr(uExport, ln='folder_path', dt='string')
+            cmds.addAttr(uExport, ln='asset_name', dt='string')
+            cmds.addAttr(uExport, ln='fbx_name', dt='string')
+
+        if uExport:
+            try:
+                if renderMeshes:
+                    self.connectRenderMeshes(uExport, renderMeshes, rewire=0)
+
+                #rootJoint
+                if rootJoint:
+                    if not attrExists(rootJoint+'.export'):
+                        cmds.addAttr(rootJoint, longName='export', attributeType='message')
+                    cmds.connectAttr(rootJoint + '.export', uExport + '.export_root')
+                else:
+                    cmds.warning('No root joint or could not find root: ' + str(rootJoint))
+            except Exception as e:
+                print cmds.warning(e)
+
+
+
+## EXPORT
+########################################################################
+
+    #TODO: Find and export blendshape meshes!
+    def setExportFlags(self):
+
+        # Mesh
+        mel.eval("FBXExportSmoothingGroups -v true")
+        mel.eval("FBXExportHardEdges -v false")
+        mel.eval("FBXExportTangents -v false")
+        mel.eval("FBXExportInstances -v false")
+        mel.eval("FBXExportInAscii -v true")
+        mel.eval("FBXExportSmoothMesh -v false")
+
+        # Animation
+        mel.eval("FBXExportBakeComplexAnimation -v true")
+        mel.eval("FBXExportBakeComplexStart -v "+str(cmds.playbackOptions(minTime=1, q=1)))
+        mel.eval("FBXExportBakeComplexEnd -v "+str(cmds.playbackOptions(maxTime=1, q=1)))
+        mel.eval("FBXExportReferencedAssetsContent -v true")
+        mel.eval("FBXExportBakeComplexStep -v 1")
+        mel.eval("FBXExportUseSceneName -v false")
+        if self.euler_CHK.isChecked():
+            mel.eval("FBXExportQuaternion -v euler")
+        if self.quat_CHK.isChecked():
+            mel.eval("FBXExportQuaternion -v quaternion")
+        if self.resample_CHK.isChecked():
+            mel.eval("FBXExportQuaternion -v resample")
+        mel.eval("FBXExportShapes -v true")
+        mel.eval("FBXExportSkins -v true")
+
+        #up axis, the rather large elephant in the room
+        idx = self.upAxisCMB.currentIndex()
+        if idx != 0:
+            if idx == 1:
+                mel.eval("FBXExportUpAxis y")
+            if idx == 2:
+                mel.eval("FBXExportUpAxis z")
+
+
+        #garbage we don't want
+        # Constraints
+        mel.eval("FBXExportConstraints -v false")
+        # Cameras
+        mel.eval("FBXExportCameras -v false")
+        # Lights
+        mel.eval("FBXExportLights -v false")
+        # Embed Media
+        mel.eval("FBXExportEmbeddedTextures -v false")
+        # Connections
+        mel.eval("FBXExportInputConnections -v false")
+        # Axis Conversion
+
+    def export(self, node, mesh=1, anim=1, path=None):
+
+        toExport = []
+
+        if not path:
+            oldPath = path
+            path = cmds.file(sceneName=1, q=1)
+            cmds.warning('No valid path set for export [' + str(oldPath) + ']/nExporting to Maya file loc: ' + path)
+            toExport.extend(cmds.listRelatives(cmds.listConnections(node + '.export_root'), type='joint',allDescendents=True,f=1))
+
+        #kvassey -- adding support for baking to root in rig before export
+        if self.bakeRoot_CHK.isChecked():
+            print "Bake Root Checked"
+            #copy root skeleton with input connections under new group
+            cmds.select(cl=True)
+            currRoot = cmds.listConnections(node + '.export_root')
+            exportSkel = cmds.duplicate(currRoot, un=True, rc=False, po=False)
+            print exportSkel
+            tempGrp = cmds.group(exportSkel[0])
+            #rename root joint
+            dupRoot = cmds.rename(exportSkel[0], currRoot)
+            #bake
+            startTime = cmds.playbackOptions(min=True, q=True)
+            endTime = cmds.playbackOptions(max=True, q=True)
+            cmds.bakeResults(dupRoot, sm=True,  hi="below", s=True, sb=1, dic=True, t=(startTime, endTime))
+
+            #move FBX export inside here, skip rest.
+            #toExport.extend(cmds.listRelatives(cmds.listConnections(dupRoot), type='joint',allDescendents=True))
+            #cmds.select(toExport)
+            cmds.select(dupRoot, r=True, hi=True)
+            toExport = cmds.ls(sl=True, type='joint')
+            cmds.select(toExport, r=True)
+            self.setExportFlags()
+            # Export!
+            print "FBXExport -f \""+ path +"\" -s"
+            mel.eval("FBXExport -f \""+ path +"\" -s")
+
+            cmds.delete(tempGrp)
+            toExport = []
+
+
+        #Assuming the meshes are skinned, maybe work with static later
+        if not self.bakeRoot_CHK.isChecked():
+            if not self.fbxPerMesh_CHK.isChecked():
+                if mesh:
+                    toExport.extend(cmds.listConnections(node + '.rendermesh'))
+                if anim:
+                    toExport.extend(cmds.listRelatives(cmds.listConnections(node + '.export_root'), type='joint', allDescendents=True, f=1))
+
+                cmds.select(toExport)
+                self.setExportFlags()
+                # Export!
+                print "FBXExport -f \""+ path +"\" -s"
+                mel.eval("FBXExport -f \""+ path +"\" -s")
+
+            else:
+                for msh in cmds.listConnections(node + '.rendermesh'):
+                    toExport.append(msh)
+                    toExport.extend(cmds.listRelatives(cmds.listConnections(node + '.export_root'), type='joint',allDescendents=True, f=1))
+                    cmds.select(toExport)
+                    self.setExportFlags()
+                    filePath = os.path.dirname(path)
+                    new_fpath = filePath + '/' + msh + '.fbx'
+                    # Export!
+                    print "FBXExport -f \"" + new_fpath + "\" -s"
+                    mel.eval("FBXExport -f \"" + new_fpath + "\" -s")
+
+                    #clear sel and export list
+                    cmds.select(d=1)
+                    toExport = []
+
+
+if __name__ == '__main__':
+    show()
\ No newline at end of file
