Branching

diff --git a/Engine/Extras/Maya_AnimationRiggingTools/ARTv1/General/Scripts/P4.py b/Engine/Extras/Maya_AnimationRiggingTools/ARTv1/General/Scripts/P4.py
new file mode 100644
--- /dev/null
+++ b/Engine/Extras/Maya_AnimationRiggingTools/ARTv1/General/Scripts/P4.py
@@ -0,0 +1,757 @@
+#!/usr/local/bin/python
+
+from __future__ import print_function
+
+"""A Python version of the Perforce "p4" client.
+
+This uses the Python type P4API.P4Adapter, which is a wrapper for the
+Perforce ClientApi object. 
+
+$Id: //depot/r12.2/p4-python/P4.py#1 $
+
+#*******************************************************************************
+# Copyright (c) 2007-2010, Perforce Software, Inc.	All rights reserved.
+# Portions Copyright (c) 1999, Mike Meyer. All rights reserved.
+# Portions Copyright (c) 2004-2007, Robert Cowham. All rights reserved.
+#
+# Redistribution and use in source and binary forms, with or without
+# modification, are permitted provided that the following conditions are met:
+#
+# 1.  Redistributions of source code must retain the above copyright
+# notice, this list of conditions and the following disclaimer.
+#
+# 2.  Redistributions in binary form must reproduce the above copyright
+# notice, this list of conditions and the following disclaimer in the
+# documentation and/or other materials provided with the distribution.
+#
+# THIS SOFTWARE IS PROVIDED BY THE COPYRIGHT HOLDERS AND CONTRIBUTORS "AS IS"
+# AND ANY EXPRESS OR IMPLIED WARRANTIES, INCLUDING, BUT NOT LIMITED TO, THE
+# IMPLIED WARRANTIES OF MERCHANTABILITY AND FITNESS FOR A PARTICULAR PURPOSE
+# ARE DISCLAIMED. IN NO EVENT SHALL PERFORCE SOFTWARE, INC. BE LIABLE FOR ANY
+# DIRECT, INDIRECT, INCIDENTAL, SPECIAL, EXEMPLARY, OR CONSEQUENTIAL DAMAGES
+# (INCLUDING, BUT NOT LIMITED TO, PROCUREMENT OF SUBSTITUTE GOODS OR SERVICES;
+# LOSS OF USE, DATA, OR PROFITS; OR BUSINESS INTERRUPTION) HOWEVER CAUSED AND
+# ON ANY THEORY OF LIABILITY, WHETHER IN CONTRACT, STRICT LIABILITY, OR TORT
+# (INCLUDING NEGLIGENCE OR OTHERWISE) ARISING IN ANY WAY OUT OF THE USE OF THIS
+# SOFTWARE, EVEN IF ADVISED OF THE POSSIBILITY OF SUCH DAMAGE.
+#*******************************************************************************
+
+Copyright 2007-2010 Perforce Software, Inc. All rights reserved
+
+
+License:
+See accompanying LICENSE.txt including for redistribution permission.
+"""
+
+import sys, os, string, datetime
+import types, re
+from contextlib import contextmanager
+
+# P4Exception - some sort of error occurred
+class P4Exception(Exception):
+	"""Exception thrown by P4 in case of Perforce errors or warnings"""
+	
+	def __init__(self, value):
+		Exception.__init__(self)
+		
+		if isinstance(value, (list, tuple)) and len(value) > 2:
+			self.value = value[0]
+			self.errors = value[1]
+			self.warnings = value[2]
+		else: 
+			self.value = value
+	def __str__(self):
+		return str(self.value)
+
+class Spec(dict):
+	"""Subclass of dict, representing the fields of a spec definition.
+	
+	Attributes can be accessed either with the conventional dict format,
+	spec['attribute'] or with shorthand spec._attribute.
+	
+	Instances of this class will preventing any unknown keys.
+	"""
+	def __init__(self, fieldmap=None):
+		self.__dict__['_Spec__fields'] = fieldmap
+	
+	def permitted_fields(self):
+		return self.__fields
+	
+	def __setitem__(self, key, value):
+		if not isinstance(value, str) and not isinstance(value, list):
+			raise P4Exception("Illegal value of type %s, must be string or list" % value.__class__)
+		if key in self or self.__fields == None:
+			dict.__setitem__(self, key, value)
+		elif str(key).lower() in self.__fields:
+			dict.__setitem__(self, self.__fields[key.lower()], value)
+		else:
+			raise P4Exception("Illegal field '%s'" % str(key))
+	
+	def __getattr__(self, attr):
+		key = str(attr).lower()
+		if key[0] != '_':
+			raise AttributeError(attr)
+		key = key[1:]
+		if key in self:
+			return self[key]
+		elif key in self.__fields:
+			return self[self.__fields[key]]
+					
+	def __setattr__(self, attr, value):
+		key = str(attr).lower()
+		if key[0] != '_':
+			raise AttributeError(attr)
+		key = key[1:]
+		self[key] = value
+
+#
+# P4Integration objects hold details about the integrations that have
+# been performed on a particular revision. Used primarily with the
+# P4Revision class
+#
+class Integration:
+	def __init__( self, how, file, srev, erev ):
+		self.how = how
+		self.file = file
+		self.srev = srev
+		self.erev = erev
+
+	def __repr__(self):
+		return "Integration (how = %s file = %s srev = %s erev = %s)" \
+			% (self.how, self.file, self.srev, self.erev)
+#
+# Each P4Revision object holds details about a particular revision
+# of a file. It may also contain the history of any integrations 
+# to/from the file
+#
+
+class Revision:
+	def __init__( self, depotFile ):
+		self.depotFile = depotFile
+		self.integrations = []
+		self.rev = None
+		self.change = None
+		self.action = None
+		self.type = None
+		self.time = None
+		self.user = None
+		self.client = None
+		self.desc = None
+		self.digest = None
+		self.fileSize = None
+
+	def integration( self, how, file, srev, erev ):
+		rec = Integration( how, file, srev, erev )
+		self.integrations.append( rec )
+		return rec
+
+# iterator over the collection calling a provided function
+# Python's poor man version of the Ruby block 
+	def each_integration(self):
+		for i in self.integrations:
+			yield i
+
+	def __repr__(self):
+		return "Revision (depotFile = %s rev = %s change = %s action = %s type = %s time = %s user = %s client = %s)" % \
+			(self.depotFile, self.rev, self.change, self.action, self.type, self.time, self.user, self.client)
+#
+# Each DepotFile entry contains details about one depot file. 
+# 
+class DepotFile:
+	def __init__( self, name ):
+		self.depotFile = name
+		self.revisions = []
+
+	def new_revision(self):
+		r = Revision( self.depotFile )
+		self.revisions.append( r )
+		return r
+	
+	def each_revision(self):
+		for r in self.revisions:
+			yield r
+
+	def __repr__(self):
+		return "DepotFile (depotFile = %s, %s revisions)" % ( self.depotFile, len( self.revisions ) )
+
+#
+# Resolver class used in p4.run_resolve()
+#
+# The default simply checks that p4.input is set to sensible value
+# This class is meant to be subclassed for a custom resolver and
+# Resolver.resolve() overriden
+#
+
+class Resolver:
+	def __init__(self):
+		pass
+	
+	def resolve(self, mergeInfo):
+		if mergeInfo.merge_hint == "e":
+			print("Standard resolver encountered merge conflict, skipping resolve")
+			return "s"
+		else:
+			return mergeInfo.merge_hint
+
+	def actionResolve(self, mergeInfo):
+		return mergeInfo.merge_hint
+
+#
+# OutputHandler base class 
+#
+# Extend this class if you want to use the handler interface
+#
+
+class OutputHandler:
+	REPORT	= 0
+	HANDLED = 1
+	CANCEL	= 2
+	
+	def __init__(self):
+		pass
+	
+	def outputText(self, s):
+		return OutputHandler.REPORT
+	
+	def outputBinary(self, b):
+		return OutputHandler.REPORT
+	
+	def outputStat(self, h):
+		return OutputHandler.REPORT
+	
+	def outputInfo(self, i):
+		return OutputHandler.REPORT
+	
+	def outputMessage(self, e):
+		return OutputHandler.REPORT
+
+class ReportHandler( OutputHandler ):
+	def __init__(self):
+		OutputHandler.__init__(self)
+
+	def outputText(self, s):
+		print( "text: ", s)
+		return OutputHandler.HANDLED
+	
+	def outputBinary(self, b):
+		print( "binary: ", b)
+		return OutputHandler.HANDLED
+	
+	def outputStat(self, h):
+		print( "stat:", h)
+		return OutputHandler.HANDLED
+	
+	def outputInfo(self, i):
+		print( "info: ", i)
+		return OutputHandler.HANDLED
+	
+	def outputMessage(self, e):
+		print( "error:", e)
+		return OutputHandler.HANDLED
+	
+class Progress:
+	def __init__(self):
+		pass
+	
+	def init(self, type):
+		self.type = type
+	
+	def setDescription( self, description, units ):
+		self.description = description
+		self.units = units
+	
+	def setTotal( self, total ):
+		self.total = total
+	
+	def update( self, position ):
+		self.position = position
+		
+	def done( self, fail ):
+		pass
+
+class TextProgress(Progress):
+	def __init__(self):
+		Progress.__init__(self)
+	
+	def init(self, type):
+		Progress.init(self, type)
+		print( "Progress.init with %s" % type )
+	
+	def setDescription(self, description, units):
+		Progress.setDescription(self, description, units)
+		print( "Progress.setDescription with %s and units %s" % (description, units) )
+
+	def setTotal( self, total ):
+		Progress.setTotal(self, total)
+		print( "Progress.setTotal with %s " % total )
+	
+	def update( self, position ):
+		Progress.update(self, position )
+		print( "Progress.update with %s" % position )
+
+	def done( self, fail ):
+		Progress.done(self, fail)
+		print( "Progress.done with %s" % fail )
+
+def processFilelog(h):
+	if "depotFile" in h:
+		df = DepotFile( h[ "depotFile" ] )
+		for n, rev in enumerate( h[ "rev" ]):
+			# Create a new revision of this file ready for populating
+			r = df.new_revision()
+			# Populate the base attributes of each revision
+			r.rev = int( rev )
+			r.change = int( h[ "change" ][ n ] )
+			r.action = h[ "action" ][ n ] 
+			r.type = h[ "type" ][ n ]
+			r.time = datetime.datetime.utcfromtimestamp( int( h[ "time" ][ n ]) )
+			r.user = h[ "user" ][ n ]
+			r.client = h[ "client" ][ n ]
+			r.desc = h[ "desc" ][ n ]
+			if "digest" in h and n < len(h[ "digest" ]):
+			  r.digest = h[ "digest" ][ n ]
+			if "fileSize" in h and n < len(h[ "fileSize" ]):
+			  r.fileSize = h[ "fileSize" ][ n ]
+	
+			# Now if there are any integration records for this revision,
+			# add them in too
+			
+			if (not "how" in h) or (n >= len(h["how"]) or h["how"][n] == None):
+				continue
+			else:
+				for m, how in enumerate( h[ "how" ][ n ] ):
+					file = h[ "file" ][ n ][ m ]
+					srev = h[ "srev" ][ n ][ m ].lstrip('#')
+					erev = h[ "erev" ][ n ][ m ].lstrip('#')
+					
+					if srev == "none":
+						srev = 0
+					else:
+						srev = int( srev )
+				  
+					if erev == "none":
+						erev = 0
+					else:
+						erev = int( erev )
+						
+					r.integration( how, file, srev, erev )
+		return df
+	else:
+		raise Exception("Not a filelog object: " + h)
+
+class FilelogOutputHandler(OutputHandler):
+	def __init__(self):
+		OutputHandler.__init__(self)
+		
+	def outputStat(self, h):
+		df = processFilelog(h)
+		
+		return self.outputFilelog(df)
+	
+	def outputFilelog(self, f):
+		return OutputHandler.REPORT
+
+# This is where the C/C++ shared library is loaded
+# It has to be in this place because the library needs to access
+# the classes defined above. Accessing classes defined below this
+# entry would cause an endless loop
+
+import P4API
+#import p4_adsk.P4API as P4API
+
+class P4(P4API.P4Adapter):
+	"""Use this class to communicate with a Perforce server
+	
+	Instances of P4 will use the environment settings (including P4CONFIG)
+	to determine the connection parameters such as P4CLIENT and P4PORT.
+	
+	This attributes can also be set separately before connecting.
+	
+	To run any Perforce commands, users of this class first need to run
+	the connect() method.
+	
+	It is good practice to disconnect() after the program is complete.
+	"""
+	# Constants useful for exception_level
+	# RAISE_ALL:	 Errors and Warnings are raised as exceptions (default)
+	# RAISE_ERROR:	 Only Errors are raised as exceptions
+	# RAISE_NONE:	 No exceptions are raised, instead False is returned
+	
+	RAISE_ALL	 = 2
+	RAISE_ERROR	 = 1
+	RAISE_ERRORS = 1
+	RAISE_NONE	 = 0
+
+	# Named values for generic error codes returned by 
+	# P4API.Message.generic
+	
+	EV_NONE		 = 0	 # misc
+	
+	# The fault of the user
+	
+	EV_USAGE	 = 0x01	 # request not consistent with dox
+	EV_UNKNOWN	 = 0x02	 # using unknown entity
+	EV_CONTEXT	 = 0x03	 # using entity in wrong context
+	EV_ILLEGAL	 = 0x04	 # trying to do something you can't
+	EV_NOTYET	 = 0x05	 # something must be corrected first
+	EV_PROTECT	 = 0x06	 # protections prevented operation
+	
+	# No fault at all
+	
+	EV_EMPTY	 = 0x11	 # action returned empty results
+	
+	# not the fault of the user
+	
+	EV_FAULT	 = 0x21	 # inexplicable program fault
+	EV_CLIENT	 = 0x22	 # client side program errors
+	EV_ADMIN	 = 0x23	 # server administrative action required
+	EV_CONFIG	 = 0x24	 # client configuration inadequate
+	EV_UPGRADE	 = 0x25	 # client or server too old to interact
+	EV_COMM		 = 0x26	 # communications error
+	EV_TOOBIG	 = 0x27	 # not even Perforce can handle this much
+	
+	# Named values for error severities returned by
+	# P4API.Message.severity
+	E_EMPTY		 = 0  # nothing yet
+	E_INFO		 = 1  # something good happened
+	E_WARN		 = 2  # something not good happened
+	E_FAILED	 = 3  # user did something wrong
+	E_FATAL		 = 4  # system broken -- nothing can continue
+
+	# mappings for __iterate
+	# list-of-specs => ( name-of-one-spec, field-name-in-list-of-specs )
+	
+	specfields = { 
+				'clients'	:	('client', 'client'),
+				'labels'	:	('label', 'label'),
+				'branches'	:	('branch', 'branch'),
+				'changes'	:	('change', 'change'),
+				'streams'	:	('stream', 'Stream'),
+				'jobs'		:	('job', 'Job'),
+				'users'		:	('user', 'User'),
+				'groups'	:	('group', 'group'),
+				'depots'	:	('depot', 'name'),
+				'servers'	:	('server', 'Name')
+	}
+	
+	def __init__(self, *args, **kwlist): 
+		P4API.P4Adapter.__init__(self, *args, **kwlist)
+
+	def __del__(self):
+		if self.debug > 3:
+			print("P4.__del__()", file=sys.stderr)
+			
+	# store the references to the created lambdas as a weakref to allow Python 
+	# to clean up the garbage. |The lambda as a closure stores a reference to self
+	# which causes a circular reference problem without the weakref
+	
+	def __getattr__(self, name):
+		if name.startswith("run_"):
+			cmd = name[len("run_"):]
+			return lambda *args, **kargs: self.run(cmd, *args, **kargs)
+		elif name.startswith("delete_"):
+			cmd = name[len("delete_"):]
+			return lambda *args, **kargs: self.run(cmd, "-d", *args, **kargs)
+		elif name.startswith("fetch_"):
+			cmd = name[len("fetch_"):]
+			return lambda *args, **kargs: self.run(cmd, "-o", *args, **kargs)[0]
+		elif name.startswith("save_"):
+			cmd = name[len("save_"):]
+			return lambda *args, **kargs: self.__save(cmd, *args, **kargs)
+		elif name.startswith("parse_"):
+			cmd = name[len("parse_"):]
+			return lambda *args, **kargs: self.parse_spec(cmd, *args, **kargs)
+		elif name.startswith("format_"):
+			cmd = name[len("format_"):]
+			return lambda *args, **kargs: self.format_spec(cmd, *args, **kargs)
+		elif name.startswith("iterate_"):
+			cmd = name[len("iterate_"):]
+			return lambda *args, **kargs: self.__iterate(cmd, *args, **kargs)
+		else:
+			raise AttributeError(name)
+	
+	def __save(self, cmd, *args, **kargs):
+		self.input = args[0]
+		return self.run(cmd, "-i", args[1:], **kargs)
+	
+	def __iterate(self, cmd, *args, **kargs):
+
+		specs = self.run(cmd, *args, **kargs)
+		if cmd in self.specfields:
+			spec = self.specfields[cmd][0]
+			field = self.specfields[cmd][1]
+			
+			# Return a generators (Python iterator object)
+			# On iteration, this will retrieve one spec at a time
+			return ( self.run(spec, '-o', x[field])[0] for x in specs )
+		else:
+			raise Exception('Unknown spec list command: %s', cmd)
+	
+	def __repr__(self):
+		state = "disconnected"
+		if self.connected():
+			state = "connected"
+			
+		return "P4 [%s@%s %s] %s" % \
+		  (self.user, self.client, self.port, state)
+	
+	def identify(cls):
+		return P4API.identify()
+	identify = classmethod(identify)
+	
+	def run(self, *args, **kargs):
+		"Generic run method"
+		context = {}
+		
+		for (k,v) in list(kargs.items()):
+			context[k] = getattr(self, k)
+			setattr(self, k, v)
+			
+		result = P4API.P4Adapter.run(self, *self.__flatten(args))
+		
+		for (k,v) in list(context.items()):
+			setattr( self, k, v)
+
+		return result
+		
+	def run_submit(self, *args, **kargs):
+		"Simplified submit - if any arguments is a dict, assume it to be the changeform"
+		nargs = list(args)
+		form = None
+		for n, arg in enumerate(nargs):
+			if isinstance( arg, dict):
+				self.input = arg
+				nargs.pop(n)
+				nargs.append("-i")
+				break
+		return self.run("submit", *nargs, **kargs)
+	
+	def run_shelve(self, *args):
+		"Simplified shelve - if any arguments is a dict, assume it to be the changeform"
+		nargs = list(args)
+		form = None
+		for n, arg in enumerate(nargs):
+			if isinstance( arg, dict):
+				self.input = arg
+				nargs.pop(n)
+				nargs.append("-i")
+				break
+		return self.run("shelve", *nargs)
+
+	def delete_shelve(self, *args):
+		"Simplified deletion of shelves - if no -c is passed in, add it to the args"
+		nargs = list(args)
+		if '-c' not in nargs:
+			nargs = ['-c'] + nargs # prepend -c if it is not there
+		nargs = ['-d'] + nargs
+		return self.run("shelve", *nargs)
+	
+	def run_login(self, *args):
+		"Simple interface to make login easier"
+		self.input = self.password
+		return self.run("login", *args)
+
+	def run_password( self, oldpass, newpass ):
+		"Simple interface to allow setting of the password"
+		if( oldpass and len(oldpass) > 0 ):
+			self.input = [ oldpass, newpass, newpass ]
+		else:
+			self.input = [ newpass, newpass ]
+			
+		return self.run( "password" )
+
+	#
+	# run_filelog: convert "p4 filelog" responses into objects with useful
+	#			   methods
+	#
+	# Requires tagged output to be of any real use. If tagged output it not 
+	# enabled then you just get the raw data back
+	#
+	def run_filelog( self, *args, **kargs ):
+		raw = self.run( 'filelog', args, **kargs )
+		if (not self.tagged or not raw): 
+			# untagged mode returns simple strings, which breaks the code below
+			# raw could be None if a handler is used
+			return raw
+		result = []
+		for h in raw:
+			df = None
+			if isinstance( h, dict ):
+				df = processFilelog( h )
+			else:
+				df = h
+			result.append( df )
+		return result
+
+	def run_print(self, *args, **kargs):
+		raw = self.run('print', args, **kargs)
+		result = []
+		if raw:
+			for line in raw:
+				if isinstance(line, dict):
+					result.append(line)
+					result.append("")
+				else:
+					result[-1] += line
+			return result
+		else:
+			return []
+	
+	def run_resolve(self, *args, **kargs):
+		if self.resolver:
+			myResolver = self.resolver
+		else:
+			myResolver = Resolver()
+		if "resolver" in kargs: 
+			myResolver = kargs["resolver"]
+			
+		savedResolver = self.resolver
+		self.resolver = myResolver
+		result = self.run("resolve", args)
+		self.resolver = savedResolver
+		
+		return result
+	
+	def run_tickets(self, *args):
+		fname = self.ticket_file
+		with open(fname) as f:
+			tickets_raw = f.readlines()
+		pattern = re.compile('([^=]*)=(.*):([^:]*)\n')
+		tickets = [ pattern.match(x).groups() for x in tickets_raw ]
+		keys = [ "Host", "User", "Ticket" ]
+		result = [ dict(zip(keys, x)) for x in tickets ]
+		return result
+	
+	def __flatten(self, args):
+		result = []
+		if isinstance(args, tuple) or isinstance(args, list):
+			for i in args:
+				result.extend(self.__flatten(i))
+		else:
+			result.append(args)
+		return tuple(result)
+
+	def __enter__( self ):
+		return self
+		
+	def __exit__( self, exc_type, exc_val, exc_tb ):
+		if self.connected():
+			self.disconnect()
+		return False
+		
+	def connect( self ):
+		P4API.P4Adapter.connect( self )
+		return self
+	
+	@contextmanager
+	def while_tagged( self, t ):
+		old = self.tagged
+		self.tagged = t
+		yield
+		self.tagged = old
+	
+	@contextmanager
+	def at_exception_level( self, e ):
+		old = self.exception_level
+		self.exception_level = e
+		yield
+		self.exception_level = old
+		
+	@contextmanager
+	def using_handler( self, c ):
+		old = self.handler
+		self.handler = c
+		yield
+		self.handler = old
+	
+	@contextmanager
+	def saved_context( self , **kargs):
+		"""Saves the context of this p4 object and restores it again at the end of the block"""
+	
+		saved_context = {}
+		for attr in self.__members__:
+			saved_context[attr] = getattr(self, attr)
+	
+		for (k,v) in list(kargs.items()):
+			setattr( self, k, v)
+		
+		yield
+		
+		# now restore the context again. Ignore AttributeError exception
+		# Exception is expected because some attributes only have getters, no setters
+		
+		for (k,v) in list(saved_context.items()):
+			try:
+				setattr( self, k, v )
+			except AttributeError:
+				pass # expected for server_level and p4config_file
+
+class Map(P4API.P4Map):
+	def __init__(self, *args):
+		P4API.P4Map.__init__(self, *args)
+		if len(args) > 0:
+			self.insert( *args )
+	
+	LEFT2RIGHT = True
+	RIGHT2LEFT = False
+	
+	def __str__( self ):
+		result = ""
+		
+		for a in self.as_array():
+			result += a + "\n"
+	
+		return result
+		
+	def is_empty(self):
+		"""Returns True if this map has no entries yet, otherwise False"""
+		
+		return self.count() == 0
+	
+	def includes(self, *args):
+		return self.translate(*args) != None
+	
+	def reverse(self):
+		return Map(P4API.P4Map.reverse(self).as_array())
+	
+	def insert(self, *args):
+		"""Insert an argument to the map. The argument can be:
+		
+		A String, 
+			Either of the form "[+-]//lhs/... //rhs/..." or "[+-]//lhs/..." 
+			for label style maps.
+		A List:
+			This is a list of strings of one of the single string formats 
+			described above.
+		A pair of Strings:
+			P4.Map.insert(lhs, rhs)
+		"""
+		
+		if len(args) == 1 :
+			arg = args[0]
+			if isinstance( arg, str ):
+				P4API.P4Map.insert( self, arg )
+			elif isinstance( arg, list ):
+				for s in arg:
+					P4API.P4Map.insert( self, s )
+					
+		else: # expecting 2 args in this case: left, right 
+			left = args[0].strip()
+			right = args[1].strip()
+			P4API.P4Map.insert(self, left, right )
+
+if __name__ == "__main__":
+	p4 = P4()
+	p4.connect()
+	try:
+		ret = p4.run(sys.argv[1:])
+		for line in ret:
+			if isinstance(line, dict):
+				print("-----")
+				for k in list(line.keys()):
+					print(k, "=", line[k])
+			else:
+				print(line)
+	except:
+		for e in p4.errors:
+			print(e)
diff --git a/Engine/Extras/Maya_AnimationRiggingTools/ARTv1/General/Scripts/Tools/ART_CopyWeightsTool.py b/Engine/Extras/Maya_AnimationRiggingTools/ARTv1/General/Scripts/Tools/ART_CopyWeightsTool.py
new file mode 100644
--- /dev/null
+++ b/Engine/Extras/Maya_AnimationRiggingTools/ARTv1/General/Scripts/Tools/ART_CopyWeightsTool.py
@@ -0,0 +1,189 @@
+from PySide import QtGui
+from PySide import QtCore
+import maya.OpenMayaUI as mui
+import shiboken
+import maya.cmds as cmds
+from functools import partial
+
+#Written by: Charles R. Anderson
+#Dec. 4, 2015 
+#Questions or comments email me: [REDACTED_EMAIL]
+
+def findCluster(meshName):
+    skinClusters = []
+    skinHistory = []
+    skinHistory = cmds.listHistory(meshName, lv=0)
+    for node in skinHistory:
+        if cmds.nodeType(node) == "skinCluster":
+            skinClusters.append(node)
+
+    return skinClusters[0]
+    
+#-----------------------------------------------------------------------------------------------
+def copyWeights(fromList, toList):
+    #filterString = textField.text()
+    #print filterString
+    if fromList.count() > 0:
+        for fromIndex in xrange(fromList.count()):
+            fromMesh = fromList.item(fromIndex).text()
+            print "From Mesh: ", fromMesh
+            if toList.count() > 0:
+                for toIndex in xrange(toList.count()):
+                    toMesh = toList.item(toIndex).text()
+                    if fromMesh.lower() in toMesh.lower():
+                        print "Matched To Mesh found: ", toMesh
+                        print "Copying weights from ", fromMesh, " to ", toMesh
+                        toCluster = findCluster(toMesh)
+                        fromCluster = findCluster(fromMesh)
+                        cmds.copySkinWeights(ss=fromCluster, ds=toCluster, noMirror=True, surfaceAssociation="closestPoint", influenceAssociation=["closestJoint", "oneToOne", "name"])
+            else:
+                cmds.warning("Your Copy TO list is empty.")
+    else:
+        cmds.warning("Your Copy FROM list is empty.")
+
+######################## UI FUNCTIONS ###############################
+#-----------------------------------------------------------------------------------------------
+# This function finds the main Maya Window
+def getMayaWindow():
+    pointer = mui.MQtUtil.mainWindow()
+    return shiboken.wrapInstance(long(pointer), QtGui.QWidget)
+
+#-----------------------------------------------------------------------------------------------
+def addMeshItem (listName):
+    listName.clear() # TEMPORARILY clears the list so that we never get dupes.  Figure out how to check for copies before adding and then remove this.
+    addItems = cmds.ls(sl=True)
+    #for one in addItems:
+        #print listName.findItems(one, QtCore.Qt.MatchExactly)
+    #for one in addItems:
+    listName.addItems(addItems)
+	
+#-----------------------------------------------------------------------------------------------
+def removeMeshItem (listName):
+    selected = listName.selectedItems()
+    for one in selected:
+        index = listName.row(one)
+        listName.takeItem(index)
+
+#-----------------------------------------------------------------------------------------------
+def clearMeshItems (listName):
+    listName.clear()
+    	
+#-----------------------------------------------------------------------------------------------
+def resetUI (fromList, toList):
+    fromList.clear()
+    toList.clear()
+
+#-----------------------------------------------------------------------------------------------
+def createUI(*args):
+    # check to see if this window exists and if so delete it.
+    objectName = "copyWeightsWindowNew"
+    if cmds.window("copyWeightsWindowNew", exists=True):
+        cmds.deleteUI("copyWeightsWindowNew", wnd=True)
+
+    # Create the Copy Weights window
+    parent = getMayaWindow()
+    currentWindow = QtGui.QMainWindow(parent)
+    currentWindow.setObjectName(objectName)
+    currentWindow.setWindowTitle("Copy Weights Tool")
+    currentWindow.resize(500, 700)
+
+    # create a widget
+    widget = QtGui.QWidget()
+    currentWindow.setCentralWidget(widget)
+
+    # create the Main vertical layout
+    mainLayout = QtGui.QVBoxLayout(widget)
+
+    # create a horizontal layout for the lists
+    listsLayout = QtGui.QHBoxLayout()
+    mainLayout.addLayout(listsLayout)
+
+    # FROM section---------------------
+    fromLayout = QtGui.QVBoxLayout()
+    listsLayout.addLayout(fromLayout)
+
+    fromListLabel = QtGui.QLabel("Copy from:")
+    fromLayout.addWidget(fromListLabel)
+    fromList = QtGui.QListWidget()
+    fromLayout.addWidget(fromList)
+    fromList.setToolTip("This is a list of meshes to copy weights from.")
+    fromList.setSelectionMode(QtGui.QAbstractItemView.ExtendedSelection)
+
+    # Add Button
+    addFromButton = QtGui.QPushButton("Add Selected")
+    fromLayout.addWidget(addFromButton)
+    addFromButton.setMinimumHeight(30)
+    addFromButton.setMaximumHeight(30)
+    addFromButton.setMinimumWidth(200)
+    addFromButton.setToolTip("Adds the selected meshes to the list above.")
+    addFromButton.clicked.connect(partial(addMeshItem, fromList))
+
+    removeFromButton = QtGui.QPushButton("Remove Selected")
+    fromLayout.addWidget(removeFromButton)
+    removeFromButton.setMinimumHeight(30)
+    removeFromButton.setMaximumHeight(30)
+    removeFromButton.setMinimumWidth(200)
+    removeFromButton.clicked.connect(partial(removeMeshItem, fromList))
+
+    clearFromButton = QtGui.QPushButton("Clear List")
+    fromLayout.addWidget(clearFromButton)
+    clearFromButton.setMinimumHeight(30)
+    clearFromButton.setMaximumHeight(30)
+    clearFromButton.setMinimumWidth(200)
+    clearFromButton.clicked.connect(partial(clearMeshItems, fromList))
+
+    # TO section-------------------------
+    toLayout = QtGui.QVBoxLayout()
+    listsLayout.addLayout(toLayout)
+
+    toListLabel = QtGui.QLabel("Copy to:")
+    toLayout.addWidget(toListLabel)
+    toList = QtGui.QListWidget()
+    toLayout.addWidget(toList)
+    toList.setToolTip("This is a list of meshes to copy weights to.")
+    toList.setSelectionMode(QtGui.QAbstractItemView.ExtendedSelection)
+
+    # Add Button
+    addToButton = QtGui.QPushButton("Add Selected")
+    toLayout.addWidget(addToButton)
+    addToButton.setMinimumHeight(30)
+    addToButton.setMaximumHeight(30)
+    addToButton.setMinimumWidth(200)
+    addToButton.clicked.connect(partial(addMeshItem, toList))
+
+    removeToButton = QtGui.QPushButton("Remove Selected")
+    toLayout.addWidget(removeToButton)
+    removeToButton.setMinimumHeight(30)
+    removeToButton.setMaximumHeight(30)
+    removeToButton.setMinimumWidth(200)
+    removeToButton.clicked.connect(partial(removeMeshItem, toList))
+
+    clearToButton = QtGui.QPushButton("Clear List")
+    toLayout.addWidget(clearToButton)
+    clearToButton.setMinimumHeight(30)
+    clearToButton.setMaximumHeight(30)
+    clearToButton.setMinimumWidth(200)
+    clearToButton.clicked.connect(partial(clearMeshItems, toList))
+
+
+    # Main widgets ---------------------
+    resetButton = QtGui.QPushButton("Reset")
+    mainLayout.addWidget(resetButton)
+    resetButton.setMinimumHeight(30)
+    resetButton.setMaximumHeight(30)
+    resetButton.setMinimumWidth(200)
+    resetButton.clicked.connect(partial(resetUI, fromList, toList))
+
+    copyButton = QtGui.QPushButton("Copy")
+    mainLayout.addWidget(copyButton)
+    copyButton.setMinimumHeight(30)
+    copyButton.setMaximumHeight(30)
+    copyButton.setMinimumWidth(200)
+    copyButton.setToolTip("Run through all of the meshes in the FROM list and copy its weights to any meshes that it matches in the TO list.  This matches based on name and is case insensitive.  For example it will find \"Art\" inside of \"fart\".")
+    copyButton.clicked.connect(partial(copyWeights, fromList, toList))
+
+    # show the window
+    currentWindow.show()
+
+#createUI()
+
diff --git a/Engine/Extras/Maya_AnimationRiggingTools/ARTv1/General/Scripts/Tools/ART_CreateHydraulicRigs.py b/Engine/Extras/Maya_AnimationRiggingTools/ARTv1/General/Scripts/Tools/ART_CreateHydraulicRigs.py
new file mode 100644
--- /dev/null
+++ b/Engine/Extras/Maya_AnimationRiggingTools/ARTv1/General/Scripts/Tools/ART_CreateHydraulicRigs.py
@@ -0,0 +1,60 @@
+import maya.cmds as cmds
+import maya.mel as mel
+import os
+
+# This tool only works on "leaf" joints created with A.R.T 
+# To use, select 3 objects, the start of your hydraulic system, the end, and the upvector.  It is assumed that you are selecting the _anim controls for the first 2
+# It is assumed that the upvector is 1 by default unless passed in otherwise.  Once set, you can always change this by simply manually reversing the value in the aimConstraint
+# It is also assumed that you are using X as the axis that points at each other and that the first control points at the 2nd, but the 2nd points away from the first.
+# Assuming Z as the up vector axis.
+
+def CreateHydraulicRig(upDir, *args):
+
+    selection = cmds.ls(sl=True)
+    if len(selection) != 3:
+        cmds.confirmDialog(icon = "warning!!", title = "Hydraulic Rig Tool", message = "You must select exactly 3 objects.  Top of hydraulic, bottom of hydraulic, and the up vector.")
+        return
+    else:
+        hyd_start_01_anim = selection[0]
+        hyd_start_02_anim = selection[1]
+        upObject = selection[2]
+        
+        hyd_start_01 = hyd_start_01_anim.rpartition("_")[0]
+        hyd_start_02 = hyd_start_02_anim.rpartition("_")[0]
+
+        # Create a rig for the lower arm extra hydraulic piston.
+        cmds.delete("driver_"+hyd_start_01+"_parentConstraint1")
+        robo_lowarm_pistonrod_01_anim_sub = cmds.group(em=True, name=hyd_start_01_anim+"_sub", parent=hyd_start_01_anim)
+        const = cmds.parentConstraint(hyd_start_01_anim, robo_lowarm_pistonrod_01_anim_sub, weight=1, mo=False)
+        cmds.delete(const)
+        const = cmds.parentConstraint(robo_lowarm_pistonrod_01_anim_sub, "driver_"+hyd_start_01, weight=1, mo=True)
+        
+        cmds.delete("driver_"+hyd_start_02+"_parentConstraint1")
+        robo_lowarm_pistonrod_02_anim_sub = cmds.group(em=True, name=hyd_start_02_anim+"_sub", parent=hyd_start_02_anim)
+        const = cmds.parentConstraint(hyd_start_02_anim, robo_lowarm_pistonrod_02_anim_sub, weight=1, mo=False)
+        cmds.delete(const)
+        const = cmds.parentConstraint(robo_lowarm_pistonrod_02_anim_sub, "driver_"+hyd_start_02, weight=1, mo=True)
+
+        # Hook up the hydraulics for the lowerarm piston.
+        const1 = cmds.aimConstraint(robo_lowarm_pistonrod_01_anim_sub, robo_lowarm_pistonrod_02_anim_sub, weight=1, mo=True, aimVector=(-1, 0, 0), upVector=(0, 0, upDir), worldUpType="object", worldUpVector=(0, 0, -1), worldUpObject=upObject)
+        const2 = cmds.aimConstraint(robo_lowarm_pistonrod_02_anim_sub, robo_lowarm_pistonrod_01_anim_sub, weight=1, mo=True, aimVector=(1, 0, 0), upVector=(0, 0, upDir), worldUpType="object", worldUpVector=(0, 0, -1), worldUpObject=upObject)
+    
+    cmds.select(const1, const2)
+
+def CreateHydraulicSlideRig(*args):
+    print "test"
+    # Create a slide mechanism on the forearm part of the lowerarm piston.
+    # CRA NOTE: this will come at a later date.
+    '''
+    cmds.addAttr("robo_lowarm_pistonrod_02_anim", ln = "Slide", dv = 0, min=0, max=10, keyable = True)
+    cmds.setDrivenKeyframe("robo_lowarm_pistonrod_02_anim_grp.translateX", cd="robo_lowarm_pistonrod_02_anim.Slide", itt="spline", ott="spline")
+    cmds.setDrivenKeyframe("robo_lowarm_pistonrod_02_anim_grp.translateZ", cd="robo_lowarm_pistonrod_02_anim.Slide", itt="spline", ott="spline")
+
+    cmds.setAttr("robo_lowarm_pistonrod_02_anim.Slide", 10)
+    cmds.setAttr("robo_lowarm_pistonrod_02_anim_grp.tx", -6.662443)
+    cmds.setAttr("robo_lowarm_pistonrod_02_anim_grp.tz", -5.510393)
+
+    cmds.setDrivenKeyframe("robo_lowarm_pistonrod_02_anim_grp.translateX", cd="robo_lowarm_pistonrod_02_anim.Slide", itt="spline", ott="spline")
+    cmds.setDrivenKeyframe("robo_lowarm_pistonrod_02_anim_grp.translateZ", cd="robo_lowarm_pistonrod_02_anim.Slide", itt="spline", ott="spline")
+    cmds.setAttr("robo_lowarm_pistonrod_02_anim.Slide", 0)
+    '''
diff --git a/Engine/Extras/Maya_AnimationRiggingTools/ARTv1/General/Scripts/Tools/ART_ImportMultipleWeights.py b/Engine/Extras/Maya_AnimationRiggingTools/ARTv1/General/Scripts/Tools/ART_ImportMultipleWeights.py
new file mode 100644
--- /dev/null
+++ b/Engine/Extras/Maya_AnimationRiggingTools/ARTv1/General/Scripts/Tools/ART_ImportMultipleWeights.py
@@ -0,0 +1,272 @@
+import maya.cmds as cmds
+import os, cPickle
+from PySide import QtGui
+import maya.OpenMayaUI as mui
+import shiboken
+from PySide import QtCore
+from functools import partial
+
+#############################################################################################
+#############################################################################################
+#############################################################################################
+# Select all of the meshes in your scene that you want to import onto.  
+# If there is a .txt file in the SkinWeights folder on disk that matches its name it will import the weights onto the mesh.
+def importMultipleSkinWeights(suffix, rui, *args):
+    
+    #tools path
+    toolsPath = cmds.internalVar(usd = True) + "mayaTools.txt"
+    if os.path.exists(toolsPath):
+
+        f = open(toolsPath, 'r')
+        mayaToolsDir = f.readline()
+        f.close()	
+
+
+    selection = cmds.ls(sl = True)
+    if selection == []:
+        cmds.confirmDialog(icon = "warning", title = "Import Skin Weights", message = "Nothing Selected")
+        return
+        
+    warnings = False
+
+    if selection != []:
+        #find all files in the system dir
+        path = mayaToolsDir + "\General\ART\SkinWeights\\"
+        print "path is "+path
+        
+        for one in selection:
+            '''
+            if len(selection) > 1:
+                cmds.confirmDialog(icon = "warning", title = "Import Skin Weights", message = "Too many objects selected. Please only selected 1 mesh.")
+                return
+            '''
+            #try:
+                #fileToImport = cmds.fileDialog2(fileFilter="*.txt", dialogStyle=2, fm = 1, startingDirectory = path)[0]
+
+            #except:
+                #cmds.warning("Operation Cancelled.")
+                #return
+            filename = one
+            print "SUFFIX: ", suffix
+            if suffix:
+                filename = one.partition(suffix)[0]
+            
+            print "FILENAME: ", filename
+
+            fileToImport = path+filename+".txt"
+            if cmds.file(fileToImport, q=True, ex=True):
+                print "file is "+fileToImport
+
+                #load the file with cPickle and parse out the information
+                file = open(fileToImport, 'r')
+                weightInfoList = cPickle.load(file)
+
+                #get the number of vertices in the mesh. This will be used to update our progress window
+                verts = cmds.polyEvaluate(one, vertex = True)
+                if verts < 20:
+                    verts = verts * 20
+
+                increment = int(verts/20)
+                matchNum = increment
+
+                #find the transforms(joints) from the weightInfoList. format = [ [vertNum, [jointName, value], [jointName, value] ] ]
+                transforms = []
+                newList = weightInfoList[:]
+                for info in newList:
+
+                    for i in info:
+                        if i != info[0]:
+                            transform = i[0]
+                            transforms.append(transform)
+
+                #remove duplicates from the transforms(joints) list
+                transforms = set(transforms)
+                print "Files Transforms List:"
+                print transforms
+
+                #clear selection. Add transforms(joints) to selection and geo
+                cmds.select(clear = True)
+                for t in transforms:
+                    if cmds.objExists(t):
+                        cmds.select(t, add = True)
+
+
+                #check if the geometry that is selected is already skinned or not
+                skinClusters = cmds.ls(type='skinCluster')
+                skinnedGeometry = []
+                for cluster in skinClusters:
+                    geometry = cmds.skinCluster(cluster, q = True, g = True)[0]
+                    geoTransform = cmds.listRelatives(geometry, parent = True)[0]
+                    skinnedGeometry.append(geoTransform)
+
+                    if geoTransform == one:
+                        skin = cluster
+                        
+                        currentTransforms = cmds.skinCluster(skin, q=True, inf=True)
+                        print "Current Transforms List:"
+                        print currentTransforms
+                        
+                        for tr in transforms:
+                            if tr not in currentTransforms:
+                                cmds.skinCluster(skin, e=True, ai=tr, lw=False)
+                                print (cmds.skinCluster(skin, q=True, inf=True))
+
+
+                #if a skin cluster does not exist already on the mesh:
+                if one not in skinnedGeometry:
+
+                    #put a skin cluster on the geo normalizeWeights = 1, dr=4.5, mi = 4,
+                    cmds.select(one, add = True)
+                    skin = cmds.skinCluster( tsb = True, skinMethod = 0, name = (one +"_skinCluster"))[0]
+
+
+                #setup a progress window to track length of progress
+                cmds.progressWindow(title='Skin Weights', progress = 0, status = "Reading skin weight information from file...")
+
+                #get the vert number, and the weights for that vertex
+                for info in weightInfoList:
+                    vertNum = info[0]
+                    info.pop(0)
+
+                    #progress bar update
+                    if vertNum == matchNum:
+                        cmds.progressWindow( edit=True, progress = ((matchNum/increment) * 5), status=  "Importing skin weights for " + one + ".\n Please be patient")
+                        matchNum += increment
+
+                    #set the weight for each vertex
+                    try:
+                        cmds.skinPercent(skin, one + ".vtx[" + str(vertNum) + "]", transformValue= info)
+                    except:
+                        pass
+
+                # remove unused influences
+                if rui == 1:
+                    infToRemove = []
+                    weightedInfs = cmds.skinCluster(skin, q=True, weightedInfluence=True)
+                    finalTransforms = cmds.skinCluster(skin, q=True, inf=True)
+                    for inf in finalTransforms:
+                        if inf not in weightedInfs:
+                            infToRemove.append(inf)
+                    for rem in infToRemove:
+                        cmds.skinCluster(skin, e=True, ri=rem)
+
+                #close the file
+                file.close()
+            else:
+                cmds.warning(fileToImport+" does not exist.")
+                warnings = True
+
+            #close progress window
+            cmds.progressWindow(endProgress = 1)
+
+    cmds.select(selection)
+    if warnings == False:
+        cmds.headsUpMessage( "Importing weights is complete!!", time=3.0 )
+        #cmds.confirmDialog(icon = "Complete", title = "Import Skin Weights", message = "Importing weights is complete!!")
+    elif warnings == True:
+        cmds.headsUpMessage( "Importing weights is complete but with Warnings!!", time=3.0 )
+        #cmds.confirmDialog(icon = "Complete", title = "Import Skin Weights", message = "Importing weights is complete but with Warnings!!")
+
+    
+######################## UI FUNCTIONS ###############################
+#-----------------------------------------------------------------------------------------------
+# This function finds the main Maya Window
+def getMayaWindow():
+    pointer = mui.MQtUtil.mainWindow()
+    return shiboken.wrapInstance(long(pointer), QtGui.QWidget)
+
+#-----------------------------------------------------------------------------------------------
+def resetUI (fromList, toList):
+    fromList.clear()
+    toList.clear()
+
+#-----------------------------------------------------------------------------------------------
+def deleteUI ():
+    cmds.deleteUI("importMultipleWeightsWindow", wnd=True)
+    
+#-----------------------------------------------------------------------------------------------
+def importFunction (suffixField, RICheckBox, closeUI):
+    suffixField = suffixField.text()
+    RICheckBox = RICheckBox.isChecked()
+
+    importMultipleSkinWeights(suffixField, RICheckBox)
+    if closeUI == True:
+        cmds.deleteUI("importMultipleWeightsWindow", wnd=True)
+
+#-----------------------------------------------------------------------------------------------
+def importMultipleWeightsUI(*args):
+    # check to see if this window exists and if so delete it.
+    objectName = "importMultipleWeightsWindow"
+    if cmds.window("importMultipleWeightsWindow", exists=True):
+        cmds.deleteUI("importMultipleWeightsWindow", wnd=True)
+
+    # Create the Weights window
+    parent = getMayaWindow()
+    currentWindow = QtGui.QMainWindow(parent)
+    currentWindow.setObjectName(objectName)
+    currentWindow.setWindowTitle("Import Multiple Weights Tool")
+    currentWindow.resize(50, 50)
+
+    # create a widget
+    widget = QtGui.QWidget()
+    currentWindow.setCentralWidget(widget)
+
+    # create the Main vertical layout
+    mainLayout = QtGui.QVBoxLayout(widget)
+    subLayout = QtGui.QHBoxLayout()
+    mainLayout.addLayout(subLayout)
+
+    # create a horizontal layout for the lists
+    optionsLayout = QtGui.QVBoxLayout()
+    mainLayout.addLayout(optionsLayout)
+
+    buttonsLayout = QtGui.QHBoxLayout()
+    mainLayout.addLayout(buttonsLayout)
+    
+    # Suffix UI
+    suffixLayout = QtGui.QHBoxLayout()
+    optionsLayout.addLayout(suffixLayout)
+    textFieldLabel = QtGui.QLabel("Suffix to be removed:")
+    suffixLayout.addWidget(textFieldLabel)
+    suffixField = QtGui.QLineEdit()
+    suffixLayout.addWidget(suffixField)
+    suffixField.setPlaceholderText("Enter Suffix...")
+
+    # Remove Unused Influences UI
+    RUILayout = QtGui.QHBoxLayout()
+    optionsLayout.addLayout(RUILayout)
+    RICheckBox = QtGui.QCheckBox("Remove Unused Influences")
+    RUILayout.addWidget(RICheckBox)
+    RICheckBox.setChecked(True)
+    
+    # Import Button
+    importButton = QtGui.QPushButton("Import")
+    buttonsLayout.addWidget(importButton)
+    importButton.setMinimumHeight(30)
+    importButton.setMaximumHeight(30)
+    importButton.setMinimumWidth(200)
+    importButton.setToolTip("Adds the selected meshes to the list above.")
+    importButton.clicked.connect(partial(importFunction, suffixField, RICheckBox, True))
+
+    # Apply Button
+    applyButton = QtGui.QPushButton("Apply")
+    buttonsLayout.addWidget(applyButton)
+    applyButton.setMinimumHeight(30)
+    applyButton.setMaximumHeight(30)
+    applyButton.setMinimumWidth(200)
+    applyButton.clicked.connect(partial(importFunction, suffixField, RICheckBox, False))
+
+    # Close Button
+    closeFromButton = QtGui.QPushButton("Close")
+    buttonsLayout.addWidget(closeFromButton)
+    closeFromButton.setMinimumHeight(30)
+    closeFromButton.setMaximumHeight(30)
+    closeFromButton.setMinimumWidth(200)
+    closeFromButton.clicked.connect(deleteUI)
+    
+    # show the window
+    currentWindow.show()
+
+#importMultipleWeightsUI()
+    
+#importMultipleSkinWeights("_LOD4")
diff --git a/Engine/Extras/Maya_AnimationRiggingTools/ARTv1/General/Scripts/Tools/ART_MeshSpliter.py b/Engine/Extras/Maya_AnimationRiggingTools/ARTv1/General/Scripts/Tools/ART_MeshSpliter.py
new file mode 100644
--- /dev/null
+++ b/Engine/Extras/Maya_AnimationRiggingTools/ARTv1/General/Scripts/Tools/ART_MeshSpliter.py
@@ -0,0 +1,93 @@
+import maya.cmds as cmds
+import maya.mel as mel
+
+# Splits the selected meshes by their skinCLuster influences and creates a new heirarchy and constrains them to the skeleton.
+# NOTE this script doesnt work well if the names of the meshes have ":" in them.
+def splitMeshByInfluences(threshold, *args):
+    import maya.mel as mel
+    
+    finalMeshList = []
+    
+    # find the meshes that the user has selected
+    selected = cmds.ls(sl=True, type="transform")
+
+    # Create an empty set.  We wil add to it later.
+    cmds.select(cl=True)
+    setName = cmds.sets(n="Rigid_Meshes")
+    
+    rigidMeshMasterGrp = cmds.group(em=True, name="Rigid_Mesh")
+
+    # for each mesh:
+    for mesh in selected:
+        rigidMeshMeshGrp = cmds.group(em=True, name="Rigid_"+mesh, p=rigidMeshMasterGrp)
+        #print "Mesh:"+mesh
+
+        skinCluster = mel.eval('findRelatedSkinCluster %s' % mesh)
+        influences = cmds.skinCluster(skinCluster, q=True, inf=True)
+        #print influences
+
+        for inf in influences:
+            # Duplicate the mesh with input connections
+            dupeMesh = cmds.duplicate(mesh, name=mesh+"__"+inf, rr=False)[0]
+            #print "DupeMesh:       "+dupeMesh
+            for attr in ["tx", "ty", "tz", "rx", "ry", "rz", "sx", "sy", "sz", "v"]:
+                cmds.setAttr(dupeMesh+"."+attr, lock=False, k=True)
+            cmds.makeIdentity(dupeMesh, t=1, r=1, s=1, apply=True)
+            cmds.select(dupeMesh, r=True)
+            mel.eval("CenterPivot")
+            cmds.parent(dupeMesh, rigidMeshMeshGrp)
+            
+            finalMeshList.append(dupeMesh)
+            cmds.sets(dupeMesh, add=setName)
+
+            # Select the verts that are influenced by the inf joint.
+            if len(influences) <= 1:
+                const = cmds.parentConstraint(inf, rigidMeshMeshGrp, dupeMesh, mo=True)[0]
+                cmds.setAttr(const+"."+rigidMeshMeshGrp+"W1", 0)
+            else:
+                verts = []
+                for x in range(cmds.polyEvaluate(mesh, v=True)):
+                    v = cmds.skinPercent(skinCluster, '%s.vtx[%d]' % (mesh, x), transform=inf, q=True)
+                    if v > threshold:
+                        verts.append('%s.vtx[%d]' % (dupeMesh, x))
+                #print verts
+                if verts:
+                    cmds.select(verts)            
+                    # Convert the selection to contained faces
+                    if len(verts) != cmds.polyEvaluate(mesh, v=True):
+                        mel.eval("InvertSelection;")
+                        cmds.select(cmds.polyListComponentConversion(fv=True, tf=True, internal=True ))
+                        cmds.delete()
+                    cmds.select(dupeMesh, r=True)
+                    mel.eval("CenterPivot")
+                    const = cmds.parentConstraint(inf, rigidMeshMeshGrp, dupeMesh, mo=True)[0]
+                    cmds.setAttr(const+"."+rigidMeshMeshGrp+"W1", 0)
+                else:
+                    print "No Valid Verts so we didnt use this influence."
+                    cmds.delete(dupeMesh)
+
+    cmds.select(rigidMeshMasterGrp, hi=True)
+    cmds.createDisplayLayer(cmds.ls(sl=True), name="RigidMesh_L")
+    
+    return finalMeshList
+
+#splitMeshByInfluences(0.5)
+
+def findThreshold(*args):
+    threshold = cmds.floatSliderGrp("skinWeightFloatSliderGrp", q=True, value=True)
+    splitMeshByInfluences(threshold)
+
+
+# Create a UI
+def splitMeshByInfluencesUI(*args):
+    if cmds.window("splitMeshByInfluencesUI", exists = True):
+        cmds.deleteUI("splitMeshByInfluencesUI")
+
+    window = cmds.window("splitMeshByInfluencesUI", title='Split Mesh by Influences', s=False)
+    cmds.columnLayout()
+    cmds.floatSliderGrp("skinWeightFloatSliderGrp", label='Skin Weight Threshold', field=True, minValue=0, maxValue=1, value=0.5, step=0.01, cw3=[110, 30, 150])
+    cmds.button("Apply", label="Apply", w=290, c=findThreshold)
+
+    cmds.showWindow(window)
+
+#splitMeshByInfluencesUI()
\ No newline at end of file
diff --git a/Engine/Extras/Maya_AnimationRiggingTools/ARTv1/General/Scripts/Tools/ART_OrientJointWithUp.py b/Engine/Extras/Maya_AnimationRiggingTools/ARTv1/General/Scripts/Tools/ART_OrientJointWithUp.py
new file mode 100644
--- /dev/null
+++ b/Engine/Extras/Maya_AnimationRiggingTools/ARTv1/General/Scripts/Tools/ART_OrientJointWithUp.py
@@ -0,0 +1,295 @@
+import maya.cmds as cmds
+from functools import partial
+import pymel.core as pm
+
+#Select the joint you want to orient, then the child you want it to be aimed at, and then the object to use as the upVector
+#CRA NOTE: OptionVars are not retaining their settings when you re-open the UI.  They always reset to the default.  Figure out how to make them remember.
+#--------------------------------------------------------------------------------------------------------------------
+#aimAxis        = axis to point at the child bone.
+#upAxis         = axis to point at the locator bone.
+#heirarchy      = Should this be applied using the heirarchy or only to the selected joints?
+#customChildInt = Do we use a custom child?   When set to 1, this will tell it to use the named joint in customChild rather than popping open a window when more than 1 child joint exists.
+#customChild    = The joint to use as the aim joint when customChildInt is set to 1.  If no custom child is defined it will either use the only child available or if there are multiples it will bring up a ui to let the user choose which child to use.
+# IMPORTANT NOTE: When calling this script directly and not using the UI, it is reccommended to have heirarchy set to 0.  If you try to orient an entire heirarchy we can not guarantee that it will pick the correct child unless you are 100% positive that there are no parents with multiple children.
+#----------------------------------------------------------------------------------------
+# This def finds the joint we want to aim (jointToAim), the joint to be aimed at (aimTarget), and the up node (upNode) based on the selection.  IF heirarchy is checked, then it uses the findChild def to find the joint to be aimed at.
+def setOrientWithUpVector(aimAxis, upAxis, heirarchy, customChildInt, customChild):
+    # THESE ARE THE 3 NODES WE NEED TO DEFINE!!
+    jointToAim = None
+    aimTarget = None
+    upNode = None
+
+    #Find the selected joint and upVector nodes from the selection
+    # CRA NOTE: Could add some more logic here later to check to see if the first 2 selected are joints and that the 2nd is a relative of the first.
+    selected = cmds.ls(sl=True)
+    if len(selected) != 2:
+        cmds.confirmDialog(icon = "warning!!", title = "Orient Joint with Up Tool", message = "You must select exactly 2 objects.  Joint you want to orient and then the object to use as the upVector.")
+        return
+    else:
+        jointToAim = selected[0]
+        upNode = selected[1]
+
+        # Find the joints that we need to apply the orient script to.
+        if heirarchy == 0:
+            # find only the immediate children of jointToAim
+            allChildren = cmds.listRelatives(jointToAim, c=True, type="joint")
+            #print "\nAll Children..."
+            #print allChildren
+
+            if customChildInt == 1:
+                aimTarget = customChild
+                if allChildren != None:
+                    doOrientJoint(jointToAim, aimTarget, upNode, aimAxis, upAxis, allChildren)
+                else:
+                    cmds.joint(jointToAim, e=True, oj="none", secondaryAxisOrient="yup", ch=True, zso=True)
+            elif allChildren != None:
+                if len(allChildren) == 1:
+                    aimTarget = allChildren[0]
+                else:
+                    aimTarget = cmds.layoutDialog(t="Child Joints List", ui=partial(findChild, jointToAim))
+                #print "\nPRINTING JOINTS LIST---- "
+                #print "jointToAim: "+jointToAim
+                #print "aimTarget: "+aimTarget
+                #print "upNode   : "+upNode   
+                #print "----DONE PRINTING JOINTS LIST"
+                doOrientJoint(jointToAim, aimTarget, upNode, aimAxis, upAxis, allChildren)
+            else:
+                cmds.joint(jointToAim, e=True, oj="none", secondaryAxisOrient="yup", ch=True, zso=True)
+
+
+        if heirarchy == 1:
+            # find all of the joints we want to use as the jointToAim
+            cmds.select(jointToAim, hi=True)
+            allJointsToAim = cmds.ls(sl=True, type="joint")
+            #print "\nALL Joints To Aim..."
+            #print allJointsToAim
+
+            for one in allJointsToAim:
+                jointToAim = one
+                #print "\nCURRENT JOINT..."
+                #print jointToAim
+                
+                # find only the immediate children of jointToAim
+                allChildren = cmds.listRelatives(jointToAim, c=True, type="joint")
+                #print "\nAll Children..."
+                #print allChildren
+
+                if allChildren != None:
+                    if len(allChildren) == 1:
+                        aimTarget = allChildren[0]
+                    else:
+                        aimTarget = cmds.layoutDialog(t="Child Joints List for: "+jointToAim, ui=partial(findChild, jointToAim))
+                   
+                    #print "\nPRINTING JOINTS LIST---- "
+                    #print "jointToAim: "+jointToAim
+                    #print "aimTarget: "+aimTarget
+                    #print "upNode: "+upNode   
+                    #print "----DONE PRINTING JOINTS LIST"
+                    doOrientJoint(jointToAim, aimTarget, upNode, aimAxis, upAxis, allChildren)
+                else:
+                    # Set the joint orient to none so that it just takes the orientation of its parent.  This is only for the leaf joints in the heirarchy.
+                    cmds.joint(jointToAim, e=True, oj="none", secondaryAxisOrient="yup", ch=True, zso=True)
+
+#----------------------------------------------------------------------------------------
+# This def is only used by the setOrientWithUpVector def IF the heirarchy check box is set.  It creates a UI that outputs the child that the user selects in a UI to be used as the aim child for joints with more than 1 child.
+def findChild(jointToAim):
+    #print "\nFIND CHILD START----->" 
+
+    allChildren = cmds.listRelatives(jointToAim, c=True, type="joint")
+
+    form = cmds.setParent(q=True)
+    cmds.formLayout(form, e=True, width=200, height=300)
+    
+    tsl = cmds.textScrollList("childJointList", nr=10, ams=1, dcc="pm.runtime.ShowAttributeEditorOrChannelBox()")
+    cmds.textScrollList(tsl, e=True, sc=partial(selectItem, tsl))
+    for s in allChildren:
+        cmds.textScrollList(tsl, e=True, a=s, sii=1) 
+    cmds.select(cmds.textScrollList(tsl, q=True, si=True))
+    txt = cmds.text("childText", l="Select the child joint that you want the aim axis to point at and press continue.")
+    butt = cmds.button("continueButton", l="Continue", c=partial(getTextScrollListName, tsl))
+    cmds.formLayout(form, e=True, 
+        af=[(tsl, 'top', 1), (tsl, 'left', 1), (tsl, 'right', 1), (txt, 'left', 1), (txt, 'right', 1), (butt, 'left', 1), (butt, 'right', 1), (butt, 'bottom', 1)], 
+        ac=[(tsl, 'bottom', 1, txt), (txt, 'bottom', 1, butt)])
+
+def selectItem(textScrollListName):
+    cmds.select(cmds.textScrollList(textScrollListName, q=True, si=True))
+    
+def getTextScrollListName(textScrollListName, *args):
+    cmds.layoutDialog(dismiss=cmds.textScrollList(textScrollListName, q=True, si=True)[0])
+
+#----------------------------------------------------------------------------------------
+def doOrientJoint (jointToAim, aimTarget, upNode, aimAxis, upAxis, allChildren):
+    
+    print jointToAim
+    print aimTarget
+    print upNode
+    print "upAxis: " + upAxis 
+    print "aimAxis: " + aimAxis
+    print allChildren
+
+    Coords_up = []
+    if upAxis == "x":
+        Coords_up = [1, 0, 0]
+    if upAxis == "-x":
+        Coords_up = [-1, 0, 0]
+    if upAxis == "y":
+        Coords_up = [0, 1, 0]
+    if upAxis == "-y":
+        Coords_up = [0, -1, 0]
+    if upAxis == "z":
+        Coords_up = [0, 0, 1]
+    if upAxis == "-z":
+        Coords_up = [0, 0, -1]
+
+    Coords_aim = [] 
+    if aimAxis == "x":
+        Coords_aim = [1, 0, 0]
+    if aimAxis == "-x":
+        Coords_aim = [-1, 0, 0]
+    if aimAxis == "y":
+        Coords_aim = [0, 1, 0]
+    if aimAxis == "-y":
+        Coords_aim = [0, -1, 0]
+    if aimAxis == "z":
+        Coords_aim = [0, 0, 1]
+    if aimAxis == "-z":
+        Coords_aim = [0, 0, -1]
+
+    #Do the aim constraint to make the joint orient properly
+    cmds.parent(allChildren, w=True)
+    aimConst = cmds.aimConstraint(aimTarget, jointToAim, n=jointToAim+"_AimConst", o=[0, 0, 0], w=1, aim=[Coords_aim[0], Coords_aim[1], Coords_aim[2]], u=[Coords_up[0], Coords_up[1], Coords_up[2]], wut="object", wuo=upNode)[0]
+    cmds.delete(aimConst)
+    cmds.select(jointToAim)
+    cmds.makeIdentity(apply=True, t=1, r=1, s=1, n=0) 
+    cmds.delete(ch=True)
+    cmds.parent(allChildren, jointToAim)
+
+#----------------------------------------------------------------------------------------
+def orientJointOB(*args):
+    #Spine Option Variables
+    if cmds.optionVar(exists="orientJointHeirarchyOV") == 0:
+        cmds.optionVar(iv=("orientJointHeirarchyOV", 1)) 
+    if cmds.optionVar(exists="OV_frameLayout_collapse") == 0:
+        cmds.optionVar(iv=("OV_frameLayout_collapse", 0))
+    if cmds.optionVar(exists="OV_upVectorType") == 0:
+        cmds.optionVar(iv=("OV_upVectorType", 3))
+    if cmds.optionVar(exists="OV_aimVectorType") == 0:
+        cmds.optionVar(iv=("OV_aimVectorType", 1))
+        
+    # Create the window
+    if cmds.window("orientJointsWindow", exists=True):
+        cmds.deleteUI("orientJointsWindow", window=True)
+    cmds.window("orientJointsWindow", t="Orient Joints with Up Vector", menuBar=True, s=0, mxb=0, rtf=1, w=300, h=175)
+    
+    # Create the menu
+    cmds.menu("editMenu", l="Edit", p="orientJointsWindow", to=0, aob=True)
+    cmds.menuItem("resetSettingsItem", l="Reset Settings", ann="Reset all of the settings in this window to their default.", echoCommand=True, p="editMenu", c=resetOrientJointVariables)
+    cmds.menu("helpMenu", l="Help", p="orientJointsWindow", hm=1, to=0, aob=True)
+    cmds.menuItem("helpItem", l="Help on Orient Joints with Up Vector...", ann="Help on this tool.  Currently there is none...", echoCommand=True, p="helpMenu", c="")
+
+    # base layout for the entire UI
+    cmds.columnLayout(adjustableColumn=True)
+    cmds.frameLayout(label="Options", bgc=[0.1, 0.1, 0.1], cll=True, cl=cmds.optionVar(q="OV_frameLayout_collapse"), cc="cmds.optionVar(iv=(\"OV_frameLayout_collapse\", 1))", ec="cmds.optionVar(iv=(\"OV_frameLayout_collapse\", 0))", borderStyle="etchedOut")
+
+    # Aim vector dropdown menu
+    cmds.optionMenuGrp("OMG_aimVectorType", l="Aim Axis", cw=[1, 80], ann="Aim Axis is the axis of your joints that will point at the child. Default is \"x\"", cc="cmds.optionVar(iv=(\"OV_aimVectorType\", "+"cmds.optionMenuGrp(\"OMG_aimVectorType\", q=True, sl=True)"+"))")
+    cmds.menuItem("xItem", l="x")
+    cmds.menuItem("nxItem", l="-x")
+    cmds.menuItem("yItem", l="y")
+    cmds.menuItem("nyItem", l="-y")
+    cmds.menuItem("zItem", l="z")
+    cmds.menuItem("nzItem", l="-z")
+    cmds.optionMenuGrp("OMG_aimVectorType", e=True, sl=cmds.optionVar(q="OV_aimVectorType"))
+
+    # Up vector dropdown menu
+    cmds.optionMenuGrp("OMG_upVectorType", l="Up Axis", cw=[1, 80], ann="Up Axis is the axis of your joint that will be aimed at the up object you have selected to be the up vector for your joint. Default is \"y\"", cc="cmds.optionVar(iv=(\"OV_upVectorType\", "+"cmds.optionMenuGrp(\"OMG_upVectorType\", q=True, sl=True)"+"))")
+    cmds.menuItem("xItem", l="x")
+    cmds.menuItem("nxItem", l="-x")
+    cmds.menuItem("yItem", l="y")
+    cmds.menuItem("nyItem", l="-y")
+    cmds.menuItem("zItem", l="z")
+    cmds.menuItem("nzItem", l="-z")
+    cmds.optionMenuGrp("OMG_upVectorType", e=True, sl=cmds.optionVar(q="OV_upVectorType"))
+
+    # Heirarchy checkbox
+    cmds.checkBox("heirarchyCB", l="Heirarchy?", ann="Do you want to orient all of the children of the selected joint?", v=cmds.optionVar(q="orientJointHeirarchyOV"), onc="cmds.optionVar(iv=(\"orientJointHeirarchyOV\", 1))", ofc="cmds.optionVar(iv=(\"orientJointHeirarchyOV\", 0))")
+    cmds.setParent( '..' )  
+    cmds.setParent( '..' ) 
+    cmds.setParent( '..' )  
+    cmds.setParent( '..' ) 
+
+    frmLyout = cmds.formLayout(nd=100) 
+
+    # Buttons
+    cmds.button("orientJointButton", l="Orient Joints", ann="Orient Joints with the current settings and close this window.", al="center", w=75, h=20, c=applyAndClose)
+    cmds.button("applyGroupButton", l="Apply", ann="Orient Joints with the current settings.", al="center", w=75, h=20, c=findOrientJointValues)
+    cmds.button("CloseButton", l="Close", ann="Close this window.", al="center", w=75, h=20, c="cmds.deleteUI(\"orientJointsWindow\", window=True)")
+    cmds.setParent( '..' )  
+    cmds.setParent( '..' ) 
+
+    cmds.formLayout(frmLyout, e=True, af=[("orientJointButton", "top", 5), ("orientJointButton", "left", 5), ("orientJointButton", "bottom", 5), ("applyGroupButton", "top", 5), ("applyGroupButton", "bottom", 5), ("CloseButton", "top", 5), ("CloseButton", "right", 5), ("CloseButton", "bottom", 5)], ap=[("orientJointButton", "right", 1, 33), ("applyGroupButton", "right", 1, 66)], ac=[("applyGroupButton", "left", 1, "orientJointButton"), ("CloseButton", "left", 1, "applyGroupButton")])
+    cmds.showWindow()
+
+#----------------------------------------------------------------------------------------
+def applyAndClose(*args):
+    findOrientJointValues()
+    cmds.deleteUI("orientJointsWindow", window=True)
+
+#----------------------------------------------------------------------------------------
+def resetOrientJointVariables(*args):
+    #Fixup options
+    cmds.optionVar(iv=("orientJointHeirarchyOV", 1))
+    cmds.checkBox("heirarchyCB", e=True, v=1)
+
+    cmds.optionVar(iv=("OV_upVectorType", 3))
+    cmds.optionMenuGrp("OMG_upVectorType", e=True, sl=3)
+
+    cmds.optionVar(iv=("OV_aimVectorType", 1))
+    cmds.optionMenuGrp("OMG_aimVectorType", e=True, sl=1)
+
+#----------------------------------------------------------------------------------------
+def findOrientJointValues(*args):
+    upVector = cmds.optionVar(q="OV_upVectorType")
+    aimVector = cmds.optionVar(q="OV_aimVectorType")
+
+    #Set the Up Axis string
+    if upVector == 1:
+        upAxis = "x" 
+    if upVector == 2:
+        upAxis = "-x" 
+    if upVector == 3:
+        upAxis = "y" 
+    if upVector == 4:
+        upAxis = "-y" 
+    if upVector == 5:
+        upAxis = "z" 
+    if upVector == 6:
+        upAxis = "-z" 
+
+    #Set the Aim Axis string
+    if aimVector == 1:
+        aimAxis = "x" 
+    if aimVector == 2:
+        aimAxis = "-x" 
+    if aimVector == 3:
+        aimAxis = "y" 
+    if aimVector == 4:
+        aimAxis = "-y" 
+    if aimVector == 5:
+        aimAxis = "z" 
+    if aimVector == 6:
+        aimAxis = "-z" 
+
+    if cmds.optionVar(q="orientJointHeirarchyOV") == 0:
+        heirarchy = 0 
+    if cmds.optionVar(q="orientJointHeirarchyOV") == 1:
+        heirarchy = 1 
+
+    setOrientWithUpVector(aimAxis, upAxis, heirarchy, 0, None)
+
+# Run it through the UI
+#orientJointOB()
+
+# Run it without the UI (Values are passed in as args)
+#setOrientWithUpVector("x", "y", 0, 1, "lowerarm_l")
+#setOrientWithUpVector("x", "y", 1, 0, None)
diff --git a/Engine/Extras/Maya_AnimationRiggingTools/ARTv1/General/Scripts/Tools/ART_RenameSuffix.py b/Engine/Extras/Maya_AnimationRiggingTools/ARTv1/General/Scripts/Tools/ART_RenameSuffix.py
new file mode 100644
--- /dev/null
+++ b/Engine/Extras/Maya_AnimationRiggingTools/ARTv1/General/Scripts/Tools/ART_RenameSuffix.py
@@ -0,0 +1,10 @@
+import maya.cmds as cmds
+
+def suffixHeirarchyNames(*args):
+    selection = cmds.ls(sl=True, type="transform")
+    result = cmds.promptDialog(title="Suffix Seleceted", message="Enter Suffix:", button=["Ok", "Cancel"])
+    if result == "Ok":
+        suffix = cmds.promptDialog(q=True, text=True)
+        for one in selection:
+            cmds.rename(one, one+suffix)
+#suffixHeirarchyNames()
\ No newline at end of file
diff --git a/Engine/Extras/Maya_AnimationRiggingTools/ARTv1/General/Scripts/Tools/ART_ReverseAttrHookup.py b/Engine/Extras/Maya_AnimationRiggingTools/ARTv1/General/Scripts/Tools/ART_ReverseAttrHookup.py
new file mode 100644
--- /dev/null
+++ b/Engine/Extras/Maya_AnimationRiggingTools/ARTv1/General/Scripts/Tools/ART_ReverseAttrHookup.py
@@ -0,0 +1,69 @@
+import maya.cmds as cmds
+from functools import partial
+
+# Creates a connection between the 2 selected objects with a reverse node to drive to values in opposite directions of each other based on one attr.
+# This is mostly used for having a single attr drive the 2 weights on a constraint with 2 sources in opposite of each other.
+# CRA FEATURE: going forward, add the ability to have 3 or even more inputs driven by the single output.
+# CRA FEATURE: Add check boxes and filters for filtering the attr lists better.
+# To use, select the driver, then the driven, then run the tool.
+
+textScrollList = None
+menuDriven0 = None
+menuDriven1 = None
+
+def ReverseAttrHookupUI(*args):
+    
+    print "UI"
+
+    selection = cmds.ls(sl=True)
+    if len(selection) != 2:
+        cmds.confirmDialog(icon = "warning!!", title = "Reverse Attr Hookup Tool", message = "You must select only 2 objects.  Driver then Driven.")
+        return
+    else:
+        driver = selection[0]
+        driven = selection[1]
+
+
+        attrsListDriver = cmds.listAttr(driver, o=True, c=True, u=True, s=True, k=False)
+        attrsListDriven = cmds.listAttr(driven, o=True, c=True, u=True, s=True, k=True)
+
+
+        if cmds.window("ReverseAttrHookupWin", exists = True):
+            cmds.deleteUI("ReverseAttrHookupWin")
+
+        window = cmds.window("ReverseAttrHookupWin", w=4300, h=500, title="Reverse Attr Hookup Tool", mxb=True, mnb=True, sizeable=False, rtf=True)
+        mainLayout = cmds.columnLayout(w = 300, h = 400, rs = 5, co = ["both", 5])
+        cmds.text(l=driver)
+        textScrollList = cmds.textScrollList("artProjCharacterList", w = 290, h = 300, parent = mainLayout, ra=True, ann="This is the attribute that will drive the Driven attributes listed below.")
+        cmds.textScrollList(textScrollList, e=True, a=attrsListDriver)
+        
+        menuDriven0 = cmds.optionMenuGrp(l=driven+" 0", cw=[1, 175], ann="This attribute will be driven as a direct connection from the driver.  When the driver is 1, this input will be 1.  When the driver is 0, this input will be 0.")
+        for attr in attrsListDriven:
+            cmds.menuItem(l=attr)
+        
+        menuDriven1 = cmds.optionMenuGrp(l=driven+" 1", cw=[1, 175], ann="This attribute will be driven by the driver through a reverse node.  When the driver is 1, this input will be 0.  When the driver is 0, this input will be 1.")
+        for attr in attrsListDriven:
+            cmds.menuItem(l=attr)
+
+        buttonUpdate = cmds.button(w = 290, h = 40, label = "Update", c=ReverseAttrHookupUI, ann = "Refresh the UI.", parent = mainLayout)
+        buttonConnect = cmds.button(w = 290, h = 40, label = "Connect", c=partial(ReverseAttrHookup, driver, driven, textScrollList, menuDriven0, menuDriven1), ann = "Connect the Selected Attrs.", parent = mainLayout)
+         
+        cmds.showWindow(window)
+    
+# This script hooks up the attrs
+def ReverseAttrHookup(driver, driven, textScrollList, menuDriven0, menuDriven1, *args):
+    
+    print "SCRIPT"
+
+    driverAttr = cmds.textScrollList(textScrollList, q=True, si=True)[0]
+    drivenAttr0 = cmds.optionMenuGrp(menuDriven0, q=True, v=True)
+    drivenAttr1 = cmds.optionMenuGrp(menuDriven1, q=True, v=True)
+    
+    print driverAttr
+    print drivenAttr0
+    print drivenAttr1
+
+    reverseNode = cmds.shadingNode("reverse", asUtility=True, name=driver+"_reverse")
+    cmds.connectAttr(driver+"."+driverAttr, driven+"."+drivenAttr0, f=True)
+    cmds.connectAttr(driver+"."+driverAttr, reverseNode+".inputX", f=True)
+    cmds.connectAttr(reverseNode+".outputX", driven+"."+drivenAttr1, f=True)
diff --git a/Engine/Extras/Maya_AnimationRiggingTools/ARTv1/General/Scripts/Tools/ART_SurfaceVertexSelector.py b/Engine/Extras/Maya_AnimationRiggingTools/ARTv1/General/Scripts/Tools/ART_SurfaceVertexSelector.py
new file mode 100644
--- /dev/null
+++ b/Engine/Extras/Maya_AnimationRiggingTools/ARTv1/General/Scripts/Tools/ART_SurfaceVertexSelector.py
@@ -0,0 +1,71 @@
+import maya.cmds as cmds
+import maya.OpenMaya as om
+from functools import partial
+
+# This script walks along the surface of a mesh from a single component and find the nearest N verts based on the number you pass into it (count).
+# It then selects these verts and returns them as a list called "elements"
+def surfaceVertSelector(startElement, *args):
+    count = args[0]
+    #print count
+    vertCount = 0
+    startDist = 0
+    cmds.select(startElement)
+
+    # Find all of the verts on the selected mesh and see if they are less than count.  If they are then set count to be all verts.
+    allVerts = cmds.polyEvaluate(v=True)
+    #print allVerts
+    if count > allVerts:
+        count = allVerts
+    #print args[0]
+        
+    # Run through the mesh and using soft select to traverse out and find verts until we reach the count value.
+    while vertCount < count:
+        startDist = startDist+0.1
+        #print "startDist: ", startDist
+        cmds.softSelect(softSelectEnabled=True, ssf=1, ssd=startDist)
+        selection = om.MSelectionList()
+        surfaceVertSelector = om.MRichSelection()
+        om.MGlobal.getRichSelection(surfaceVertSelector)
+        surfaceVertSelector.getSelection(selection)
+
+        dagPath = om.MDagPath()
+        component = om.MObject()
+
+        iter = om.MItSelectionList( selection,om.MFn.kMeshVertComponent )
+        elements = []
+        while not iter.isDone():
+            iter.getDagPath( dagPath, component )
+            dagPath.pop()
+            node = dagPath.fullPathName()
+            fnComp = om.MFnSingleIndexedComponent(component)
+
+            for i in range(fnComp.elementCount()):
+                elements.append('%s.vtx[%i]' % (node, fnComp.element(i)))
+            iter.next()
+        vertCount = len(elements)
+        #print "vertCount: ", vertCount
+    cmds.softSelect(softSelectEnabled=False)
+    cmds.select(elements)
+
+#startElement = cmds.ls(sl=True)
+#surfaceVertSelector(startElement, 75)
+    
+# Create a UI that will drive the selection tool dynamically using dragCommand.
+# The reset button will allow you to change the starting component(s) on your mesh.
+def surfaceVertSelectorUI(*args):
+    startElement = cmds.ls(sl=True, type="float3")
+    if startElement:
+        if cmds.window("surfaceVertSelector", exists = True):
+            cmds.deleteUI("surfaceVertSelector")
+
+        window = cmds.window("surfaceVertSelector", title='Surface Vert Selector', s=False)
+        cmds.columnLayout()
+        cmds.intSliderGrp("vertCountIntSliderGrp", label='Vert Count', field=True, minValue=1, maxValue=2000, value=1, step=1, cw3=[75, 50, 1000], dragCommand=partial(surfaceVertSelector, startElement))
+        #surfaceVertSelector(startElement, 20)
+        cmds.button("refreshButton", label="Reset", c=surfaceVertSelectorUI)
+
+        cmds.showWindow(window)
+    else:
+        cmds.warning("You must select some components.")
+
+#surfaceVertSelectorUI()
\ No newline at end of file
diff --git a/Engine/Extras/Maya_AnimationRiggingTools/ARTv1/General/Scripts/Tools/ART_TransferWeights.py b/Engine/Extras/Maya_AnimationRiggingTools/ARTv1/General/Scripts/Tools/ART_TransferWeights.py
new file mode 100644
--- /dev/null
+++ b/Engine/Extras/Maya_AnimationRiggingTools/ARTv1/General/Scripts/Tools/ART_TransferWeights.py
@@ -0,0 +1,172 @@
+import maya.cmds as cmds
+from PySide import QtGui
+import maya.OpenMayaUI as mui
+import shiboken
+
+#Written by: Charles R. Anderson
+#Dec. 4, 2015 
+#Questions or comments email me: [REDACTED_EMAIL]
+
+
+# This function finds the main Maya Window
+def getMayaWindow():
+    pointer = mui.MQtUtil.mainWindow()
+    return shiboken.wrapInstance(long(pointer), QtGui.QWidget)
+
+# This function finds the cluster for the passed in mesh.
+def findCluster(meshName):
+    skinClusters = []
+    skinHistory = []
+    skinHistory = cmds.listHistory(meshName, lv=0)
+    for node in skinHistory:
+        if cmds.nodeType(node) == "skinCluster":
+            skinClusters.append(node)
+
+    return skinClusters[0]
+
+# This function transfers the weighting for each mesh in the meshes list for the influences in the jointsToRemove list to the jointToTransferTo joint.
+def LOD_transferWeights(meshes, jointsToRemove, jointToTransferTo, *args):
+    for mesh in meshes:
+        #print "MESH: ", mesh
+        
+        # Find the skin cluster for the current mesh
+        cluster = findCluster(mesh)
+        #print "CLUSTER: ", cluster
+        
+        # Prune weights on the current mesh
+        cmds.skinPercent(cluster, mesh, prw=0.001)
+
+        # Find all of the current influences on the current skin cluster.
+        meshInfluences = cmds.skinCluster(cluster, q=True, inf=True)
+        #print "Current Influences: ", meshInfluences
+
+        for joint in jointsToRemove:
+            if joint in meshInfluences:
+                #print "Current Joint: ", joint
+                
+                # If the jointToTransferTo is not already an influence on the current mesh then add it.
+                currentInfluences = cmds.skinCluster(cluster, q=True, inf=True)
+                if jointToTransferTo not in currentInfluences:
+                    cmds.skinCluster(cluster, e=True, ai=jointToTransferTo, lw=False)
+
+                # Now transfer all of the influences we want to remove onto the jointToTransferTo.
+                for x in range(cmds.polyEvaluate(mesh, v=True)):
+                    #print "TRANSFERRING DATA....."
+                    value = cmds.skinPercent(cluster, (mesh+".vtx["+str(x)+"]"), t=joint, q=True)
+                    if value > 0:
+                        cmds.skinPercent(cluster, (mesh+".vtx["+str(x)+"]"), tmw=[joint, jointToTransferTo])
+
+        # Remove unused influences
+        currentInfluences = cmds.skinCluster(cluster, q=True, inf=True)
+        #print "Current Influences: ", currentInfluences
+        influencesToRemove = []
+        weightedInfs = cmds.skinCluster(cluster, q=True, weightedInfluence=True)
+        #print "Weighted Influences: ", weightedInfs
+        for inf in currentInfluences:
+            #print "Influence: ", inf
+            if inf not in weightedInfs:
+                #print "Update Influences to Remove List: ", inf
+                influencesToRemove.append(inf)
+            
+        #print "ToRemove Influences: ", influencesToRemove
+        if influencesToRemove != []:
+            for inf in influencesToRemove:
+                cmds.skinCluster(cluster, e=True, ri=inf)
+
+
+########  Example
+'''
+meshes = ["hands_geo_LOD4"]
+jointsToRemove = ["index_01_r", "index_02_r", "middle_01_r", "middle_02_r", "pinky_01_r", "pinky_02_r", "ring_01_r", "ring_02_r", "thumb_01_r", "thumb_02_r"]
+jointToTransferTo = "hand_r"
+LOD_transferWeights(meshes, jointsToRemove, jointToTransferTo)
+'''
+##################
+
+
+
+
+
+
+# THIS IS OLD CODE THAT I WAS USING TO CREATE A UI FOR TRANSFERING WEIGHTS.  MAY STILL USE SOME OF IT FOR A TOOL TO CREATE THE REMOVE BONES LIST.
+'''
+def create_LOD_TransferWeightsUI():
+
+    # check for existing window.
+    windowName = "LOD_TransferWeightsWindow"
+    if cmds.window("LOD_TransferWeightsWindow", exists=True):
+        cmds.deleteUI("LOD_TransferWeightsWindow", wnd=True)
+        
+    # Create a window
+    mainMayaWindow = getMayaWindow()
+    LOD_TransferWeightsWindow = QtGui.QMainWindow(mainMayaWindow)
+    LOD_TransferWeightsWindow.setObjectName(windowName)
+
+    # Create a font
+    font = QtGui.QFont()
+    font.setPointSize(20)
+    font.setBold(True)
+
+    # create a widget
+    widget = QtGui.QWidget()
+    LOD_TransferWeightsWindow.setCentralWidget(widget)
+
+    # create a layout
+    layout = QtGui.QVBoxLayout(widget)
+
+    # create a button and add it to the layout
+    button = QtGui.QPushButton("Create Locator")
+    layout.addWidget(button)
+    button.setFont(font)
+    button.setMinimumSize(200, 40)
+    button.setMaximumSize(200, 40)
+    imagePath = cmds.internalVar(upd=True) + "icons/extractDeltasDuplicate.png"
+    button.setStyleSheet("background-image: url("+imagePath+"); border:solid black 1px; color")
+    button.clicked.connect(LOD_transferWeights)
+
+    # Create a close button
+    closeButton = QtGui.QPushButton("Close")
+    layout.addWidget(closeButton)
+    closeButton.setFont(font)
+    closeButton.setMinimumSize(200, 40)
+    closeButton.setMaximumSize(200, 40)
+    imagePath = cmds.internalVar(upd=True) + "icons/extractDeltasDuplicate.png"
+    closeButton.setStyleSheet("background-image: url("+imagePath+"); border:solid black 1px")
+    closeButton.clicked.connect(LOD_TransferWeightsWindow.close)
+
+
+    # show the button
+    LOD_TransferWeightsWindow.show()
+
+create_LOD_TransferWeightsUI()
+
+
+
+def OutlinerSelection():
+    itemSelection = cmds.ls(sl=True, type="joint")
+    sc = cmds.selectionConnection(lst=True)
+    if len(itemSelection):
+        for node in itemSelection:
+            print node
+            cmds.selectionConnection(sc, e=True, select=node)
+
+        # Create a window layout to hold outliner
+        newWindow = cmds.window(title="Outliner (Custom)", iconName="Outliner*", rtf=0, h=800, w=500)
+        frame = cmds.frameLayout(labelVisible = False)
+        oped = cmds.outlinerEditor("jointOutliner", showSetMembers=True, unParent=True, mainListConnection=sc, selectionConnection="modelList", eo=True)
+        testWin = cmds.outlinerEditor(oped, e=True, parent=frame)
+
+        # scriptJob to delete selectionConnection
+        print "Deleting of the connection"
+        cmds.scriptJob(uiDeleted = [newWindow, "cmds.deleteUI('"+sc+"')"])
+
+        cmds.showWindow(newWindow)
+
+    else:
+        cmds.warning("Nothing Selected! Custom Outliner will not be created")
+
+OutlinerSelection() 
+
+jointList = cmds.ls(sl=True, type="joint")
+print jointList
+'''
\ No newline at end of file
diff --git a/Engine/Extras/Maya_AnimationRiggingTools/ARTv1/General/Scripts/Tools/ART_addSpaceSwitching.py b/Engine/Extras/Maya_AnimationRiggingTools/ARTv1/General/Scripts/Tools/ART_addSpaceSwitching.py
new file mode 100644
--- /dev/null
+++ b/Engine/Extras/Maya_AnimationRiggingTools/ARTv1/General/Scripts/Tools/ART_addSpaceSwitching.py
@@ -0,0 +1,63 @@
+import maya.cmds as cmds
+
+# Creates a new space switching object for the mainObj using the spaceMasterObj.  
+# If you want this new space to be in control, set the control arg to "True".  False will leave everything as it is and just add the new space without it being turned on.
+
+def addSpaceSwitching(mainObj, spaceMasterObj, control, *args):
+    # find all of the objects we need to do the space switching
+    space_switcher_follow = mainObj+"_space_switcher_follow"
+    space_switcher = mainObj+"_space_switcher"
+    cmds.select(space_switcher)
+    
+    spaceMasterObjClean = spaceMasterObj.replace(":", "_")
+    mainObjClean = mainObj.replace(":", "_")
+     
+    # add attr to the space switcher node
+    cmds.addAttr(space_switcher, ln = "space_"+spaceMasterObjClean, minValue = 0, maxValue = 1, dv = 0, keyable = True)
+    
+    # add constraint to the new object on the follow node
+    constraint = cmds.parentConstraint(spaceMasterObj, space_switcher_follow, mo = True)[0]
+
+    # find the entire list of targets on the constraint node
+    targets = cmds.parentConstraint(constraint, q = True, targetList = True)
+
+    # using the targets list, find the weight value for the node we are working with.
+    weight = 0
+    for i in range(int(len(targets))):
+        if targets[i].find(spaceMasterObj) != -1:
+            weight = i
+            
+    # Connect the attr we added earlier to the weight value in the constraint
+    tempName = spaceMasterObj
+    #print spaceMasterObj.find(":")
+    if spaceMasterObj.find(":") != -1:
+        tempName = spaceMasterObj.partition(":")[2]
+        #print "TEST"
+    
+    #print "tempname = ",tempName
+    cmds.connectAttr(space_switcher+".space_"+spaceMasterObjClean, constraint+"."+tempName+"W"+str(weight))
+
+    # lockNode on space object so it cannot be deleted by the user (if node is not a referenced node)
+    if spaceMasterObj.find(":") == -1:
+        cmds.lockNode(spaceMasterObj)
+
+    # set the values of the space switcher based on user input.
+    if control == True:
+        attrs = cmds.listAttr(space_switcher, k=True)
+        for attr in attrs:
+            if attr.find("space") != -1:
+                cmds.setAttr(space_switcher+"."+attr, 0)
+                
+        cmds.setAttr(space_switcher+"."+"space_"+spaceMasterObjClean, 1)
+
+
+    cmds.select(space_switcher)
+
+
+# Example calls to this script.
+'''
+addSpaceSwitching("Rider:ik_foot_anim_l", "cockpit_pedal_l_anim", True)
+'''
+
+#addSpaceSwitching("Rider:ik_foot_anim_l", "cockpit_foot_pedal_l_anim", True)
+#addSpaceSwitching("cockpit_foot_pedal_l_anim", "Rider:ik_foot_anim_l", True)
