Branching

diff --git a/Engine/Extras/Maya_AnimationRiggingTools/ARTv1/General/Scripts/autoRig_CharacterUI.py b/Engine/Extras/Maya_AnimationRiggingTools/ARTv1/General/Scripts/autoRig_CharacterUI.py
new file mode 100644
--- /dev/null
+++ b/Engine/Extras/Maya_AnimationRiggingTools/ARTv1/General/Scripts/autoRig_CharacterUI.py
@@ -0,0 +1,5499 @@
+import maya.cmds as cmds
+import maya.mel as mel
+from functools import partial
+import os
+import re
+
+
+def forceSceneUpdate():
+    #func updates the scene/dependency graph
+    cmds.setToolTo( "moveSuperContext" )
+    nodes = cmds.ls()
+
+    for node in nodes:
+        cmds.select(node, replace = True)
+
+    cmds.select(clear = True)
+    cmds.setToolTo( "selectSuperContext" )
+    
+    
+def addFaceRig(*args):
+    namespace = cmds.optionMenu("selectedCharacterOptionMenu", q = True, v = True)
+    namespace += ":"
+    
+    #find out if the namespace is from a referenced file
+    refNode = cmds.referenceQuery( (namespace + "NonRigItems"), isNodeReferenced=True )
+    
+    if refNode == True:
+        cmds.warning( "Function Failed. Character is from a referenced file." )
+        return
+            
+    #reset the rig to be zeroed out
+    resetAll()
+    
+    #get all of the face rig files
+    configFile = cmds.internalVar(usd = True)
+    configFilePath = (configFile + "autoRigConfig.txt")
+        
+    #open file, get lines, find peoject value
+    f = open(configFilePath, "r")
+    lines=f.readlines()
+    path =  lines[1]
+    f.close()
+    
+    pathName = path.rpartition("/")[0]
+    newPath = pathName.rpartition("/")[0]
+    newPath += "/FaceRigs/"
+    faceRigs = os.listdir(newPath)
+    
+    #check to see if you can find the namespace in any of those files(or actually, the other way around)
+    for face in faceRigs:
+        faceName = face.rpartition(".mb")[0]
+        
+        if re.search('(.*?)%s' %faceName, namespace):
+            cmds.file((newPath + faceName + ".mb"), i=True)
+            
+            
+    #Rename everything to have the proper namespace
+    cmds.select("FaceRig_Group", r = True, hi = True)
+    selection = cmds.ls(sl = True, transforms = True, shapes = False)
+    
+    for each in selection:
+        cmds.rename(each, (namespace + each))
+                    
+    cmds.rename("FaceRig_Container", (namespace + "FaceRig_Container"))
+    
+    
+    #Now constrain everything up
+    
+    #constraint the face rig to the head bone of the character
+    cmds.parentConstraint((namespace + "b_MF_Head"), (namespace + "Cage_Head"), mo = True)
+    
+    #Lower mouth
+    cmds.parentConstraint((namespace + "Cage_Jaw_Offset"), (namespace + "b_MF_Jaw"), mo = True)
+    cmds.parentConstraint((namespace + "Control_Rig_b_MF_L_Bot_Mid_Lip"), (namespace + "b_MF_L_Bot_Mid_Lip"), mo = True)
+    cmds.parentConstraint((namespace + "Control_Rig_b_MF_L_Bot_Lip"), (namespace + "b_MF_L_Bot_Lip"), mo = True)
+    cmds.parentConstraint((namespace + "Control_Rig_b_MF_L_Corner"), (namespace + "b_MF_L_Corner"), mo = True)
+    cmds.parentConstraint((namespace + "Control_Rig_b_MF_Bot_Mid_Lip"), (namespace + "b_MF_Bot_Mid_Lip"), mo = True)
+    cmds.parentConstraint((namespace + "Control_Rig_b_MF_R_Bot_Mid_Lip"), (namespace + "b_MF_R_Bot_Mid_Lip"), mo = True)
+    cmds.parentConstraint((namespace + "Control_Rig_b_MF_R_Bot_Lip"), (namespace + "b_MF_R_Bot_Lip"), mo = True)
+    cmds.parentConstraint((namespace + "Control_Rig_b_MF_R_Corner"), (namespace + "b_MF_R_Corner"), mo = True)
+    
+    #Brow area
+    cmds.parentConstraint((namespace + "Control_Rig_b_MF_L_Inner_Brow"), (namespace + "b_MF_L_Inner_Brow"), mo = True)
+    cmds.parentConstraint((namespace + "Control_Rig_b_MF_L_Mid_Brow"), (namespace + "b_MF_L_Mid_Brow"), mo = True)
+    cmds.parentConstraint((namespace + "Control_Rig_b_MF_L_Out_Brow"), (namespace + "b_MF_L_Out_Brow"), mo = True)
+    cmds.parentConstraint((namespace + "Control_Rig_b_MF_R_Inner_Brow"), (namespace + "b_MF_R_Inner_Brow"), mo = True)
+    cmds.parentConstraint((namespace + "Control_Rig_b_MF_R_Mid_Brow"), (namespace + "b_MF_R_Mid_Brow"), mo = True)
+    cmds.parentConstraint((namespace + "Control_Rig_b_MF_R_Out_Brow"), (namespace + "b_MF_R_Out_Brow"), mo = True)
+    
+    #Nose/Cheek area
+    cmds.parentConstraint((namespace + "Control_Rig_b_MF_L_Nose"), (namespace + "b_MF_L_Nose"), mo = True)
+    cmds.parentConstraint((namespace + "Control_Rig_b_MF_R_Nose"), (namespace + "b_MF_R_Nose"), mo = True)
+    cmds.parentConstraint((namespace + "Control_Rig_b_MF_L_Cheek"), (namespace + "b_MF_L_Cheek"), mo = True)
+    cmds.parentConstraint((namespace + "Control_Rig_b_MF_R_Cheek"), (namespace + "b_MF_R_Cheek"), mo = True)
+    cmds.parentConstraint((namespace + "Control_Rig_b_MF_L_Far_Cheek"), (namespace + "b_MF_L_Far_Cheek"), mo = True)
+    cmds.parentConstraint((namespace + "Control_Rig_b_MF_L_Far_Cheek"), (namespace + "b_MF_L_Far_Cheek"), mo = True)
+    cmds.parentConstraint((namespace + "Control_Rig_b_MF_L_Puff"), (namespace + "b_MF_L_Puff"), mo = True)
+    cmds.parentConstraint((namespace + "Control_Rig_b_MF_R_Puff"), (namespace + "b_MF_R_Puff"), mo = True)
+    
+    #Top Mouth
+    cmds.parentConstraint((namespace + "Control_Rig_b_MF_L_Top_Mid_Lip"), (namespace + "b_MF_L_Top_Mid_Lip"), mo = True)
+    cmds.parentConstraint((namespace + "Control_Rig_b_MF_R_Top_Mid_Lip"), (namespace + "b_MF_R_Top_Mid_Lip"), mo = True)
+    cmds.parentConstraint((namespace + "Control_Rig_b_MF_Top_Mid_Lip"), (namespace + "b_MF_Top_Mid_Lip"), mo = True)
+    cmds.parentConstraint((namespace + "Control_Rig_b_MF_L_Top_Lip"), (namespace + "b_MF_L_Top_Lip"), mo = True)
+    cmds.parentConstraint((namespace + "Control_Rig_b_MF_R_Top_Lip"), (namespace + "b_MF_R_Top_Lip"), mo = True)
+    
+    #Lids and Eyes
+    cmds.parentConstraint((namespace + "Left_Eye_Anim"), (namespace + "b_MF_L_Eye"), mo = True)
+    cmds.parentConstraint((namespace + "Right_Eye_Anim"), (namespace + "b_MF_R_Eye"), mo = True)
+    
+    cmds.parentConstraint((namespace + "Left_Upper_Lid_Anim"), (namespace + "b_MF_L_Upper_Lid"), mo = True)
+    cmds.parentConstraint((namespace + "Left_Lower_Lid_Anim"), (namespace + "b_MF_L_Lower_Lid"), mo = True)
+    
+    cmds.parentConstraint((namespace + "Right_Upper_Lid_Anim"), (namespace + "b_MF_R_Upper_Lid"), mo = True)
+    cmds.parentConstraint((namespace + "Right_Lower_Lid_Anim"), (namespace + "b_MF_R_Lower_Lid"), mo = True)
+    
+    #Add shader to controls
+    name = namespace.rpartition(":")[0]
+    
+    cmds.select((namespace + "FaceRig_Group"), r = True, hi = True)
+    selection = cmds.ls(sl = True, transforms = True)
+    controls = []
+    
+    for each in selection:
+        
+        if re.search("_Anim", each):
+            controls.append(each)
+            
+    for control in controls:
+        cmds.hyperShade(control, assign = (name + "_Face_Control_M"))
+    
+
+
+
+
+    #Add new brows up/down attr,inner brow up/down attr, move sneer over to lower face(add new attrs, hide old ones)setup sdks
+    cmds.setAttr((namespace + "FACS_Morphs_Group.v"),0)
+    cmds.setAttr((namespace + "Morph_Cage.v"),0)
+    
+    
+    #add attrs
+    cmds.select((namespace + "UpperFaceShapes"), r = True)
+    cmds.addAttr(ln = "L_Brow_Down_Up", at = 'double', min = -1, max = 1, dv = 0, k = True)
+    cmds.addAttr(ln = "R_Brow_Down_Up", at = 'double', min = -1, max = 1, dv = 0, k = True)
+    cmds.addAttr(ln = "L_Inner_Brow_Down_Up", at = 'double', min = -1, max = 1, dv = 0, k = True)
+    cmds.addAttr(ln = "R_Inner_Brow_Down_Up", at = 'double', min = -1, max = 1, dv = 0, k = True)
+    
+    cmds.select((namespace + "LowerFaceShapes"), r = True)
+    cmds.addAttr(ln = "L_Sneer", at = 'double', min = 0, max = 1, dv = 0, k = True)
+    cmds.addAttr(ln = "R_Sneer", at = 'double', min = 0, max = 1, dv = 0, k = True)
+    
+    #sdks
+    cmds.setAttr((namespace + "UpperFaceShapes.L_Brow_Down_Up"), 0)
+    cmds.setAttr((namespace + "UpperFaceShapes.L_Brow_Up"), 0)
+    cmds.setAttr((namespace + "UpperFaceShapes.L_Brow_Down"), 0)
+    
+    cmds.setDrivenKeyframe((namespace + "UpperFaceShapes.L_Brow_Up"), currentDriver = (namespace + "UpperFaceShapes.L_Brow_Down_Up"))
+    cmds.setDrivenKeyframe((namespace + "UpperFaceShapes.L_Brow_Down"), currentDriver = (namespace + "UpperFaceShapes.L_Brow_Down_Up"))
+    
+    cmds.setAttr((namespace + "UpperFaceShapes.L_Brow_Down_Up"), -1)
+    cmds.setAttr((namespace + "UpperFaceShapes.L_Brow_Up"), 0)
+    cmds.setAttr((namespace + "UpperFaceShapes.L_Brow_Down"), 1)
+    
+    cmds.setDrivenKeyframe((namespace + "UpperFaceShapes.L_Brow_Up"), currentDriver = (namespace + "UpperFaceShapes.L_Brow_Down_Up"))
+    cmds.setDrivenKeyframe((namespace + "UpperFaceShapes.L_Brow_Down"), currentDriver = (namespace + "UpperFaceShapes.L_Brow_Down_Up"))
+    
+    cmds.setAttr((namespace + "UpperFaceShapes.L_Brow_Down_Up"), 1)
+    cmds.setAttr((namespace + "UpperFaceShapes.L_Brow_Up"), 1)
+    cmds.setAttr((namespace + "UpperFaceShapes.L_Brow_Down"), 0)
+    
+    cmds.setDrivenKeyframe((namespace + "UpperFaceShapes.L_Brow_Up"), currentDriver = (namespace + "UpperFaceShapes.L_Brow_Down_Up"))
+    cmds.setDrivenKeyframe((namespace + "UpperFaceShapes.L_Brow_Down"), currentDriver = (namespace + "UpperFaceShapes.L_Brow_Down_Up"))
+    
+    cmds.setAttr((namespace + "UpperFaceShapes.L_Brow_Down_Up"), 0)
+    
+    
+    #r brow up/down
+    cmds.setAttr((namespace + "UpperFaceShapes.R_Brow_Down_Up"), 0)
+    cmds.setAttr((namespace + "UpperFaceShapes.R_Brow_Up"), 0)
+    cmds.setAttr((namespace + "UpperFaceShapes.R_Brow_Down"), 0)
+    
+    cmds.setDrivenKeyframe((namespace + "UpperFaceShapes.R_Brow_Up"), currentDriver = (namespace + "UpperFaceShapes.R_Brow_Down_Up"))
+    cmds.setDrivenKeyframe((namespace + "UpperFaceShapes.R_Brow_Down"), currentDriver = (namespace + "UpperFaceShapes.R_Brow_Down_Up"))
+    
+    cmds.setAttr((namespace + "UpperFaceShapes.R_Brow_Down_Up"), -1)
+    cmds.setAttr((namespace + "UpperFaceShapes.R_Brow_Up"), 0)
+    cmds.setAttr((namespace + "UpperFaceShapes.R_Brow_Down"), 1)
+    
+    cmds.setDrivenKeyframe((namespace + "UpperFaceShapes.R_Brow_Up"), currentDriver = (namespace + "UpperFaceShapes.R_Brow_Down_Up"))
+    cmds.setDrivenKeyframe((namespace + "UpperFaceShapes.R_Brow_Down"), currentDriver = (namespace + "UpperFaceShapes.R_Brow_Down_Up"))
+    
+    cmds.setAttr((namespace + "UpperFaceShapes.R_Brow_Down_Up"), 1)
+    cmds.setAttr((namespace + "UpperFaceShapes.R_Brow_Up"), 1)
+    cmds.setAttr((namespace + "UpperFaceShapes.R_Brow_Down"), 0)
+    
+    cmds.setDrivenKeyframe((namespace + "UpperFaceShapes.R_Brow_Up"), currentDriver = (namespace + "UpperFaceShapes.R_Brow_Down_Up"))
+    cmds.setDrivenKeyframe((namespace + "UpperFaceShapes.R_Brow_Down"), currentDriver = (namespace + "UpperFaceShapes.R_Brow_Down_Up"))
+        
+        
+    cmds.setAttr((namespace + "UpperFaceShapes.R_Brow_Down_Up"), 0)
+    
+    
+    #inner brows (L)
+    cmds.setAttr((namespace + "UpperFaceShapes.L_Inner_Brow_Down_Up"), 0)
+    cmds.setAttr((namespace + "UpperFaceShapes.L_Inner_Brow_Up"), 0)
+    cmds.setAttr((namespace + "UpperFaceShapes.L_Inner_Brow_Down"), 0)
+    
+    cmds.setDrivenKeyframe((namespace + "UpperFaceShapes.L_Inner_Brow_Up"), currentDriver = (namespace + "UpperFaceShapes.L_Inner_Brow_Down_Up"))
+    cmds.setDrivenKeyframe((namespace + "UpperFaceShapes.L_Inner_Brow_Down"), currentDriver = (namespace + "UpperFaceShapes.L_Inner_Brow_Down_Up"))
+    
+    cmds.setAttr((namespace + "UpperFaceShapes.L_Inner_Brow_Down_Up"), -1)
+    cmds.setAttr((namespace + "UpperFaceShapes.L_Inner_Brow_Up"), 0)
+    cmds.setAttr((namespace + "UpperFaceShapes.L_Inner_Brow_Down"), 1)
+    
+    cmds.setDrivenKeyframe((namespace + "UpperFaceShapes.L_Inner_Brow_Up"), currentDriver = (namespace + "UpperFaceShapes.L_Inner_Brow_Down_Up"))
+    cmds.setDrivenKeyframe((namespace + "UpperFaceShapes.L_Inner_Brow_Down"), currentDriver = (namespace + "UpperFaceShapes.L_Inner_Brow_Down_Up"))
+    
+    cmds.setAttr((namespace + "UpperFaceShapes.L_Inner_Brow_Down_Up"), 1)
+    cmds.setAttr((namespace + "UpperFaceShapes.L_Inner_Brow_Up"), 1)
+    cmds.setAttr((namespace + "UpperFaceShapes.L_Inner_Brow_Down"), 0)
+    
+    cmds.setDrivenKeyframe((namespace + "UpperFaceShapes.L_Inner_Brow_Up"), currentDriver = (namespace + "UpperFaceShapes.L_Inner_Brow_Down_Up"))
+    cmds.setDrivenKeyframe((namespace + "UpperFaceShapes.L_Inner_Brow_Down"), currentDriver = (namespace + "UpperFaceShapes.L_Inner_Brow_Down_Up"))
+    
+    cmds.setAttr((namespace + "UpperFaceShapes.L_Inner_Brow_Down_Up"), 0)
+    
+    #inner brows (R)
+    cmds.setAttr((namespace + "UpperFaceShapes.R_Inner_Brow_Down_Up"), 0)
+    cmds.setAttr((namespace + "UpperFaceShapes.R_Inner_Brow_Up"), 0)
+    cmds.setAttr((namespace + "UpperFaceShapes.R_Inner_Brow_Down"), 0)
+    
+    cmds.setDrivenKeyframe((namespace + "UpperFaceShapes.R_Inner_Brow_Up"), currentDriver = (namespace + "UpperFaceShapes.R_Inner_Brow_Down_Up"))
+    cmds.setDrivenKeyframe((namespace + "UpperFaceShapes.R_Inner_Brow_Down"), currentDriver = (namespace + "UpperFaceShapes.R_Inner_Brow_Down_Up"))
+    
+    cmds.setAttr((namespace + "UpperFaceShapes.R_Inner_Brow_Down_Up"), -1)
+    cmds.setAttr((namespace + "UpperFaceShapes.R_Inner_Brow_Up"), 0)
+    cmds.setAttr((namespace + "UpperFaceShapes.R_Inner_Brow_Down"), 1)
+    
+    cmds.setDrivenKeyframe((namespace + "UpperFaceShapes.R_Inner_Brow_Up"), currentDriver = (namespace + "UpperFaceShapes.R_Inner_Brow_Down_Up"))
+    cmds.setDrivenKeyframe((namespace + "UpperFaceShapes.R_Inner_Brow_Down"), currentDriver = (namespace + "UpperFaceShapes.R_Inner_Brow_Down_Up"))
+    
+    cmds.setAttr((namespace + "UpperFaceShapes.R_Inner_Brow_Down_Up"), 1)
+    cmds.setAttr((namespace + "UpperFaceShapes.R_Inner_Brow_Up"), 1)
+    cmds.setAttr((namespace + "UpperFaceShapes.R_Inner_Brow_Down"), 0)
+    
+    cmds.setDrivenKeyframe((namespace + "UpperFaceShapes.R_Inner_Brow_Up"), currentDriver = (namespace + "UpperFaceShapes.R_Inner_Brow_Down_Up"))
+    cmds.setDrivenKeyframe((namespace + "UpperFaceShapes.R_Inner_Brow_Down"), currentDriver = (namespace + "UpperFaceShapes.R_Inner_Brow_Down_Up"))
+    
+    cmds.setAttr((namespace + "UpperFaceShapes.R_Inner_Brow_Down_Up"), 0)
+    
+    
+    #sneers(moving to lower face node)
+    cmds.setAttr((namespace + "LowerFaceShapes.L_Sneer"), 0)
+    cmds.setAttr((namespace + "UpperFaceShapes.L_Sneer"), 0)
+    
+    cmds.setDrivenKeyframe((namespace + "UpperFaceShapes.L_Sneer"), currentDriver = (namespace + "LowerFaceShapes.L_Sneer"))
+    
+    cmds.setAttr((namespace + "LowerFaceShapes.L_Sneer"), 1)
+    cmds.setAttr((namespace + "UpperFaceShapes.L_Sneer"), 1)
+    
+    cmds.setDrivenKeyframe((namespace + "UpperFaceShapes.L_Sneer"), currentDriver = (namespace + "LowerFaceShapes.L_Sneer"))
+    
+
+    cmds.setAttr((namespace + "LowerFaceShapes.L_Sneer"), 0)
+    
+
+    cmds.setAttr((namespace + "LowerFaceShapes.R_Sneer"), 0)
+    cmds.setAttr((namespace + "UpperFaceShapes.R_Sneer"), 0)
+    
+    cmds.setDrivenKeyframe((namespace + "UpperFaceShapes.R_Sneer"), currentDriver = (namespace + "LowerFaceShapes.R_Sneer"))
+    
+    cmds.setAttr((namespace + "LowerFaceShapes.R_Sneer"), 1)
+    cmds.setAttr((namespace + "UpperFaceShapes.R_Sneer"), 1)
+    
+    cmds.setDrivenKeyframe((namespace + "UpperFaceShapes.R_Sneer"), currentDriver = (namespace + "LowerFaceShapes.R_Sneer"))
+    
+
+    cmds.setAttr((namespace + "LowerFaceShapes.R_Sneer"), 0)
+    
+    
+    
+    #hide original attrs
+    cmds.setAttr((namespace + "UpperFaceShapes.R_Brow_Up"), channelBox = False, k = False)
+    cmds.setAttr((namespace + "UpperFaceShapes.R_Brow_Down"), channelBox = False, k = False)
+    cmds.setAttr((namespace + "UpperFaceShapes.L_Brow_Up"), channelBox = False, k = False)
+    cmds.setAttr((namespace + "UpperFaceShapes.L_Brow_Down"), channelBox = False, k = False)
+    
+    cmds.setAttr((namespace + "UpperFaceShapes.R_Inner_Brow_Up"), channelBox = False, k = False)
+    cmds.setAttr((namespace + "UpperFaceShapes.R_Inner_Brow_Down"), channelBox = False, k = False)
+    cmds.setAttr((namespace + "UpperFaceShapes.L_Inner_Brow_Up"), channelBox = False, k = False)
+    cmds.setAttr((namespace + "UpperFaceShapes.L_Inner_Brow_Down"), channelBox = False, k = False)
+    
+    cmds.setAttr((namespace + "UpperFaceShapes.L_Sneer"), channelBox = False, k = False)
+    cmds.setAttr((namespace + "UpperFaceShapes.R_Sneer"), channelBox = False, k = False)
+    
+    
+    #hide un-used offset controls( inner cheek and cheek fat)
+    cmds.setAttr((namespace + "b_MF_L_Inner_Cheek_Anim.v"),0)
+    cmds.setAttr((namespace + "b_MF_L_CheekFat_Anim.v"),0)
+    cmds.setAttr((namespace + "b_MF_R_Inner_Cheek_Anim.v"),0)
+    cmds.setAttr((namespace + "b_MF_R_CheekFat_Anim.v"),0)
+    
+    
+    #Hook up morph targets.
+    
+    if cmds.objExists((namespace + "Correctives")):
+        
+        try:
+            
+            
+            cmds.setAttr((namespace + "LowerFaceShapes.Jaw_Drop"), 0)
+            cmds.setAttr((namespace + "Correctives.Jaw_Drop"), 0)
+            
+            cmds.setDrivenKeyframe((namespace + "Correctives.Jaw_Drop"), currentDriver = (namespace + "LowerFaceShapes.Jaw_Drop"))
+            
+            
+            cmds.setAttr((namespace + "LowerFaceShapes.Jaw_Drop"), -1)
+            cmds.setAttr((namespace + "Correctives.Jaw_Drop"), 1)
+            
+            cmds.setDrivenKeyframe((namespace + "Correctives.Jaw_Drop"), currentDriver = (namespace + "LowerFaceShapes.Jaw_Drop"))
+            
+            
+            cmds.setAttr((namespace + "LowerFaceShapes.Jaw_Drop"), 0)
+            
+            
+        except RuntimeError:
+            print "No Jaw Drop morph found"
+            
+        try:
+            cmds.setAttr((namespace + "LowerFaceShapes.LNarrow"), 0)
+            cmds.setAttr((namespace + "Correctives.L_Narrow"), 0)
+            
+            
+            cmds.setDrivenKeyframe((namespace + "Correctives.L_Narrow"), currentDriver = (namespace + "LowerFaceShapes.LNarrow"))
+            
+            cmds.setAttr((namespace + "LowerFaceShapes.LNarrow"), 1)
+            cmds.setAttr((namespace + "Correctives.L_Narrow"), 1)
+            
+            
+            cmds.setDrivenKeyframe((namespace + "Correctives.L_Narrow"), currentDriver = (namespace + "LowerFaceShapes.LNarrow"))
+            
+            cmds.setAttr((namespace + "LowerFaceShapes.LNarrow"), 0)
+            
+            
+        except RuntimeError:
+            print "No L Narrow morph found"
+            
+        try:
+            
+            cmds.setAttr((namespace + "LowerFaceShapes.RNarrow"), 0)
+            cmds.setAttr((namespace + "Correctives.R_Narrow"), 0)
+            
+            
+            cmds.setDrivenKeyframe((namespace + "Correctives.R_Narrow"), currentDriver = (namespace + "LowerFaceShapes.RNarrow"))
+            
+            
+            cmds.setAttr((namespace + "LowerFaceShapes.RNarrow"), 1)
+            cmds.setAttr((namespace + "Correctives.R_Narrow"), 1)
+            
+            
+            cmds.setDrivenKeyframe((namespace + "Correctives.R_Narrow"), currentDriver = (namespace + "LowerFaceShapes.RNarrow"))
+            
+            
+            cmds.setAttr((namespace + "LowerFaceShapes.RNarrow"), 0)
+            
+            
+        except RuntimeError:
+            print "No R Narrow morph found"
+            
+            
+        try:
+            
+            cmds.setAttr((namespace + "UpperFaceShapes.L_Brow_Squeeze"), 0)
+            cmds.setAttr((namespace + "Correctives.L_Brow_Squeeze"), 0)
+            
+            
+            cmds.setDrivenKeyframe((namespace + "Correctives.L_Brow_Squeeze"), currentDriver = (namespace + "UpperFaceShapes.L_Brow_Squeeze"))
+            
+            
+            cmds.setAttr((namespace + "UpperFaceShapes.L_Brow_Squeeze"), 1)
+            cmds.setAttr((namespace + "Correctives.L_Brow_Squeeze"), 1)
+            
+            
+            cmds.setDrivenKeyframe((namespace + "Correctives.L_Brow_Squeeze"), currentDriver = (namespace + "UpperFaceShapes.L_Brow_Squeeze"))
+            
+            
+            cmds.setAttr((namespace + "UpperFaceShapes.L_Brow_Squeeze"), 0)
+            
+            
+        except RuntimeError:
+            print "No L_Brow Squeeze morph found"
+            
+        try:
+            
+            cmds.setAttr((namespace + "UpperFaceShapes.R_Brow_Squeeze"), 0)
+            cmds.setAttr((namespace + "Correctives.R_Brow_Squeeze"), 0)
+            
+            
+            cmds.setDrivenKeyframe((namespace + "Correctives.R_Brow_Squeeze"), currentDriver = (namespace + "UpperFaceShapes.R_Brow_Squeeze"))
+            
+            
+            cmds.setAttr((namespace + "UpperFaceShapes.R_Brow_Squeeze"), 1)
+            cmds.setAttr((namespace + "Correctives.R_Brow_Squeeze"), 1)
+            
+            
+            cmds.setDrivenKeyframe((namespace + "Correctives.R_Brow_Squeeze"), currentDriver = (namespace + "UpperFaceShapes.R_Brow_Squeeze"))
+            
+            
+            cmds.setAttr((namespace + "UpperFaceShapes.R_Brow_Squeeze"), 0)
+            
+            
+        except RuntimeError:
+            print "No L_Brow Squeeze morph found"
+            
+            
+        try:
+            
+            cmds.setAttr((namespace + "UpperFaceShapes.L_Brow_Down_Up"), 0)
+            cmds.setAttr((namespace + "Correctives.L_Brow_Down"), 0)
+            
+            
+            cmds.setDrivenKeyframe((namespace + "Correctives.L_Brow_Down"), currentDriver = (namespace + "UpperFaceShapes.L_Brow_Down_Up"))
+            
+            
+            cmds.setAttr((namespace + "UpperFaceShapes.L_Brow_Down_Up"), -1)
+            cmds.setAttr((namespace + "Correctives.L_Brow_Down"), 1)
+            
+            
+            cmds.setDrivenKeyframe((namespace + "Correctives.L_Brow_Down"), currentDriver = (namespace + "UpperFaceShapes.L_Brow_Down_Up"))
+            
+            
+            cmds.setAttr((namespace + "UpperFaceShapes.L_Brow_Down_Up"), 0)
+            
+            
+        except RuntimeError:
+            print "No morph found"
+            
+            
+        try:
+            
+            cmds.setAttr((namespace + "UpperFaceShapes.R_Brow_Down_Up"), 0)
+            cmds.setAttr((namespace + "Correctives.R_Brow_Down"), 0)
+            
+            
+            cmds.setDrivenKeyframe((namespace + "Correctives.R_Brow_Down"), currentDriver = (namespace + "UpperFaceShapes.R_Brow_Down_Up"))
+            
+            
+            cmds.setAttr((namespace + "UpperFaceShapes.R_Brow_Down_Up"), -1)
+            cmds.setAttr((namespace + "Correctives.R_Brow_Down"), 1)
+            
+            
+            cmds.setDrivenKeyframe((namespace + "Correctives.R_Brow_Down"), currentDriver = (namespace + "UpperFaceShapes.R_Brow_Down_Up"))
+            
+            
+            cmds.setAttr((namespace + "UpperFaceShapes.R_Brow_Down_Up"), 0)
+            
+            
+        except RuntimeError:
+            print "No morph found"
+            
+            
+        try:
+            
+            cmds.setAttr((namespace + "UpperFaceShapes.L_Brow_Down_Up"), 0)
+            cmds.setAttr((namespace + "Correctives.L_Brow_Up"), 0)
+            
+            
+            cmds.setDrivenKeyframe((namespace + "Correctives.L_Brow_Up"), currentDriver = (namespace + "UpperFaceShapes.L_Brow_Down_Up"))
+            
+            
+            cmds.setAttr((namespace + "UpperFaceShapes.L_Brow_Down_Up"), 1)
+            cmds.setAttr((namespace + "Correctives.L_Brow_Up"), 1)
+            
+            
+            cmds.setDrivenKeyframe((namespace + "Correctives.L_Brow_Up"), currentDriver = (namespace + "UpperFaceShapes.L_Brow_Down_Up"))
+            
+            
+            cmds.setAttr((namespace + "UpperFaceShapes.L_Brow_Down_Up"), 0)
+            
+            
+        except RuntimeError:
+            print "No morph found"
+            
+            
+        try:
+            
+            cmds.setAttr((namespace + "UpperFaceShapes.R_Brow_Down_Up"), 0)
+            cmds.setAttr((namespace + "Correctives.R_Brow_Up"), 0)
+            
+            
+            cmds.setDrivenKeyframe((namespace + "Correctives.R_Brow_Up"), currentDriver = (namespace + "UpperFaceShapes.R_Brow_Down_Up"))
+            
+            
+            cmds.setAttr((namespace + "UpperFaceShapes.R_Brow_Down_Up"), 1)
+            cmds.setAttr((namespace + "Correctives.R_Brow_Up"), 1)
+            
+            
+            cmds.setDrivenKeyframe((namespace + "Correctives.R_Brow_Up"), currentDriver = (namespace + "UpperFaceShapes.R_Brow_Down_Up"))
+            
+            
+            cmds.setAttr((namespace + "UpperFaceShapes.R_Brow_Down_Up"), 0)
+            
+            
+        except RuntimeError:
+            print "No morph found"
+            
+        try:
+            
+            cmds.setAttr((namespace + "LowerFaceShapes.L_Cheek_Puff"), 0)
+            cmds.setAttr((namespace + "Correctives.L_Cheek_Puff"), 0)
+            
+            
+            cmds.setDrivenKeyframe((namespace + "Correctives.L_Cheek_Puff"), currentDriver = (namespace + "LowerFaceShapes.L_Cheek_Puff"))
+            
+            
+            cmds.setAttr((namespace + "LowerFaceShapes.L_Cheek_Puff"), 1)
+            cmds.setAttr((namespace + "Correctives.L_Cheek_Puff"), 1)
+            
+            
+            cmds.setDrivenKeyframe((namespace + "Correctives.L_Cheek_Puff"), currentDriver = (namespace + "LowerFaceShapes.L_Cheek_Puff"))
+            
+            
+            cmds.setAttr((namespace + "LowerFaceShapes.L_Cheek_Puff"), 0)
+            
+            
+        except RuntimeError:
+            print "No morph found"
+            
+            
+        try:
+            
+            cmds.setAttr((namespace + "LowerFaceShapes.R_Cheek_Puff"), 0)
+            cmds.setAttr((namespace + "Correctives.R_Cheek_Puff"), 0)
+            
+            
+            cmds.setDrivenKeyframe((namespace + "Correctives.R_Cheek_Puff"), currentDriver = (namespace + "LowerFaceShapes.R_Cheek_Puff"))
+            
+            
+            cmds.setAttr((namespace + "LowerFaceShapes.R_Cheek_Puff"), 1)
+            cmds.setAttr((namespace + "Correctives.R_Cheek_Puff"), 1)
+            
+            
+            cmds.setDrivenKeyframe((namespace + "Correctives.R_Cheek_Puff"), currentDriver = (namespace + "LowerFaceShapes.R_Cheek_Puff"))
+            
+            
+            cmds.setAttr((namespace + "LowerFaceShapes.R_Cheek_Puff"), 0)
+            
+            
+        except RuntimeError:
+            print "No morph found"
+            
+
+        try:
+            
+            cmds.setAttr((namespace + "LowerFaceShapes.L_Corner_Down"), 0)
+            cmds.setAttr((namespace + "Correctives.L_Corner_Down"), 0)
+            
+            
+            cmds.setDrivenKeyframe((namespace + "Correctives.L_Corner_Down"), currentDriver = (namespace + "LowerFaceShapes.L_Corner_Down"))
+            
+            
+            cmds.setAttr((namespace + "LowerFaceShapes.L_Corner_Down"), 1)
+            cmds.setAttr((namespace + "Correctives.L_Corner_Down"), 1)
+            
+            
+            cmds.setDrivenKeyframe((namespace + "Correctives.L_Corner_Down"), currentDriver = (namespace + "LowerFaceShapes.L_Corner_Down"))
+            
+            
+            cmds.setAttr((namespace + "LowerFaceShapes.L_Corner_Down"), 0)
+            
+            
+        except RuntimeError:
+            print "No morph found"
+            
+            
+        try:
+            
+            cmds.setAttr((namespace + "LowerFaceShapes.R_Corner_Down"), 0)
+            cmds.setAttr((namespace + "Correctives.R_Corner_Down"), 0)
+            
+            
+            cmds.setDrivenKeyframe((namespace + "Correctives.R_Corner_Down"), currentDriver = (namespace + "LowerFaceShapes.R_Corner_Down"))
+            
+            
+            cmds.setAttr((namespace + "LowerFaceShapes.R_Corner_Down"), 1)
+            cmds.setAttr((namespace + "Correctives.R_Corner_Down"), 1)
+            
+            
+            cmds.setDrivenKeyframe((namespace + "Correctives.R_Corner_Down"), currentDriver = (namespace + "LowerFaceShapes.R_Corner_Down"))
+            
+            
+            cmds.setAttr((namespace + "LowerFaceShapes.R_Corner_Down"), 0)
+            
+            
+        except RuntimeError:
+            print "No morph found"
+            
+
+        try:
+            
+            cmds.setAttr((namespace + "LowerFaceShapes.L_Corner_Up"), 0)
+            cmds.setAttr((namespace + "Correctives.L_Corner_Up"), 0)
+            
+            
+            cmds.setDrivenKeyframe((namespace + "Correctives.L_Corner_Up"), currentDriver = (namespace + "LowerFaceShapes.L_Corner_Up"))
+            
+            
+            cmds.setAttr((namespace + "LowerFaceShapes.L_Corner_Up"), 1)
+            cmds.setAttr((namespace + "Correctives.L_Corner_Up"), 1)
+            
+            
+            cmds.setDrivenKeyframe((namespace + "Correctives.L_Corner_Up"), currentDriver = (namespace + "LowerFaceShapes.L_Corner_Up"))
+            
+            
+            cmds.setAttr((namespace + "LowerFaceShapes.L_Corner_Up"), 0)
+            
+            
+        except RuntimeError:
+            print "No morph found"
+            
+
+        try:
+            
+            cmds.setAttr((namespace + "LowerFaceShapes.R_Corner_Up"), 0)
+            cmds.setAttr((namespace + "Correctives.R_Corner_Up"), 0)
+            
+            
+            cmds.setDrivenKeyframe((namespace + "Correctives.R_Corner_Up"), currentDriver = (namespace + "LowerFaceShapes.R_Corner_Up"))
+            
+            
+            cmds.setAttr((namespace + "LowerFaceShapes.R_Corner_Up"), 1)
+            cmds.setAttr((namespace + "Correctives.R_Corner_Up"), 1)
+            
+            
+            cmds.setDrivenKeyframe((namespace + "Correctives.R_Corner_Up"), currentDriver = (namespace + "LowerFaceShapes.R_Corner_Up"))
+            
+            
+            cmds.setAttr((namespace + "LowerFaceShapes.R_Corner_Up"), 0)
+            
+            
+        except RuntimeError:
+            print "No morph found"
+            
+            
+        try:
+            
+            cmds.setAttr((namespace + "LowerFaceShapes.L_Sneer"), 0)
+            cmds.setAttr((namespace + "Correctives.L_Sneer"), 0)
+            
+            
+            cmds.setDrivenKeyframe((namespace + "Correctives.L_Sneer"), currentDriver = (namespace + "LowerFaceShapes.L_Sneer"))
+            
+            
+            cmds.setAttr((namespace + "LowerFaceShapes.L_Sneer"), 1)
+            cmds.setAttr((namespace + "Correctives.L_Sneer"), 1)
+            
+            
+            cmds.setDrivenKeyframe((namespace + "Correctives.L_Sneer"), currentDriver = (namespace + "LowerFaceShapes.L_Sneer"))
+            
+            
+            cmds.setAttr((namespace + "LowerFaceShapes.L_Sneer"), 0)
+            
+            
+        except RuntimeError:
+            print "No morph found"
+            
+            
+        try:
+            
+            cmds.setAttr((namespace + "LowerFaceShapes.R_Sneer"), 0)
+            cmds.setAttr((namespace + "Correctives.R_Sneer"), 0)
+            
+            
+            cmds.setDrivenKeyframe((namespace + "Correctives.R_Sneer"), currentDriver = (namespace + "LowerFaceShapes.R_Sneer"))
+            
+            
+            cmds.setAttr((namespace + "LowerFaceShapes.R_Sneer"), 1)
+            cmds.setAttr((namespace + "Correctives.R_Sneer"), 1)
+            
+            
+            cmds.setDrivenKeyframe((namespace + "Correctives.R_Sneer"), currentDriver = (namespace + "LowerFaceShapes.R_Sneer"))
+            
+            
+            cmds.setAttr((namespace + "LowerFaceShapes.R_Sneer"), 0)
+            
+            
+        except RuntimeError:
+            print "No morph found"
+            
+            
+        try:
+            
+            cmds.setAttr((namespace + "LowerFaceShapes.L_Low_Lip_Up_Down"), 0)
+            cmds.setAttr((namespace + "Correctives.L_Low_Lip_Down"), 0)
+            
+            
+            cmds.setDrivenKeyframe((namespace + "Correctives.L_Low_Lip_Down"), currentDriver = (namespace + "LowerFaceShapes.L_Low_Lip_Up_Down"))
+            
+            
+            cmds.setAttr((namespace + "LowerFaceShapes.L_Low_Lip_Up_Down"), -1)
+            cmds.setAttr((namespace + "Correctives.L_Low_Lip_Down"), 1)
+            
+            
+            cmds.setDrivenKeyframe((namespace + "Correctives.L_Low_Lip_Down"), currentDriver = (namespace + "LowerFaceShapes.L_Low_Lip_Up_Down"))
+            
+            
+            cmds.setAttr((namespace + "LowerFaceShapes.L_Low_Lip_Up_Down"), 0)
+            
+            
+        except RuntimeError:
+            print "No morph found"
+            
+            
+        try:
+            
+            cmds.setAttr((namespace + "LowerFaceShapes.R_Low_Lip_Up_Down"), 0)
+            cmds.setAttr((namespace + "Correctives.R_Low_Lip_Down"), 0)
+            
+            
+            cmds.setDrivenKeyframe((namespace + "Correctives.R_Low_Lip_Down"), currentDriver = (namespace + "LowerFaceShapes.R_Low_Lip_Up_Down"))
+            
+            
+            cmds.setAttr((namespace + "LowerFaceShapes.R_Low_Lip_Up_Down"), -1)
+            cmds.setAttr((namespace + "Correctives.R_Low_Lip_Down"), 1)
+            
+            
+            cmds.setDrivenKeyframe((namespace + "Correctives.R_Low_Lip_Down"), currentDriver = (namespace + "LowerFaceShapes.R_Low_Lip_Up_Down"))
+            
+            
+            cmds.setAttr((namespace + "LowerFaceShapes.R_Low_Lip_Up_Down"), 0)
+            
+            
+        except RuntimeError:
+            print "No morph found"
+            
+            
+        try:
+            
+            cmds.setAttr((namespace + "LowerFaceShapes.L_Up_Lip_Up_Down"), 0)
+            cmds.setAttr((namespace + "Correctives.L_Up_Lip_Up"), 0)
+            
+            
+            cmds.setDrivenKeyframe((namespace + "Correctives.L_Up_Lip_Up"), currentDriver = (namespace + "LowerFaceShapes.L_Up_Lip_Up_Down"))
+            
+            
+            cmds.setAttr((namespace + "LowerFaceShapes.L_Up_Lip_Up_Down"), 1)
+            cmds.setAttr((namespace + "Correctives.L_Up_Lip_Up"), 1)
+            
+            
+            cmds.setDrivenKeyframe((namespace + "Correctives.L_Up_Lip_Up"), currentDriver = (namespace + "LowerFaceShapes.L_Up_Lip_Up_Down"))
+            
+            
+            cmds.setAttr((namespace + "LowerFaceShapes.L_Up_Lip_Up_Down"), 0)
+            
+            
+        except RuntimeError:
+            print "No morph found"
+            
+            
+        try:
+            
+            cmds.setAttr((namespace + "LowerFaceShapes.R_Up_Lip_Up_Down"), 0)
+            cmds.setAttr((namespace + "Correctives.R_Up_Lip_Up"), 0)
+            
+            
+            cmds.setDrivenKeyframe((namespace + "Correctives.R_Up_Lip_Up"), currentDriver = (namespace + "LowerFaceShapes.R_Up_Lip_Up_Down"))
+            
+            
+            cmds.setAttr((namespace + "LowerFaceShapes.R_Up_Lip_Up_Down"), 1)
+            cmds.setAttr((namespace + "Correctives.R_Up_Lip_Up"), 1)
+            
+            
+            cmds.setDrivenKeyframe((namespace + "Correctives.R_Up_Lip_Up"), currentDriver = (namespace + "LowerFaceShapes.R_Up_Lip_Up_Down"))
+            
+            
+            cmds.setAttr((namespace + "LowerFaceShapes.R_Up_Lip_Up_Down"), 0)
+            
+            
+        except RuntimeError:
+            print "No morph found"
+            
+            
+        try:
+            
+            cmds.setAttr((namespace + "UpperFaceShapes.L_Inner_Brow_Down_Up"), 0)
+            cmds.setAttr((namespace + "Correctives.L_Inner_Brow_Down"), 0)
+            
+            
+            cmds.setDrivenKeyframe((namespace + "Correctives.L_Inner_Brow_Down"), currentDriver = (namespace + "UpperFaceShapes.L_Inner_Brow_Down_Up"))
+            
+            
+            cmds.setAttr((namespace + "UpperFaceShapes.L_Inner_Brow_Down_Up"), -1)
+            cmds.setAttr((namespace + "Correctives.L_Inner_Brow_Down"), 1)
+            
+            
+            cmds.setDrivenKeyframe((namespace + "Correctives.L_Inner_Brow_Down"), currentDriver = (namespace + "UpperFaceShapes.L_Inner_Brow_Down_Up"))
+            
+            
+            cmds.setAttr((namespace + "UpperFaceShapes.L_Inner_Brow_Down_Up"), 0)
+            
+            
+        except RuntimeError:
+            print "No morph found"
+            
+
+        try:
+            
+            cmds.setAttr((namespace + "UpperFaceShapes.R_Inner_Brow_Down_Up"), 0)
+            cmds.setAttr((namespace + "Correctives.R_Inner_Brow_Down"), 0)
+            
+            
+            cmds.setDrivenKeyframe((namespace + "Correctives.R_Inner_Brow_Down"), currentDriver = (namespace + "UpperFaceShapes.R_Inner_Brow_Down_Up"))
+            
+            
+            cmds.setAttr((namespace + "UpperFaceShapes.R_Inner_Brow_Down_Up"), -1)
+            cmds.setAttr((namespace + "Correctives.R_Inner_Brow_Down"), 1)
+            
+            
+            cmds.setDrivenKeyframe((namespace + "Correctives.R_Inner_Brow_Down"), currentDriver = (namespace + "UpperFaceShapes.R_Inner_Brow_Down_Up"))
+            
+            
+            cmds.setAttr((namespace + "UpperFaceShapes.R_Inner_Brow_Down_Up"), 0)
+            
+            
+        except RuntimeError:
+            print "No morph found"
+            
+            
+        try:
+            
+            cmds.setAttr((namespace + "UpperFaceShapes.L_Inner_Brow_Down_Up"), 0)
+            cmds.setAttr((namespace + "Correctives.L_Inner_Brow_Up"), 0)
+            
+            
+            cmds.setDrivenKeyframe((namespace + "Correctives.L_Inner_Brow_Up"), currentDriver = (namespace + "UpperFaceShapes.L_Inner_Brow_Down_Up"))
+            
+            
+            cmds.setAttr((namespace + "UpperFaceShapes.L_Inner_Brow_Down_Up"), 1)
+            cmds.setAttr((namespace + "Correctives.L_Inner_Brow_Up"), 1)
+            
+            
+            cmds.setDrivenKeyframe((namespace + "Correctives.L_Inner_Brow_Up"), currentDriver = (namespace + "UpperFaceShapes.L_Inner_Brow_Down_Up"))
+            
+            
+            cmds.setAttr((namespace + "UpperFaceShapes.L_Inner_Brow_Down_Up"), 0)
+            
+            
+        except RuntimeError:
+            print "No morph found"
+            
+            
+        try:
+            
+            cmds.setAttr((namespace + "UpperFaceShapes.R_Inner_Brow_Down_Up"), 0)
+            cmds.setAttr((namespace + "Correctives.R_Inner_Brow_Up"), 0)
+            
+            
+            cmds.setDrivenKeyframe((namespace + "Correctives.R_Inner_Brow_Up"), currentDriver = (namespace + "UpperFaceShapes.R_Inner_Brow_Down_Up"))
+            
+            
+            cmds.setAttr((namespace + "UpperFaceShapes.R_Inner_Brow_Down_Up"), 1)
+            cmds.setAttr((namespace + "Correctives.R_Inner_Brow_Up"), 1)
+            
+            
+            cmds.setDrivenKeyframe((namespace + "Correctives.R_Inner_Brow_Up"), currentDriver = (namespace + "UpperFaceShapes.R_Inner_Brow_Down_Up"))
+            
+            
+            cmds.setAttr((namespace + "UpperFaceShapes.R_Inner_Brow_Down_Up"), 0)
+            
+            
+        except RuntimeError:
+            print "No morph found"
+            
+            
+        try:
+            
+            cmds.setAttr((namespace + "UpperFaceShapes.L_Cheek_Raiser"), 0)
+            cmds.setAttr((namespace + "Correctives.L_Cheek_Raise"), 0)
+            
+            
+            cmds.setDrivenKeyframe((namespace + "Correctives.L_Cheek_Raise"), currentDriver = (namespace + "UpperFaceShapes.L_Cheek_Raiser"))
+            
+            
+            cmds.setAttr((namespace + "UpperFaceShapes.L_Cheek_Raiser"), 1)
+            cmds.setAttr((namespace + "Correctives.L_Cheek_Raise"), 1)
+            
+            
+            cmds.setDrivenKeyframe((namespace + "Correctives.L_Cheek_Raise"), currentDriver = (namespace + "UpperFaceShapes.L_Cheek_Raiser"))
+            
+            
+            cmds.setAttr((namespace + "UpperFaceShapes.L_Cheek_Raiser"), 0)
+            
+            
+        except RuntimeError:
+            print "No morph found"
+    
+        try:
+            
+            cmds.setAttr((namespace + "UpperFaceShapes.R_Cheek_Raiser"), 0)
+            cmds.setAttr((namespace + "Correctives.R_Cheek_Raise"), 0)
+            
+            
+            cmds.setDrivenKeyframe((namespace + "Correctives.R_Cheek_Raise"), currentDriver = (namespace + "UpperFaceShapes.R_Cheek_Raiser"))
+            
+            
+            cmds.setAttr((namespace + "UpperFaceShapes.R_Cheek_Raiser"), 1)
+            cmds.setAttr((namespace + "Correctives.R_Cheek_Raise"), 1)
+            
+            
+            cmds.setDrivenKeyframe((namespace + "Correctives.R_Cheek_Raise"), currentDriver = (namespace + "UpperFaceShapes.R_Cheek_Raiser"))
+            
+            
+            cmds.setAttr((namespace + "UpperFaceShapes.R_Cheek_Raiser"), 0)
+            
+            
+        except RuntimeError:
+            print "No morph found"
+            
+            
+
+        try:
+            
+            cmds.setAttr((namespace + "LowerFaceShapes.Upper_Lip_Roll"), 0)
+            cmds.setAttr((namespace + "Correctives.Upper_Lip_Roll"), 0)
+            
+            
+            cmds.setDrivenKeyframe((namespace + "Correctives.Upper_Lip_Roll"), currentDriver = (namespace + "LowerFaceShapes.Upper_Lip_Roll"))
+            
+            
+            cmds.setAttr((namespace + "LowerFaceShapes.Upper_Lip_Roll"), 1)
+            cmds.setAttr((namespace + "Correctives.Upper_Lip_Roll"), 1)
+            
+            
+            cmds.setDrivenKeyframe((namespace + "Correctives.Upper_Lip_Roll"), currentDriver = (namespace + "LowerFaceShapes.Upper_Lip_Roll"))
+            
+            
+            cmds.setAttr((namespace + "LowerFaceShapes.Upper_Lip_Roll"), 0)
+            
+            
+        except RuntimeError:
+            print "No morph found"
+            
+            
+        try:
+            
+            cmds.setAttr((namespace + "LowerFaceShapes.Lower_Lip_Roll"), 0)
+            cmds.setAttr((namespace + "Correctives.Lower_Lip_Roll"), 0)
+            
+            
+            cmds.setDrivenKeyframe((namespace + "Correctives.Lower_Lip_Roll"), currentDriver = (namespace + "LowerFaceShapes.Lower_Lip_Roll"))
+            
+            
+            cmds.setAttr((namespace + "LowerFaceShapes.Lower_Lip_Roll"), 1)
+            cmds.setAttr((namespace + "Correctives.Lower_Lip_Roll"), 1)
+            
+            
+            cmds.setDrivenKeyframe((namespace + "Correctives.Lower_Lip_Roll"), currentDriver = (namespace + "LowerFaceShapes.Lower_Lip_Roll"))
+            
+            
+            cmds.setAttr((namespace + "LowerFaceShapes.Lower_Lip_Roll"), 0)
+            
+            
+        except RuntimeError:
+            print "No morph found"
+    
+    
+    #set eye master to "head" orientation
+    cmds.setAttr((namespace + "Eye_Master_Anim.HeadOrient_WorldOrient"), 0)
+
+    
+    
+    
+    
+def matchOverFrameRange(mode, *args):
+    
+    #grab the namespace of the character
+    namespace = cmds.optionMenu("selectedCharacterOptionMenu", q = True, v = True)
+    namespace += ":"
+    
+    #get the limbs to be matched
+    leftArm = cmds.checkBox("LeftArmCB", q = True, v = True)
+    rightArm = cmds.checkBox("RightArmCB", q = True, v = True)
+    leftLeg = cmds.checkBox("LeftLegCB", q = True, v = True)
+    rightLeg = cmds.checkBox("RightLegCB", q = True, v = True)
+    
+    print leftArm, rightArm, leftLeg, rightLeg
+    
+    #then get the frame range
+    start = cmds.textField("startFrameField", q = True, text = True)
+    end = cmds.textField("endFrameField", q = True, text = True)
+    
+    timeRange = (int(end) - int(start))
+
+
+    
+    print start, end
+    
+    if mode == "FK":
+        print "matching FK to IK"
+        cmds.currentTime(start)
+        
+        for i in range(int(timeRange + 1)):
+            
+            
+            if leftArm == True:
+                leftArmMatchIK(namespace)
+                
+
+            if rightArm == True:
+                rightArmMatchIK(namespace)
+                
+            if leftLeg == True:
+                LLegMatchIK(namespace)
+                
+            if rightLeg == True:
+                RLegMatchIK(namespace)
+                
+            #move the current time up one
+            time = cmds.currentTime(q = True)
+            cmds.currentTime(time + 1)
+        
+        
+        #now get the current Time
+        time = cmds.currentTime(q = True)
+        
+        #then switch the limbs to FK
+        if leftArm == True:
+            cmds.setAttr((namespace + "Settings.L_Arm_Mode"), 0)
+            cmds.setKeyframe((namespace + "Settings.L_Arm_Mode"), time = (time - 1), ott = 'step')
+            
+        if rightArm == True:
+            cmds.setAttr((namespace + "Settings.R_Arm_Mode"), 0)
+            cmds.setKeyframe((namespace + "Settings.R_Arm_Mode"), time = (time - 1), ott = 'step')
+            
+        if leftLeg == True:
+            cmds.setAttr((namespace + "Settings.L_Leg_Mode"), 0)
+            cmds.setKeyframe((namespace + "Settings.L_Leg_Mode"), time = (time - 1), ott = 'step')
+            
+        if rightLeg == True:
+            cmds.setAttr((namespace + "Settings.R_Leg_Mode"), 0)
+            cmds.setKeyframe((namespace + "Settings.R_Leg_Mode"), time = (time - 1), ott = 'step')
+            
+
+                          
+                
+        
+    
+    if mode == "IK":
+        print "matching IK to FK"
+        
+        cmds.currentTime(start)
+        
+        for i in range(int(timeRange + 1)):
+            
+            
+            if leftArm == True:
+                leftArmMatchFK(namespace)
+                
+
+            if rightArm == True:
+                rightArmMatchFK(namespace)
+                
+            if leftLeg == True:
+                LLegMatchFK(namespace)
+                
+            if rightLeg == True:
+                RLegMatchFK(namespace)
+                
+            #move the current time up one
+            time = cmds.currentTime(q = True)
+            cmds.currentTime(time + 1)
+        
+        
+        #now get the current Time
+        time = cmds.currentTime(q = True)
+        
+        #then switch the limbs to IK
+        if leftArm == True:
+            cmds.setAttr((namespace + "Settings.L_Arm_Mode"), 10)
+            cmds.setKeyframe((namespace + "Settings.L_Arm_Mode"), time = (time - 1), ott = 'step')
+            
+        if rightArm == True:
+            cmds.setAttr((namespace + "Settings.R_Arm_Mode"), 10)
+            cmds.setKeyframe((namespace + "Settings.R_Arm_Mode"), time = (time - 1), ott = 'step')
+            
+        if leftLeg == True:
+            cmds.setAttr((namespace + "Settings.L_Leg_Mode"), 10)
+            cmds.setKeyframe((namespace + "Settings.L_Leg_Mode"), time = (time - 1), ott = 'step')
+            
+        if rightLeg == True:
+            cmds.setAttr((namespace + "Settings.R_Leg_Mode"), 10)
+            cmds.setKeyframe((namespace + "Settings.R_Leg_Mode"), time = (time - 1), ott = 'step')
+            
+    cmds.deleteUI("autoRigMatchRangeUI")
+        
+
+
+def matchOverFrameRange_UI(*args):
+    
+    #if the window exists, kill it
+    if cmds.window("autoRigMatchRangeUI", exists = True):
+        cmds.deleteUI("autoRigMatchRangeUI")
+        
+    #create the window
+    window = cmds.window("autoRigMatchRangeUI", title = "Match Over Frame Range", w = 400, h = 150, mxb = False, mnb = False, sizeable = False)
+    
+    #create the main layout
+    mainLayout = cmds.columnLayout(w = 400, h = 150)
+    
+    #create the layout for the frame range fields
+    frameRangeLayout = cmds.rowColumnLayout(nc = 4, cw = [(1,100), (2, 100), (3, 100), (4, 100)])
+    
+    #create the labels and fields
+    cmds.text(label = "Start Frame :", align = 'left')
+    
+    cmds.textField("startFrameField")
+    
+    cmds.text(label = "    End Frame :", align = 'left')
+    
+    cmds.textField("endFrameField")
+    
+    
+    cmds.separator(style = 'out', h = 30)
+    cmds.separator(style = 'out', h = 30)
+    cmds.separator(style = 'out', h = 30)
+    cmds.separator(style = 'out', h = 30)
+    
+    
+    #next create the checkboxes for the limbs
+    cmds.checkBox("LeftArmCB", v = 0, label = "Left Arm")
+    
+    cmds.checkBox("RightArmCB", v = 0, label = "Right Arm")
+    
+    cmds.checkBox("LeftLegCB", v = 0, label = "Left Leg")
+    
+    cmds.checkBox("RightLegCB", v = 0, label = "Right Leg")
+    
+    
+    cmds.separator(style = 'out', h = 30)
+    cmds.separator(style = 'out', h = 30)
+    cmds.separator(style = 'out', h = 30)
+    cmds.separator(style = 'out', h = 30)
+    
+    
+    #next create the two match buttons
+    cmds.text(label = "", visible = False)
+    cmds.text(label = "", visible = False)
+    
+    cmds.button(label = "Match FK to IK", c = partial(matchOverFrameRange, "FK"))
+    cmds.button(label = "Match IK to FK", c = partial(matchOverFrameRange, "IK"))
+    
+    
+    #show the window
+    cmds.showWindow(window)
+    
+    
+def hideControls(*args):
+    
+    #grab the namespace of the character
+    namespace = cmds.optionMenu("selectedCharacterOptionMenu", q = True, v = True)
+    namespace += ":"
+    
+    #select all controls
+    selectAll()
+    cmds.select((namespace + "L_Index_1_Anim"),add = True)
+    cmds.select((namespace + "L_Index_2_Anim"),add = True)
+    cmds.select((namespace + "L_Index_3_Anim"),add = True)
+    
+    cmds.select((namespace + "L_Middle_1_Anim"),add = True)
+    cmds.select((namespace + "L_Middle_2_Anim"),add = True)
+    cmds.select((namespace + "L_Middle_3_Anim"),add = True)
+    
+    cmds.select((namespace + "L_Ring_1_Anim"),add = True)
+    cmds.select((namespace + "L_Ring_2_Anim"),add = True)
+    cmds.select((namespace + "L_Ring_3_Anim"),add = True)
+    
+    cmds.select((namespace + "L_Pinky_1_Anim"),add = True)
+    cmds.select((namespace + "L_Pinky_2_Anim"),add = True)
+    cmds.select((namespace + "L_Pinky_3_Anim"),add = True)
+    
+    cmds.select((namespace + "L_Thumb_1_Anim"),add = True)
+    cmds.select((namespace + "L_Thumb_2_Anim"),add = True)
+    cmds.select((namespace + "L_Thumb_3_Anim"),add = True)
+    
+    cmds.select((namespace + "R_Index_1_Anim"),add = True)
+    cmds.select((namespace + "R_Index_2_Anim"),add = True)
+    cmds.select((namespace + "R_Index_3_Anim"),add = True)
+    
+    cmds.select((namespace + "R_Middle_1_Anim"),add = True)
+    cmds.select((namespace + "R_Middle_2_Anim"),add = True)
+    cmds.select((namespace + "R_Middle_3_Anim"),add = True)
+    
+    cmds.select((namespace + "R_Ring_1_Anim"),add = True)
+    cmds.select((namespace + "R_Ring_2_Anim"),add = True)
+    cmds.select((namespace + "R_Ring_3_Anim"),add = True)
+    
+    cmds.select((namespace + "R_Pinky_1_Anim"),add = True)
+    cmds.select((namespace + "R_Pinky_2_Anim"),add = True)
+    cmds.select((namespace + "R_Pinky_3_Anim"),add = True)
+    
+    cmds.select((namespace + "R_Thumb_1_Anim"),add = True)
+    cmds.select((namespace + "R_Thumb_2_Anim"),add = True)
+    cmds.select((namespace + "R_Thumb_3_Anim"),add = True)
+    
+    selection = cmds.ls(sl = True, transforms = True)
+    
+    #unlock the container
+    cmds.lockNode((namespace + "Rig_Container"), lock = False, lockUnpublished = False)
+    
+    #for each in selection, hide
+    for each in selection:
+        if each != (namespace + "Settings"):
+            cmds.setAttr((each + ".v"), 0)
+        
+    cmds.select(clear = True)
+    
+    #relock the container
+    cmds.lockNode((namespace + "Rig_Container"), lock = False, lockUnpublished = False)
+    
+    
+def showControls(*args):
+    
+    #grab the namespace of the character
+    namespace = cmds.optionMenu("selectedCharacterOptionMenu", q = True, v = True)
+    namespace += ":"
+    
+    #select all controls
+    selectAll()
+    cmds.select((namespace + "L_Index_1_Anim"),add = True)
+    cmds.select((namespace + "L_Index_2_Anim"),add = True)
+    cmds.select((namespace + "L_Index_3_Anim"),add = True)
+    
+    cmds.select((namespace + "L_Middle_1_Anim"),add = True)
+    cmds.select((namespace + "L_Middle_2_Anim"),add = True)
+    cmds.select((namespace + "L_Middle_3_Anim"),add = True)
+    
+    cmds.select((namespace + "L_Ring_1_Anim"),add = True)
+    cmds.select((namespace + "L_Ring_2_Anim"),add = True)
+    cmds.select((namespace + "L_Ring_3_Anim"),add = True)
+    
+    cmds.select((namespace + "L_Pinky_1_Anim"),add = True)
+    cmds.select((namespace + "L_Pinky_2_Anim"),add = True)
+    cmds.select((namespace + "L_Pinky_3_Anim"),add = True)
+    
+    cmds.select((namespace + "L_Thumb_1_Anim"),add = True)
+    cmds.select((namespace + "L_Thumb_2_Anim"),add = True)
+    cmds.select((namespace + "L_Thumb_3_Anim"),add = True)
+    
+    cmds.select((namespace + "R_Index_1_Anim"),add = True)
+    cmds.select((namespace + "R_Index_2_Anim"),add = True)
+    cmds.select((namespace + "R_Index_3_Anim"),add = True)
+    
+    cmds.select((namespace + "R_Middle_1_Anim"),add = True)
+    cmds.select((namespace + "R_Middle_2_Anim"),add = True)
+    cmds.select((namespace + "R_Middle_3_Anim"),add = True)
+    
+    cmds.select((namespace + "R_Ring_1_Anim"),add = True)
+    cmds.select((namespace + "R_Ring_2_Anim"),add = True)
+    cmds.select((namespace + "R_Ring_3_Anim"),add = True)
+    
+    cmds.select((namespace + "R_Pinky_1_Anim"),add = True)
+    cmds.select((namespace + "R_Pinky_2_Anim"),add = True)
+    cmds.select((namespace + "R_Pinky_3_Anim"),add = True)
+    
+    cmds.select((namespace + "R_Thumb_1_Anim"),add = True)
+    cmds.select((namespace + "R_Thumb_2_Anim"),add = True)
+    cmds.select((namespace + "R_Thumb_3_Anim"),add = True)
+    
+    selection = cmds.ls(sl = True, transforms = True)
+    
+    #unlock the container
+    cmds.lockNode((namespace + "Rig_Container"), lock = False, lockUnpublished = False)
+    
+    
+    #get the current settings
+    lArmMode = cmds.getAttr((namespace + "Settings.L_Arm_Mode"))
+    rArmMode = cmds.getAttr((namespace + "Settings.R_Arm_Mode"))
+    
+    lLegMode = cmds.getAttr((namespace + "Settings.L_Leg_Mode"))
+    rLegMode = cmds.getAttr((namespace + "Settings.R_Leg_Mode"))
+    
+    
+    #for each in selection, hide
+    for each in selection:
+        if each != (namespace + "Settings"):
+            cmds.setAttr((each + ".v"), 1)
+        
+    cmds.select(clear = True)
+    
+    #relock the container
+    cmds.lockNode((namespace + "Rig_Container"), lock = False, lockUnpublished = False)
+    
+    #reset settings
+    cmds.setAttr((namespace + "Settings.L_Arm_Mode"), lArmMode)
+    cmds.setAttr((namespace + "Settings.R_Arm_Mode"), rArmMode)
+    
+    cmds.setAttr((namespace + "Settings.L_Leg_Mode"), lLegMode)
+    cmds.setAttr((namespace + "Settings.R_Leg_Mode"), rLegMode)
+
+def addWeapon(*args):
+    
+    #grab the namespace of the character
+    namespace = cmds.optionMenu("selectedCharacterOptionMenu", q = True, v = True)
+
+    
+    #open the UI
+    addWeaponUI(namespace)
+    
+    
+def addWeaponUI(namespace):
+    
+    #if the window exists, kill it. then create the window
+    if cmds.window("Weapon_Locker",exists = True):
+        cmds.deleteUI("Weapon_Locker")
+        
+    window = cmds.window("Weapon_Locker", title = "Weapon Locker", w = 300, h = 150, sizeable = False)
+    
+    #main layout
+    mainLayout = cmds.columnLayout(w = 300, h = 150)
+    
+    
+    #create the layout for the top label and the project combo box
+    topLayout = cmds.rowColumnLayout(nc = 2, cw = [(1, 100), (2, 200)], parent = mainLayout)
+    
+    #add a separator
+    cmds.separator(h =10,style='out' )
+    cmds.separator(h =10,style='out' )    
+    
+    #create the text label
+    projLabel = cmds.text(label = "Choose A Project: ", align = 'left', parent = topLayout)
+    #Next we will create the combo box to hold the list of projects, but first we must find the project list in the directory
+    projectOptionMenu = cmds.optionMenu('WeaponLockerProjectOptionMenu',w = 200, cc = updateList)
+    
+    
+    #open the autoRigConfig file and find the models path
+    configFile = cmds.internalVar(usd = True)
+    configFilePath = (configFile + "autoRigConfig.txt")
+        
+    #open file, get lines, find peoject value
+    f = open(configFilePath, "r")
+    lines=f.readlines()
+    path =  lines[1]
+    f.close()
+    
+    pathName = path.rpartition("/")[0]
+    newPath = pathName.rpartition("/")[0]
+    newPath += "/Weapons/"
+    projects = os.listdir(newPath)
+    
+    #for each project, add the project to the option menu
+    for proj in projects:
+        #add a combo box entry
+        cmds.menuItem( label= proj)
+        
+    #add a separator
+    cmds.separator(h =30,style='out' )
+    cmds.separator(h =30,style='out' )
+    
+    #create a new layout for the bottom portion which will be for choosing your character
+    bottomLayout = cmds.rowColumnLayout(nc = 2, cw = [(1, 100), (2, 200)], parent = mainLayout)
+    
+    
+    #create a label for choosing the character
+    charLabel = cmds.text(label = "Choose Weapon: ",align = 'left', parent = bottomLayout)
+    
+    #create the option menu for the characters
+    charOptionMenu = cmds.optionMenu('weaponOptionMenu', w = 200)
+    
+    #grab the selected project
+    selectedProject = cmds.optionMenu('WeaponLockerProjectOptionMenu', q = True, v = True)
+    
+    newPath += selectedProject
+    weapons = os.listdir(newPath)
+    
+    
+    #for each weapon, add the weapon to the option menu
+    for weapon in weapons:
+        #add a combo box entry
+        weaponName = weapon.partition(".")[0]
+        cmds.menuItem( label= weaponName)
+    
+    
+    #add a separator
+    cmds.separator(h =10,style='out' )
+    cmds.separator(h =10,style='out' )
+    
+    
+    #and lastly, create the layout for the button to build.
+    buttonLayout = cmds.rowColumnLayout(nc = 2, cw = [(1, 100), (2, 200)], parent = mainLayout)
+    
+    #create the button
+    cmds.text(label = "")
+    cmds.button(label = "Import", h = 50, c = partial(importWeapon, namespace))
+    
+    #add a separator
+    cmds.separator(h =10,style='out' )
+    cmds.separator(h =10,style='out' )
+    
+    
+    #show the window
+    cmds.showWindow(window)
+    
+    
+def updateList(*args):
+
+    #find all of the existing characters in the character option meny
+    existingWeapons = cmds.optionMenu('weaponOptionMenu', q = True, itemListShort = True)
+    
+    #for each character in there, delete it out from the option menu 
+    for each in existingWeapons:
+        cmds.deleteUI(each, menuItem = True)
+    
+    #now find all of the characters in the selected project folder
+    #open the autoRigConfig file and find the models path
+    configFile = cmds.internalVar(usd = True)
+    configFilePath = (configFile + "autoRigConfig.txt")
+        
+    #open file, get lines, find peoject value
+    f = open(configFilePath, "r")
+    lines=f.readlines()
+    path =  lines[1]
+    f.close()
+    
+    pathName = path.rpartition("/")[0]
+    newPath = pathName.rpartition("/")[0]
+    newPath += "/Weapons/"
+    projects = os.listdir(newPath)
+    
+    
+    selectedProject = cmds.optionMenu('WeaponLockerProjectOptionMenu', q = True, v = True)
+    newPath += selectedProject
+    weapons = os.listdir(newPath)
+    
+    
+    #for each character, add the character to the option menu
+    for weapon in weapons:
+        #add a combo box entry
+        weaponName = weapon.partition(".")[0]
+        cmds.menuItem( label= weaponName)
+        
+        
+        
+def importWeapon(namespace, *args):
+    
+
+        
+    #open the autoRigConfig file and find the models path
+    configFile = cmds.internalVar(usd = True)
+    configFilePath = (configFile + "autoRigConfig.txt")
+        
+    #open file, get lines, find peoject value
+    f = open(configFilePath, "r")
+    lines=f.readlines()
+    path =  lines[1]
+    f.close()
+    
+    
+    pathName = path.rpartition("/")[0]
+    newPath = pathName.rpartition("/")[0]
+    newPath += "/Weapons/"
+    projects = os.listdir(newPath)
+    
+    selectedProject = cmds.optionMenu('WeaponLockerProjectOptionMenu', q = True, v = True)
+    newPath += selectedProject
+    
+    # 2: get the selected character
+    weapon = cmds.optionMenu('weaponOptionMenu', q = True, v = True)
+    
+    # 3: build the full path to that model
+    newPath += ("/" + weapon + ".mb")
+    
+
+    #now we need to create a reference with that model in the scene
+    cmds.file(newPath, r = True, namespace = weapon, loadReferenceDepth = "all", gl = True)
+    
+    #lastly, need to constrain the weapon root to the right hand weapon control
+    
+    names = cmds.namespaceInfo( listNamespace=True )
+    number = findHighestTrailingNumber(names, weapon)
+
+    if number == 0:
+        weapon = weapon
+    else:
+        weapon = (weapon + str(number))
+        
+        
+    if weapon.find("Lancer") == 0:
+        cmds.parentConstraint((namespace + ":" + "R_Weapon_Anim"), (weapon + ":" + "b_AR_Root"))
+
+    if weapon.find("Bat") == 0:
+        cmds.parentConstraint((namespace + ":" + "R_Weapon_Anim"), (weapon + ":" + "b_AR_Root"))
+        
+        
+    if weapon.find("LocustSawnOff") == 0:
+        cmds.parentConstraint((namespace + ":" + "R_Weapon_Anim"), (weapon + ":" + "b_Shotgun_Root"))
+        
+    if weapon.find("Shotgun") == 0:
+        cmds.parentConstraint((namespace + ":" + "R_Weapon_Anim"), (weapon + ":" + "b_AR_Root"))
+        
+    if weapon.find("Pistol") == 0:
+        cmds.parentConstraint((namespace + ":" + "R_Weapon_Anim"), (weapon + ":" + "b_AR_Root"))
+        
+    if weapon.find("One_Handed") == 0:
+        cmds.parentConstraint((namespace + ":" + "R_Weapon_Anim"), (weapon + ":" + "b_AR_Root"))
+    
+    if weapon.find("Two_Handed") == 0:
+        cmds.parentConstraint((namespace + ":" + "R_Weapon_Anim"), (weapon + ":" + "b_AR_Root"))
+        
+    if weapon.find("Hammer") == 0:
+        cmds.parentConstraint((namespace + ":" + "R_Weapon_Anim"), (weapon + ":" + "b_AR_Root"))
+    
+    
+    
+def findHighestTrailingNumber(names, basename):
+    import re
+    
+    highestValue = 0
+
+    for n in names:
+        #this means we find our basename at the start of the string (0)
+        if n.find(basename) == 0:
+            #this is the value
+            suffix = n.partition(basename)[2]  
+
+            #this is checking that the end of the value only contains digits.
+            if re.match("[0-9]", suffix):    
+                numericalElement = int(suffix)
+
+                if numericalElement > highestValue:
+                    highestValue = numericalElement
+                #the part above keeps tracking of the highest number of the user specified instances. if it finds 4 as the highest val
+                #out of 4 instances of the same module (fingers for example) then it knows the next instance will need to be 5
+
+    return highestValue
+    
+def addMocap(*args):
+    
+    #import in the fbx file and set it up
+    cmds.unloadPlugin(  "fbxmaya.mll" )
+    cmds.loadPlugin(  "fbxmaya.mll" )
+    
+    
+    #grab the namespace of the character
+    namespace = cmds.optionMenu("selectedCharacterOptionMenu", q = True, v = True)
+    namespace += ":"
+    
+
+    
+    #zero out the rig
+    resetAll()
+    
+
+    #set the modes to FK
+    cmds.setAttr((namespace + "Settings.L_Arm_Mode"), 0)
+    cmds.setAttr((namespace + "Settings.R_Arm_Mode"), 0)
+    cmds.setAttr((namespace + "Settings.L_Leg_Mode"), 0)
+    cmds.setAttr((namespace + "Settings.R_Leg_Mode"), 0)
+    
+    
+    #duplicate the root
+    newSkeleton = cmds.duplicate((namespace + "b_MF_Root"),rr = True,  name = "b_MF_Root", inputConnections = False, upstreamNodes = False)
+    
+    #find out if character is referenced
+    referenced = cmds.referenceQuery(namespace + "Master_Anim", isNodeReferenced = True)
+    
+    if not referenced:
+        
+        cmds.lockNode((namespace + "Rig_Container"), lock = False)
+        cmds.container((namespace + "Rig_Container"),edit = True, removeNode = newSkeleton)
+    
+    #if the skeleton is not parented to the world already, then parent it to the world
+    #first, need to find the parent of the root bone.
+    parent = cmds.listRelatives("b_MF_Root", parent = True)
+    
+    if parent != None:
+        cmds.parent("b_MF_Root", world = True)
+        
+
+    #Constrain the rig controls to the skeleton
+    #maybe eventually get mods here. click == IK, shift click == FK
+    constraints = []
+    
+    #Core
+    bodyConstraint = cmds.parentConstraint("b_MF_Pelvis", (namespace + "Body_Anim"), mo = True)[0]
+    constraints.append(bodyConstraint)
+    
+    spine1Constraint = cmds.orientConstraint("b_MF_Spine_01", (namespace + "Spine_01_Anim"), mo = True)[0]
+    constraints.append(spine1Constraint)
+    
+    spine2Constraint = cmds.orientConstraint("b_MF_Spine_02", (namespace + "Spine_02_Anim"), mo = True)[0]
+    constraints.append(spine2Constraint)
+    
+    spine3Constraint = cmds.orientConstraint("b_MF_Spine_03", (namespace + "Spine_03_Anim"), mo = True)[0]
+    constraints.append(spine3Constraint)\
+    
+    neckConstraint = cmds.orientConstraint("b_MF_Neck", (namespace + "Neck_Anim"), mo = True)[0]
+    constraints.append(neckConstraint)
+    
+    headConstraint = cmds.orientConstraint("b_MF_Head", (namespace + "Head_Anim"), mo = True)[0]
+    constraints.append(headConstraint)
+    
+    #Arms
+    lShoulderConstraint = cmds.orientConstraint("b_MF_Clavicle_L", (namespace + "L_Shoulder_Anim"), mo = True)[0]
+    constraints.append(lShoulderConstraint)
+    
+    rShoulderConstraint = cmds.orientConstraint("b_MF_Clavicle_R", (namespace + "R_Shoulder_Anim"), mo = True)[0]
+    constraints.append(rShoulderConstraint)
+    
+
+    lFKArmConstraint = cmds.orientConstraint("b_MF_UpperArm_L", (namespace + "FK_b_MF_UpperArm_L"), mo = True)[0]
+    constraints.append(lFKArmConstraint)
+    
+    lFKLowArmConstraint = cmds.orientConstraint("b_MF_Forearm_L", (namespace + "FK_b_MF_Forearm_L"), mo = True)[0]
+    constraints.append(lFKLowArmConstraint)
+    
+    lFKHandConstraint = cmds.orientConstraint("b_MF_Hand_L", (namespace + "FK_b_MF_Hand_L"), mo = True)[0]
+    constraints.append(lFKHandConstraint)
+    
+
+    rFKArmConstraint = cmds.orientConstraint("b_MF_UpperArm_R", (namespace + "FK_b_MF_UpperArm_R"), mo = True)[0]
+    constraints.append(rFKArmConstraint)
+    
+    rFKLowArmConstraint = cmds.orientConstraint("b_MF_Forearm_R", (namespace + "FK_b_MF_Forearm_R"), mo = True)[0]
+    constraints.append(rFKLowArmConstraint)
+    
+    rFKHandConstraint = cmds.orientConstraint("b_MF_Hand_R", (namespace + "FK_b_MF_Hand_R"), mo = True)[0]
+    constraints.append(rFKHandConstraint)
+    
+    
+    #Legs
+
+    lFKUpLegConstraint = cmds.orientConstraint("b_MF_Thigh_L", (namespace + "FK_b_MF_Thigh_L"), mo = True)[0]
+    constraints.append(lFKUpLegConstraint)
+    
+    lFKLoLegConstraint = cmds.orientConstraint("b_MF_Calf_L", (namespace + "FK_b_MF_Calf_L"), mo = True)[0]
+    constraints.append(lFKLoLegConstraint)
+    
+    lFKFootConstraint = cmds.orientConstraint("b_MF_Foot_L", (namespace + "FK_b_MF_Foot_L"), mo = True)[0]
+    constraints.append(lFKFootConstraint)
+    
+    lFKToeConstraint = cmds.orientConstraint("b_MF_Toe_L", (namespace + "FK_b_MF_Toe_L"), mo = True)[0]
+    constraints.append(lFKToeConstraint)
+    
+    
+    rFKUpLegConstraint = cmds.orientConstraint("b_MF_Thigh_R", (namespace + "FK_b_MF_Thigh_R"), mo = True)[0]
+    constraints.append(rFKUpLegConstraint)
+    
+    rFKLoLegConstraint = cmds.orientConstraint("b_MF_Calf_R", (namespace + "FK_b_MF_Calf_R"), mo = True)[0]
+    constraints.append(rFKLoLegConstraint)
+    
+    rFKFootConstraint = cmds.orientConstraint("b_MF_Foot_R", (namespace + "FK_b_MF_Foot_R"), mo = True)[0]
+    constraints.append(rFKFootConstraint)
+    
+    rFKToeConstraint = cmds.orientConstraint("b_MF_Toe_R", (namespace + "FK_b_MF_Toe_R"), mo = True)[0]
+    constraints.append(rFKToeConstraint)
+    
+    
+    #Fingers -Left
+    index1Constraint = cmds.orientConstraint("b_MF_Finger_01_L", (namespace + "L_Index_1_Anim"), mo = True)[0]
+    constraints.append(index1Constraint)
+    
+    index2Constraint = cmds.orientConstraint("b_MF_Finger_02_L", (namespace + "L_Index_2_Anim"), mo = True)[0]
+    constraints.append(index2Constraint)
+    
+    index3Constraint = cmds.orientConstraint("b_MF_Finger_03_L", (namespace + "L_Index_3_Anim"), mo = True)[0]
+    constraints.append(index3Constraint)
+    
+    
+    middle1Constraint = cmds.orientConstraint("b_MF_Finger_04_L", (namespace + "L_Middle_1_Anim"), mo = True)[0]
+    constraints.append(middle1Constraint)
+    
+    middle2Constraint = cmds.orientConstraint("b_MF_Finger_05_L", (namespace + "L_Middle_2_Anim"), mo = True)[0]
+    constraints.append(middle2Constraint)
+    
+    middle3Constraint = cmds.orientConstraint("b_MF_Finger_06_L", (namespace + "L_Middle_3_Anim"), mo = True)[0]
+    constraints.append(middle3Constraint)
+    
+    
+    ring1Constraint = cmds.orientConstraint("b_MF_Finger_07_L", (namespace + "L_Ring_1_Anim"), mo = True)[0]
+    constraints.append(ring1Constraint)
+    
+    ring2Constraint = cmds.orientConstraint("b_MF_Finger_08_L", (namespace + "L_Ring_2_Anim"), mo = True)[0]
+    constraints.append(ring2Constraint)
+    
+    ring3Constraint = cmds.orientConstraint("b_MF_Finger_09_L", (namespace + "L_Ring_3_Anim"), mo = True)[0]
+    constraints.append(ring3Constraint)
+    
+    
+    pinky1Constraint = cmds.orientConstraint("b_MF_Finger_10_L", (namespace + "L_Pinky_1_Anim"), mo = True)[0]
+    constraints.append(pinky1Constraint)
+    
+    pinky2Constraint = cmds.orientConstraint("b_MF_Finger_11_L", (namespace + "L_Pinky_2_Anim"), mo = True)[0]
+    constraints.append(pinky2Constraint)
+    
+    pinky3Constraint = cmds.orientConstraint("b_MF_Finger_12_L", (namespace + "L_Pinky_3_Anim"), mo = True)[0]
+    constraints.append(pinky3Constraint)
+    
+    
+    thumb1Constraint = cmds.orientConstraint("b_MF_Thumb_01_L", (namespace + "L_Thumb_1_Anim"), mo = True)[0]
+    constraints.append(thumb1Constraint)
+    
+    thumb2Constraint = cmds.orientConstraint("b_MF_Thumb_02_L", (namespace + "L_Thumb_2_Anim"), mo = True)[0]
+    constraints.append(thumb2Constraint)
+    
+    thumb3Constraint = cmds.orientConstraint("b_MF_Thumb_03_L", (namespace + "L_Thumb_3_Anim"), mo = True)[0]
+    constraints.append(thumb3Constraint)
+    
+    
+    #Fingers -Right
+    rindex1Constraint = cmds.orientConstraint("b_MF_Finger_01_R", (namespace + "R_Index_1_Anim"), mo = True)[0]
+    constraints.append(rindex1Constraint)
+    
+    rindex2Constraint = cmds.orientConstraint("b_MF_Finger_02_R", (namespace + "R_Index_2_Anim"), mo = True)[0]
+    constraints.append(rindex2Constraint)
+    
+    rindex3Constraint = cmds.orientConstraint("b_MF_Finger_03_R", (namespace + "R_Index_3_Anim"), mo = True)[0]
+    constraints.append(rindex3Constraint)
+    
+    
+    rmiddle1Constraint = cmds.orientConstraint("b_MF_Finger_04_R", (namespace + "R_Middle_1_Anim"), mo = True)[0]
+    constraints.append(rmiddle1Constraint)
+    
+    rmiddle2Constraint = cmds.orientConstraint("b_MF_Finger_05_R", (namespace + "R_Middle_2_Anim"), mo = True)[0]
+    constraints.append(rmiddle2Constraint)
+    
+    rmiddle3Constraint = cmds.orientConstraint("b_MF_Finger_06_R", (namespace + "R_Middle_3_Anim"), mo = True)[0]
+    constraints.append(rmiddle3Constraint)
+    
+    
+    rring1Constraint = cmds.orientConstraint("b_MF_Finger_07_R", (namespace + "R_Ring_1_Anim"), mo = True)[0]
+    constraints.append(rring1Constraint)
+    
+    rring2Constraint = cmds.orientConstraint("b_MF_Finger_08_R", (namespace + "R_Ring_2_Anim"), mo = True)[0]
+    constraints.append(rring2Constraint)
+    
+    rring3Constraint = cmds.orientConstraint("b_MF_Finger_09_R", (namespace + "R_Ring_3_Anim"), mo = True)[0]
+    constraints.append(rring3Constraint)
+    
+    
+    rpinky1Constraint = cmds.orientConstraint("b_MF_Finger_10_R", (namespace + "R_Pinky_1_Anim"), mo = True)[0]
+    constraints.append(rpinky1Constraint)
+    
+    rpinky2Constraint = cmds.orientConstraint("b_MF_Finger_11_R", (namespace + "R_Pinky_2_Anim"), mo = True)[0]
+    constraints.append(rpinky2Constraint)
+    
+    rpinky3Constraint = cmds.orientConstraint("b_MF_Finger_12_R", (namespace + "R_Pinky_3_Anim"), mo = True)[0]
+    constraints.append(rpinky3Constraint)
+    
+    
+    rthumb1Constraint = cmds.orientConstraint("b_MF_Thumb_01_R", (namespace + "R_Thumb_1_Anim"), mo = True)[0]
+    constraints.append(rthumb1Constraint)
+    
+    rthumb2Constraint = cmds.orientConstraint("b_MF_Thumb_02_R", (namespace + "R_Thumb_2_Anim"), mo = True)[0]
+    constraints.append(rthumb2Constraint)
+    
+    rthumb3Constraint = cmds.orientConstraint("b_MF_Thumb_03_R", (namespace + "R_Thumb_3_Anim"), mo = True)[0]
+    constraints.append(rthumb3Constraint)
+    
+    
+    #create locators for the weapons as well    
+    lWeaponLoc = cmds.spaceLocator()[0]
+    lWeaponLoc = cmds.rename(lWeaponLoc, (namespace + ":L_Weapon_Loc"))
+    
+    rWeaponLoc = cmds.spaceLocator()[0]
+    rWeaponLoc = cmds.rename(rWeaponLoc, (namespace + ":R_Weapon_Loc"))
+    
+    cmds.parent(lWeaponLoc, (namespace + "Master_Anim"))
+    cmds.parent(rWeaponLoc, (namespace + "Master_Anim"))
+    
+    
+    lwepanimConstraint = cmds.parentConstraint("b_MF_Weapon_L", lWeaponLoc, mo = False)[0]
+    constraints.append(lwepanimConstraint)
+    
+    rwepanimConstraint = cmds.parentConstraint("b_MF_Weapon_R", rWeaponLoc, mo = False)[0]
+    constraints.append(rwepanimConstraint)
+    
+    
+
+    
+    #Open up a file browser dialog, and have the user choose a file
+    filePath = cmds.fileDialog(directoryMask = "*.fbx")
+
+
+    #import the fbx, which the animation will go onto our new skeleton
+    string = "FBXImportScaleFactor 1.0;"
+    string += "FBXImportMode -v \"exmerge\";"
+    string += "FBXImportSetLockedAttribute -v true;"
+    mel.eval(string)
+    
+    #string = ("FBXImport -f /"" + str(filePath) +  /"";")
+    string = ("FBXImport -f \"%s\" ;" %filePath)
+    mel.eval(string)
+    
+    #set time range to fbx range
+    cmds.select("b_MF_Root", r = True, hi = True)
+    cmds.currentTime(0)
+    
+    timeStart = cmds.findKeyframe("b_MF_Pelvis", which = 'first')
+    timeEnd = cmds.findKeyframe("b_MF_Pelvis", which = 'last')
+    cmds.playbackOptions(min = timeStart, max = timeEnd)
+
+    
+    
+    #bake down the animation onto the controls
+    selectAll()
+    cmds.select((namespace + "L_Index_1_Anim"),add = True)
+    cmds.select((namespace + "L_Index_2_Anim"),add = True)
+    cmds.select((namespace + "L_Index_3_Anim"),add = True)
+    
+    cmds.select((namespace + "L_Middle_1_Anim"),add = True)
+    cmds.select((namespace + "L_Middle_2_Anim"),add = True)
+    cmds.select((namespace + "L_Middle_3_Anim"),add = True)
+    
+    cmds.select((namespace + "L_Ring_1_Anim"),add = True)
+    cmds.select((namespace + "L_Ring_2_Anim"),add = True)
+    cmds.select((namespace + "L_Ring_3_Anim"),add = True)
+    
+    cmds.select((namespace + "L_Pinky_1_Anim"),add = True)
+    cmds.select((namespace + "L_Pinky_2_Anim"),add = True)
+    cmds.select((namespace + "L_Pinky_3_Anim"),add = True)
+    
+    cmds.select((namespace + "L_Thumb_1_Anim"),add = True)
+    cmds.select((namespace + "L_Thumb_2_Anim"),add = True)
+    cmds.select((namespace + "L_Thumb_3_Anim"),add = True)
+    
+    cmds.select((namespace + "R_Index_1_Anim"),add = True)
+    cmds.select((namespace + "R_Index_2_Anim"),add = True)
+    cmds.select((namespace + "R_Index_3_Anim"),add = True)
+    
+    cmds.select((namespace + "R_Middle_1_Anim"),add = True)
+    cmds.select((namespace + "R_Middle_2_Anim"),add = True)
+    cmds.select((namespace + "R_Middle_3_Anim"),add = True)
+    
+    cmds.select((namespace + "R_Ring_1_Anim"),add = True)
+    cmds.select((namespace + "R_Ring_2_Anim"),add = True)
+    cmds.select((namespace + "R_Ring_3_Anim"),add = True)
+    
+    cmds.select((namespace + "R_Pinky_1_Anim"),add = True)
+    cmds.select((namespace + "R_Pinky_2_Anim"),add = True)
+    cmds.select((namespace + "R_Pinky_3_Anim"),add = True)
+    
+    cmds.select((namespace + "R_Thumb_1_Anim"),add = True)
+    cmds.select((namespace + "R_Thumb_2_Anim"),add = True)
+    cmds.select((namespace + "R_Thumb_3_Anim"),add = True)
+    
+    cmds.select(lWeaponLoc, add = True)
+    cmds.select(rWeaponLoc, add = True)
+    
+    
+    start = cmds.playbackOptions(q = True, min = True)
+    end = cmds.playbackOptions(q = True, max = True)
+    
+    cmds.bakeResults(simulation = True, t = (start, end))
+    
+    #delete the skeleton and constraints
+    cmds.delete(constraints)
+    cmds.delete("b_MF_Root")
+    
+    
+    #get rid of any dirty nodes
+    containers = cmds.ls(type = 'container')
+    
+    for container in containers:
+        children = cmds.container(container, q = True, nodeList = True)
+        if children == None:
+            cmds.delete(container)
+            
+            
+    #set the geo layer to be referenced
+    #cmds.setAttr((namespace + "*Geo.displayType"), 2)
+    
+    
+
+
+
+    
+    
+
+
+def export(*args):
+    
+    #grab the namespace of the character
+    namespace = cmds.optionMenu("selectedCharacterOptionMenu", q = True, v = True)
+    namespace += ":"
+    
+    
+    #find out if the namespace is from a referenced file
+    refNode = cmds.referenceQuery( (namespace + "NonRigItems"), isNodeReferenced=True )
+    
+    
+    #duplicate the root
+    newSkeleton = cmds.duplicate((namespace + "b_MF_Root"),rr = True,  name = "b_MF_Root", inputConnections = False, upstreamNodes = False)
+    
+    
+    if refNode != True:
+
+        cmds.lockNode((namespace + "Rig_Container"), lock = False)
+        cmds.container((namespace + "Rig_Container"),edit = True, removeNode = newSkeleton)
+    
+    
+    #if the skeleton is not parented to the world already, then parent it to the world
+    #first, need to find the parent of the root bone.
+    parent = cmds.listRelatives("b_MF_Root", parent = True)
+    
+    if parent != None:
+        cmds.parent("b_MF_Root", world = True)
+    
+    
+    
+    #make a list of the bones to export
+    exportBones = []
+    
+    #find the joints that have the Export label
+    joints = cmds.ls(type = "joint")
+    
+    for joint in joints:
+        attr = cmds.getAttr(joint + ".type")
+        
+        if attr == 18:
+            attrName = cmds.getAttr(joint + ".otherType")
+            
+            if attrName == "Export":
+                jointNoNamespace = joint.partition(":")[2]
+                exportBones.append(str(jointNoNamespace))
+           
+    exportBones = filter(None, exportBones) 
+    print exportBones
+            
+    if len(exportBones) == 0:
+        result = cmds.confirmDialog( title='Export Joints', message='There appears to be no joints in your scene designated for export. Would you like to use the defaults?', button=['Yes','No'], defaultButton='No', cancelButton='Yes', dismissString='Yes' )
+        if result == "No":
+            selectExportJointsUI()
+            cmds.delete(newSkeleton)
+            return
+            
+            
+        #USE DEFAULTS
+        else:
+            
+            #use defaults
+            exportBones.append(("b_MF_Root"))
+            exportBones.append(("b_MF_Camera_Root"))
+            exportBones.append(("b_MF_Camera"))
+            exportBones.append(("b_MF_Pelvis"))
+                
+            exportBones.append(("b_MF_Thigh_R"))
+            exportBones.append(("b_MF_Calf_R"))
+            exportBones.append(("b_MF_Foot_R"))
+            exportBones.append(("b_MF_Toe_R"))
+            
+            exportBones.append(("b_MF_Thigh_L"))
+            exportBones.append(("b_MF_Calf_L"))
+            exportBones.append(("b_MF_Foot_L"))
+            exportBones.append(("b_MF_Toe_L"))
+            
+            exportBones.append(("b_MF_Spine_01"))
+            exportBones.append(("b_MF_Spine_02"))
+            exportBones.append(("b_MF_Spine_03"))
+            
+            exportBones.append(("b_MF_Neck"))
+            exportBones.append(("b_MF_Head"))
+            
+            exportBones.append(("b_MF_Clavicle_L"))
+            exportBones.append(("b_MF_UpperArm_L"))
+            exportBones.append(("b_MF_Armroll_L"))
+            exportBones.append(("b_MF_Forearm_L"))
+            exportBones.append(("b_MF_Hand_L"))
+            exportBones.append(("b_MF_ForeTwist2_L"))
+            
+            exportBones.append(("b_MF_Clavicle_R"))
+            exportBones.append(("b_MF_UpperArm_R"))
+            exportBones.append(("b_MF_Armroll_R"))
+            exportBones.append(("b_MF_Forearm_R"))
+            exportBones.append(("b_MF_Hand_R"))
+            exportBones.append(("b_MF_ForeTwist2_R"))
+            
+            exportBones.append(("b_MF_Thumb_01_L"))
+            exportBones.append(("b_MF_Thumb_02_L"))
+            exportBones.append(("b_MF_Thumb_03_L"))
+            
+            exportBones.append(("b_MF_Finger_01_L"))
+            exportBones.append(("b_MF_Finger_02_L"))
+            exportBones.append(("b_MF_Finger_03_L"))
+            
+            exportBones.append(("b_MF_Finger_04_L"))
+            exportBones.append(("b_MF_Finger_05_L"))
+            exportBones.append(("b_MF_Finger_06_L"))
+            
+            exportBones.append(("b_MF_Finger_07_L"))
+            exportBones.append(("b_MF_Finger_08_L"))
+            exportBones.append(("b_MF_Finger_09_L"))
+            
+            exportBones.append(("b_MF_Finger_10_L"))
+            exportBones.append(("b_MF_Finger_11_L"))
+            exportBones.append(("b_MF_Finger_12_L"))
+            
+            exportBones.append(("b_MF_Thumb_01_R"))
+            exportBones.append(("b_MF_Thumb_02_R"))
+            exportBones.append(("b_MF_Thumb_03_R"))
+            
+            exportBones.append(("b_MF_Finger_01_R"))
+            exportBones.append(("b_MF_Finger_02_R"))
+            exportBones.append(("b_MF_Finger_03_R"))
+            
+            exportBones.append(("b_MF_Finger_04_R"))
+            exportBones.append(("b_MF_Finger_05_R"))
+            exportBones.append(("b_MF_Finger_06_R"))
+            
+            exportBones.append(("b_MF_Finger_07_R"))
+            exportBones.append(("b_MF_Finger_08_R"))
+            exportBones.append(("b_MF_Finger_09_R"))
+            
+            exportBones.append(("b_MF_Finger_10_R"))
+            exportBones.append(("b_MF_Finger_11_R"))
+            exportBones.append(("b_MF_Finger_12_R"))
+            
+            exportBones.append(("b_MF_Weapon_L"))
+            exportBones.append(("b_MF_Weapon_L_End"))
+            
+            exportBones.append(("b_MF_Weapon_R"))
+            exportBones.append(("b_MF_Weapon_R_End"))
+            
+            exportBones.append(("b_MF_BackMount_L"))
+            exportBones.append(("b_MF_BackMount_R"))
+            
+            exportBones.append(("b_MF_IK_Foot_Root"))
+            exportBones.append(("b_MF_IK_Foot_L"))
+            exportBones.append(("b_MF_IK_Foot_R"))
+            
+            exportBones.append(("b_MF_IK_Hand_Root"))
+            exportBones.append(("b_MF_IK_Gun"))
+            exportBones.append(("b_MF_IK_Hand_R"))
+            exportBones.append(("b_MF_IK_Hand_L"))
+            
+            exportBones.append(("b_MF_Attach"))
+    
+    
+
+       
+    #get the name of your opened file
+    openedFile = cmds.file(q = True, sceneName = True)
+    fileName = openedFile.rpartition(".")[0]
+    
+    
+    #for each bone in export bones, constraint to the corresponding bone
+    for bone in exportBones:
+        if len(bone) > 2:
+            cmds.parentConstraint((namespace + bone), bone)
+        
+        
+    #select the export bones
+    cmds.select(exportBones, replace = True)
+    
+ 
+    #get the key modifier(is shift held down or not) and if it is not, just do a fast export. if it is, then bring up a UI that shows fields for frame range, and name
+    mods = cmds.getModifiers()
+    
+    
+    #if you have shift held down:
+    if (mods & 1) > 0:
+
+        exportAnimUI(exportBones)
+
+        
+        
+        
+    #if you are just pressing the button
+    if (mods & 1) == 0:
+            
+
+        #first get your timeline range
+        originalStart = int(cmds.playbackOptions(q = True, min = True))
+        originalEnd = int(cmds.playbackOptions(q = True, max = True))
+        
+        #then set the anim range to the timeline range
+        cmds.playbackOptions(animationStartTime = originalStart)
+        cmds.playbackOptions(animationEndTime = originalEnd)
+        
+        #bake down the animation onto the bones
+        cmds.bakeResults(simulation = True, t = (originalStart, originalEnd))
+    
+    
+        #select the whole hierarchy, and if a bone in the selection, does not exist in export bones, delete it
+        cmds.select("b_MF_Root", hi = True, cc = True)
+    
+        selection = cmds.ls(sl = True)
+        
+        for joint in selection:
+            if joint not in exportBones:
+                if cmds.objExists(joint):
+                    cmds.delete(joint)
+                    
+        cmds.select("b_MF_Root", hi = True, cc = True)
+        selection = cmds.ls(sl = True, exactType = 'joint')
+        
+        #export the file out as an fbx
+        string = "FBXResetExport;"
+        string += "FBXExportConstraints -v 0;"
+        string += "FBXExportScaleFactor 1.0;"
+        string += "FBXExportBakeComplexAnimation -v 1;"
+        string += "FBXExportAnimationOnly -v 0;"
+        string += "FBXExportReferencedContainersContent -v 0;"
+        string += ("FBXExportBakeComplexStart -v %d ;" %originalStart)
+        string += ("FBXExportBakeComplexEnd -v %d ;" %originalEnd)
+        mel.eval(string)
+        
+        
+        #export out selected to path
+        cmds.file(fileName + ".fbx", type = "FBX export", pr = True, es = True, force = True, prompt = False)
+        
+        
+        #no matter what, delete the export skeleton, and relock the containers
+        cmds.delete("b_MF_Root")
+        cmds.lockNode((namespace + "Rig_Container"), lock = False)
+        
+        #set the time range back to the original
+        cmds.playbackOptions(min = originalStart)
+        cmds.playbackOptions(max = originalEnd)
+
+def selectExportJointsUI(*args):
+
+    if cmds.window("selectExportJointsUI", exists = True):
+        cmds.deleteUI("selectExportJointsUI")
+    
+    window = cmds.window("selectExportJointsUI", title = "Select Export Joints", w = 300, h = 600, sizeable = False)
+    
+    mainLayout = cmds.columnLayout(w = 300, h = 600)
+    
+    #add a big textScrollist
+    textScrollList = cmds.textScrollList("textScrollList_joints", w = 295, h = 550,numberOfRows=1000, allowMultiSelection=True)
+    cmds.button(label = "Select Export Joints", w = 295, h = 50, c = getJoints)
+    
+    cmds.showWindow(window)
+    
+    #grab joints, add to textScrollList
+    cmds.select("*:b_MF_Root", hi = True)
+    allJoints = cmds.ls(sl = True, type = 'joint')
+    cmds.select(clear = True)
+    
+    for joint in allJoints:
+        cmds.textScrollList("textScrollList_joints", edit = True, append = str(joint))
+        
+        
+def getExportJoints_UI(*args):
+    
+    if cmds.window("getExportJointsUI", exists = True):
+        cmds.deleteUI("getExportJointsUI")
+    
+    window = cmds.window("getExportJointsUI", title = "Current Export Joints", w = 300, h = 600, sizeable = False)
+    
+    mainLayout = cmds.columnLayout(w = 300, h = 600)
+    
+
+    #add a big textScrollist
+    textScrollList = cmds.textScrollList("textScrollList_exportjoints", w = 295, h = 550,numberOfRows=1000, allowMultiSelection=True)
+    cmds.button(label = "Remove Selected", w = 295, h = 50, c = removeSelected)
+    cmds.button(label = "Close", w = 295, h = 50, c = getExportJoints_Close)
+    
+    cmds.showWindow(window)
+    
+    getExportJoints()
+    
+def removeSelected(*args):
+    namespace = cmds.optionMenu("selectedCharacterOptionMenu", q = True, v = True)
+    selected = cmds.textScrollList("textScrollList_exportjoints", q = True, si = True)
+    
+    for each in selected:
+        cmds.textScrollList("textScrollList_exportjoints", edit = True, removeItem = each)
+        
+        #remove the label
+        cmds.setAttr(namespace + ":" + each + ".type", 0)
+        cmds.setAttr(namespace + ":" + each + ".otherType", "None", type = "string")
+        
+        
+        
+def getExportJoints():
+    
+    joints = cmds.ls(type = "joint")
+    
+    exportJoints = []
+    
+    for joint in joints:
+        attr = cmds.getAttr(joint + ".type")
+        
+        if attr == 18:
+            jointNoNamespace = joint.partition(":")[2]
+            exportJoints.append(jointNoNamespace)
+            
+    for joint in exportJoints:
+        cmds.textScrollList("textScrollList_exportjoints", edit = True, append = str(joint))
+        
+def getExportJoints_Close(*args):
+    cmds.deleteUI("getExportJointsUI")
+
+
+def getJoints(*args):
+    
+    exportJoints = cmds.textScrollList("textScrollList_joints", q = True, si = True)
+    print exportJoints
+    
+    #each = cmds.ls(sl = True)[0]
+    for joint in exportJoints:
+        attr = cmds.getAttr(joint + ".type")
+        if attr == 0:
+            cmds.setAttr(joint + ".type", 18)
+            cmds.setAttr(joint + ".otherType", "Export", type = "string")
+    cmds.deleteUI("selectExportJointsUI")
+    
+    
+    
+def exportAnimUI(exportBones, *args):
+        
+    
+    #grab the namespace of the character
+    namespace = cmds.optionMenu("selectedCharacterOptionMenu", q = True, v = True)
+    namespace += ":"
+    
+    exportBones = exportBones
+    
+    #if window exists, kill it
+    if cmds.window("exportAnimation_UI_window", exists = True):
+        cmds.deleteUI("exportAnimation_UI_window")
+            
+            
+    windowWidth = 300
+    windowHeight = 152
+    #create the window
+    window = cmds.window("exportAnimation_UI_window", width = windowWidth, height = windowHeight, title = "Export Animation Options", sizeable = False)
+        
+    #create a main layout
+    topLevelCol = cmds.columnLayout(adj = True, columnAlign = "center", rs = 3)
+        
+    #create a row column layout for text labels and text fields
+    animName_rowCol = cmds.rowColumnLayout(nc = 2, columnAttach = (1, "right", 0), columnWidth = [(1, 90), (2, windowWidth - 100)], parent = topLevelCol )
+        
+        
+    #text and text Field for anim name
+    cmds.text(label = "Anim Name :", parent = animName_rowCol)
+    animName = cmds.textField("animName", text = "([a-z] [A-Z] [0-9] and _ only)", parent = animName_rowCol)
+        
+    #text and text Field for anim frame range
+    cmds.text(label = "Frame Range :", parent = animName_rowCol)
+    frameRange = cmds.textField("frameRange", text = "#-#", parent = animName_rowCol)
+        
+
+    #Now make a layout for the 2 buttons (accept, cancel)
+    columnWidth = (windowWidth/2)- 5
+        
+    button_row = cmds.rowLayout(nc = 2, columnWidth = [(1, columnWidth), (2, columnWidth)], cat = [ (1, "both", 10), (2, "both", 10) ], columnAlign = [ (1, "center"), (2, "center") ], parent = topLevelCol )
+        
+    #make the buttons
+    cmds.button(label = "Accept",c = partial(exportAnimUI_Accept,exportBones))
+    cmds.button(label = "Cancel", c = exportAnimUI_Cancel)
+    
+        
+    #show the window
+    cmds.showWindow(window)
+
+def exportAnimUI_Accept(exportBones, *args):
+    
+    #grab the namespace of the character
+    namespace = cmds.optionMenu("selectedCharacterOptionMenu", q = True, v = True)
+    namespace += ":"
+    
+    
+    #get the name of your opened file
+    openedFile = cmds.file(q = True, sceneName = True)
+    file = openedFile.rpartition("/")[0] + "/"
+    
+    #first, grab the name
+    name = cmds.textField("animName", q = True, text = True)
+    
+    #then grab the frame range
+    range = cmds.textField("frameRange", q = True, text = True)
+    start = int(range.partition("-")[0])
+    end = int(range.partition("-")[2])
+    
+    #now set the frame range in your scene file to that specified in the text field
+    
+    originalStart = cmds.playbackOptions(q = True, min = True)
+    originalEnd = cmds.playbackOptions(q = True, max = True)
+    
+    cmds.playbackOptions(animationStartTime = start)
+    cmds.playbackOptions(animationEndTime = end)
+
+
+    #bake down the animation onto the bones
+    cmds.bakeResults(simulation = True, t = (start, end))
+    
+    
+    #select the whole hierarchy, and if a bone in the selection, does not exist in export bones, delete it
+    cmds.select("b_MF_Root", hi = True, cc = True)
+    
+    selection = cmds.ls(sl = True)
+        
+    for joint in selection:
+        if joint not in exportBones:
+            if cmds.objExists(joint):
+                cmds.delete(joint)
+                
+    cmds.select("b_MF_Root", hi = True, cc = True)
+    selection = cmds.ls(sl = True, exactType = 'joint')
+
+    #export the file out as an fbx
+    string = "FBXResetExport;"
+    string += "FBXExportConstraints -v 0;"
+    string += "FBXExportScaleFactor 1.0;"
+    string += "FBXExportBakeComplexAnimation -v 1;"
+    string += "FBXExportAnimationOnly -v 0;"
+    string += "FBXExportReferencedContainersContent -v 0;"
+    
+    string += ("FBXExportBakeComplexStart -v %d ;" %start)
+    string += ("FBXExportBakeComplexEnd -v %d ;" %end)
+    
+    mel.eval(string)
+    cmds.file(file + name + ".fbx", type = "FBX export", pr = True, es = True, force = True, prompt = False)
+    
+        
+        
+    #no matter what, delete the export skeleton, and relock the containers
+    cmds.delete("b_MF_Root")
+    cmds.lockNode((namespace + "Rig_Container"), lock = False)
+    
+    cmds.deleteUI("exportAnimation_UI_window")
+    
+    #set the time range back to the original
+    cmds.playbackOptions(min = originalStart)
+    cmds.playbackOptions(max = originalEnd)
+    
+
+def exportAnimUI_Cancel(*args):
+    
+    #grab the namespace of the character
+    namespace = cmds.optionMenu("selectedCharacterOptionMenu", q = True, v = True)
+    namespace += ":"
+    
+    #first, delete the UI
+    cmds.deleteUI("exportAnimation_UI_window")
+    
+    #then delete the export skeleton and relock the containers:
+    cmds.delete("b_MF_Root")
+    cmds.lockNode((namespace + "Rig_Container"), lock = False)
+
+
+def resetSelection(*args):
+    
+    #grab the namespace of the character
+    namespace = cmds.optionMenu("selectedCharacterOptionMenu", q = True, v = True)
+    
+    #grab the current selection
+    selection = cmds.ls(sl = True)
+    
+    #for each control in the selection, zero out the attrs
+    for control in selection:
+        
+        attrs = cmds.listAttr(control, k = True, u = True)
+        
+        for attr in attrs:
+            if attr != "globalScale":
+                cmds.setAttr((control + "." + attr), 0)
+        
+    
+def resetAll(*args):
+    
+    #grab the namespace of the character
+    namespace = cmds.optionMenu("selectedCharacterOptionMenu", q = True, v = True)
+    
+    
+    #make a list of all of the controls
+    controls = []
+    
+    controls.append((namespace + ":" + "FK_b_MF_Thigh_R"))
+    controls.append((namespace + ":" + "FK_b_MF_Calf_R"))
+    controls.append((namespace + ":" + "FK_b_MF_Foot_R"))
+    controls.append((namespace + ":" + "FK_b_MF_Toe_R"))
+    controls.append((namespace + ":" + "Right_Knee_Anim"))
+    controls.append((namespace + ":" + "Right_Foot_Anim"))
+    controls.append((namespace + ":" + "FK_b_MF_Thigh_L"))
+    controls.append((namespace + ":" + "FK_b_MF_Calf_L"))
+    controls.append((namespace + ":" + "FK_b_MF_Foot_L"))
+    controls.append((namespace + ":" + "FK_b_MF_Toe_L"))
+    controls.append((namespace + ":" + "Left_Knee_Anim"))
+    controls.append((namespace + ":" + "Left_Foot_Anim"))
+    controls.append((namespace + ":" + "R_Shoulder_Anim"))
+    controls.append((namespace + ":" + "R_IK_Elbow_Anim"))
+    controls.append((namespace + ":" + "R_IK_Hand_Anim"))
+    controls.append((namespace + ":" + "FK_b_MF_UpperArm_R"))
+    controls.append((namespace + ":" + "FK_b_MF_Forearm_R"))
+    controls.append((namespace + ":" + "FK_b_MF_Hand_R"))
+    controls.append((namespace + ":" + "R_Weapon_Anim"))
+    controls.append((namespace + ":" + "L_Shoulder_Anim"))
+    controls.append((namespace + ":" + "L_IK_Elbow_Anim"))
+    controls.append((namespace + ":" + "L_IK_Hand_Anim"))
+    controls.append((namespace + ":" + "FK_b_MF_UpperArm_L"))
+    controls.append((namespace + ":" + "FK_b_MF_Forearm_L"))
+    controls.append((namespace + ":" + "FK_b_MF_Hand_L"))
+    controls.append((namespace + ":" + "L_Weapon_Anim"))
+    controls.append((namespace + ":" + "Neck_Anim"))
+    controls.append((namespace + ":" + "Head_Anim"))
+    controls.append((namespace + ":" + "Spine_01_Anim"))
+    controls.append((namespace + ":" + "Spine_02_Anim"))
+    controls.append((namespace + ":" + "Spine_03_Anim"))
+    controls.append((namespace + ":" + "Body_Anim"))
+    controls.append((namespace + ":" + "Hip_Anim"))
+    controls.append((namespace + ":" + "Master_Anim"))
+    controls.append((namespace + ":" + "Root_Anim"))
+    controls.append((namespace + ":" + "Attach_Anim"))
+    controls.append((namespace + ":" + "L_Index_1_Anim"))
+    controls.append((namespace + ":" + "L_Index_2_Anim"))
+    controls.append((namespace + ":" + "L_Index_3_Anim"))
+    controls.append((namespace + ":" + "L_Middle_1_Anim"))
+    controls.append((namespace + ":" + "L_Middle_2_Anim"))
+    controls.append((namespace + ":" + "L_Middle_3_Anim")) 
+    controls.append((namespace + ":" + "L_Ring_1_Anim"))
+    controls.append((namespace + ":" + "L_Ring_2_Anim"))
+    controls.append((namespace + ":" + "L_Ring_3_Anim"))
+    controls.append((namespace + ":" + "L_Pinky_1_Anim"))
+    controls.append((namespace + ":" + "L_Pinky_2_Anim"))
+    controls.append((namespace + ":" + "L_Pinky_3_Anim"))
+    controls.append((namespace + ":" + "L_Thumb_1_Anim"))
+    controls.append((namespace + ":" + "L_Thumb_2_Anim"))
+    controls.append((namespace + ":" + "L_Thumb_3_Anim"))
+    controls.append((namespace + ":" + "R_Index_1_Anim"))
+    controls.append((namespace + ":" + "R_Index_2_Anim"))
+    controls.append((namespace + ":" + "R_Index_3_Anim"))
+    controls.append((namespace + ":" + "R_Middle_1_Anim"))
+    controls.append((namespace + ":" + "R_Middle_2_Anim"))
+    controls.append((namespace + ":" + "R_Middle_3_Anim"))
+    controls.append((namespace + ":" + "R_Ring_1_Anim"))
+    controls.append((namespace + ":" + "R_Ring_2_Anim"))
+    controls.append((namespace + ":" + "R_Ring_3_Anim"))
+    controls.append((namespace + ":" + "R_Pinky_1_Anim"))
+    controls.append((namespace + ":" + "R_Pinky_2_Anim"))
+    controls.append((namespace + ":" + "R_Pinky_3_Anim"))
+    controls.append((namespace + ":" + "R_Thumb_1_Anim"))
+    controls.append((namespace + ":" + "R_Thumb_2_Anim"))
+    controls.append((namespace + ":" + "R_Thumb_3_Anim"))
+    
+    #for each control, grab the keyable, unlocked attrs, and zero them out
+    
+    for control in controls:
+        attrs = cmds.listAttr(control, k = True, u = True)
+        
+        for attr in attrs:
+            if attr != "globalScale":
+                cmds.setAttr((control + "." + attr), 0)
+            
+    
+def keySelection(*args):
+    
+    #grab the namespace of the character
+    namespace = cmds.optionMenu("selectedCharacterOptionMenu", q = True, v = True)
+    namespace += ":"
+    
+    #get the current selection
+    selection = cmds.ls(sl = True, transforms = True)
+    
+    #key each object in selection
+    for control in selection:
+        cmds.setKeyframe(control)
+        
+    
+def keyAll(*args):
+    
+    #grab the namespace of the character
+    namespace = cmds.optionMenu("selectedCharacterOptionMenu", q = True, v = True)
+    
+    #set keyframes on all of the controls
+    
+    cmds.setKeyframe((namespace + ":" + "FK_b_MF_Thigh_R"))
+    cmds.setKeyframe((namespace + ":" + "FK_b_MF_Calf_R"))
+    cmds.setKeyframe((namespace + ":" + "FK_b_MF_Foot_R"))
+    cmds.setKeyframe((namespace + ":" + "FK_b_MF_Toe_R"))
+    cmds.setKeyframe((namespace + ":" + "Right_Knee_Anim"))
+    cmds.setKeyframe((namespace + ":" + "Right_Foot_Anim"))
+        
+    cmds.setKeyframe((namespace + ":" + "FK_b_MF_Thigh_L"))
+    cmds.setKeyframe((namespace + ":" + "FK_b_MF_Calf_L"))
+    cmds.setKeyframe((namespace + ":" + "FK_b_MF_Foot_L"))
+    cmds.setKeyframe((namespace + ":" + "FK_b_MF_Toe_L"))
+    cmds.setKeyframe((namespace + ":" + "Left_Knee_Anim"))
+    cmds.setKeyframe((namespace + ":" + "Left_Foot_Anim"))
+        
+    cmds.setKeyframe((namespace + ":" + "R_Shoulder_Anim"))
+    cmds.setKeyframe((namespace + ":" + "R_IK_Elbow_Anim"))
+    cmds.setKeyframe((namespace + ":" + "R_IK_Hand_Anim"))
+    cmds.setKeyframe((namespace + ":" + "FK_b_MF_UpperArm_R"))
+    cmds.setKeyframe((namespace + ":" + "FK_b_MF_Forearm_R"))
+    cmds.setKeyframe((namespace + ":" + "FK_b_MF_Hand_R"))
+    cmds.setKeyframe((namespace + ":" + "R_Weapon_Anim"))
+        
+    cmds.setKeyframe((namespace + ":" + "L_Shoulder_Anim"))
+    cmds.setKeyframe((namespace + ":" + "L_IK_Elbow_Anim"))
+    cmds.setKeyframe((namespace + ":" + "L_IK_Hand_Anim"))
+    cmds.setKeyframe((namespace + ":" + "FK_b_MF_UpperArm_L"))
+    cmds.setKeyframe((namespace + ":" + "FK_b_MF_Forearm_L"))
+    cmds.setKeyframe((namespace + ":" + "FK_b_MF_Hand_L"))
+    cmds.setKeyframe((namespace + ":" + "L_Weapon_Anim"))
+        
+    cmds.setKeyframe((namespace + ":" + "Neck_Anim"))
+    cmds.setKeyframe((namespace + ":" + "Head_Anim"))
+        
+    cmds.setKeyframe((namespace + ":" + "Spine_01_Anim"))
+    cmds.setKeyframe((namespace + ":" + "Spine_02_Anim"))
+    cmds.setKeyframe((namespace + ":" + "Spine_03_Anim"))
+    cmds.setKeyframe((namespace + ":" + "Body_Anim"))
+    cmds.setKeyframe((namespace + ":" + "Hip_Anim"))
+        
+    cmds.setKeyframe((namespace + ":" + "Master_Anim"))
+    cmds.setKeyframe((namespace + ":" + "Root_Anim"))
+    cmds.setKeyframe((namespace + ":" + "Attach_Anim"))
+        
+    cmds.setKeyframe((namespace + ":" + "Settings"))
+        
+    #fingers
+    cmds.setKeyframe((namespace + ":" + "L_Index_1_Anim"))
+    cmds.setKeyframe((namespace + ":" + "L_Index_2_Anim"))
+    cmds.setKeyframe((namespace + ":" + "L_Index_3_Anim"))
+        
+    cmds.setKeyframe((namespace + ":" + "L_Middle_1_Anim"))
+    cmds.setKeyframe((namespace + ":" + "L_Middle_2_Anim"))
+    cmds.setKeyframe((namespace + ":" + "L_Middle_3_Anim"))
+        
+    cmds.setKeyframe((namespace + ":" + "L_Ring_1_Anim"))
+    cmds.setKeyframe((namespace + ":" + "L_Ring_2_Anim"))
+    cmds.setKeyframe((namespace + ":" + "L_Ring_3_Anim"))
+        
+    cmds.setKeyframe((namespace + ":" + "L_Pinky_1_Anim"))
+    cmds.setKeyframe((namespace + ":" + "L_Pinky_2_Anim"))
+    cmds.setKeyframe((namespace + ":" + "L_Pinky_3_Anim"))
+        
+    cmds.setKeyframe((namespace + ":" + "L_Thumb_1_Anim"))
+    cmds.setKeyframe((namespace + ":" + "L_Thumb_2_Anim"))
+    cmds.setKeyframe((namespace + ":" + "L_Thumb_3_Anim"))
+        
+        
+        
+    cmds.setKeyframe((namespace + ":" + "R_Index_1_Anim"))
+    cmds.setKeyframe((namespace + ":" + "R_Index_2_Anim"))
+    cmds.setKeyframe((namespace + ":" + "R_Index_3_Anim"))
+        
+    cmds.setKeyframe((namespace + ":" + "R_Middle_1_Anim"))
+    cmds.setKeyframe((namespace + ":" + "R_Middle_2_Anim"))
+    cmds.setKeyframe((namespace + ":" + "R_Middle_3_Anim"))
+        
+    cmds.setKeyframe((namespace + ":" + "R_Ring_1_Anim"))
+    cmds.setKeyframe((namespace + ":" + "R_Ring_2_Anim"))
+    cmds.setKeyframe((namespace + ":" + "R_Ring_3_Anim"))
+        
+    cmds.setKeyframe((namespace + ":" + "R_Pinky_1_Anim"))
+    cmds.setKeyframe((namespace + ":" + "R_Pinky_2_Anim"))
+    cmds.setKeyframe((namespace + ":" + "R_Pinky_3_Anim"))
+        
+    cmds.setKeyframe((namespace + ":" + "R_Thumb_1_Anim"))
+    cmds.setKeyframe((namespace + ":" + "R_Thumb_2_Anim"))
+    cmds.setKeyframe((namespace + ":" + "R_Thumb_3_Anim"))
+    
+def activeCharacterChangeCommand(*args):
+    
+    namespace = cmds.optionMenu("selectedCharacterOptionMenu", q = True, v = True)
+    name = namespace
+    namespace += ":"
+    
+    
+    #if the face Rig UI tab is there, see if you can also find and change that tab
+    tabs = cmds.tabLayout("autoRigUI_TabLayout", q = True, tabLabelIndex = True)
+    
+    for tab in tabs:
+        if tab == "Face Rig UI":
+            
+            try:
+                cmds.optionMenu("selectedCharacterOptionMenuFaceRig", e = True, v = name)
+                cmds.tabLayout( "autoRigUI_TabLayout", edit=True, tabsVisible = 1)
+                
+            except RuntimeError:
+                cmds.warning("Character has no face rig. Disabling Face Rig Tab")
+                cmds.tabLayout( "autoRigUI_TabLayout", edit=True, tabsVisible = 0)
+            
+        
+    #LEFT
+    #need to look up the fk orient mode on the character
+    mode = cmds.getAttr((namespace + "Settings.L_Arm_FK_Orient"))
+    
+    if mode == 0:
+        cmds.optionMenu("lClavOrientMenu", e= True, sl = 1)
+        
+    if mode == 1:
+        cmds.optionMenu("lClavOrientMenu", e= True, sl = 2)
+        
+    if mode == 2:
+        cmds.optionMenu("lClavOrientMenu", e= True, sl = 3)
+        
+    if mode == 3:
+        cmds.optionMenu("lClavOrientMenu", e= True, sl = 4)
+        
+    #RIGHT
+    #need to look up the fk orient mode on the character
+    rmode = cmds.getAttr((namespace + "Settings.R_Arm_FK_Orient"))
+    
+    if rmode == 0:
+        cmds.optionMenu("rClavOrientMenu", e= True, sl = 1)
+        
+    if rmode == 1:
+        cmds.optionMenu("rClavOrientMenu", e= True, sl = 2)
+        
+    if rmode == 2:
+        cmds.optionMenu("rClavOrientMenu", e= True, sl = 3)
+        
+    if rmode == 3:
+        cmds.optionMenu("rClavOrientMenu", e= True, sl = 4)
+        
+        
+def lFKArmOrient(*args):
+    namespace = cmds.optionMenu("selectedCharacterOptionMenu", q = True, v = True)
+    namespace += ":"
+    
+    #for the selected character, look up the current fk arm orient mode
+    mode = cmds.optionMenu("lClavOrientMenu", q= True, sl = True)
+    
+    
+    if mode == 1:
+        cmds.setAttr((namespace + "Settings.L_Arm_FK_Orient"), 0)
+        
+    if mode == 2:
+        cmds.setAttr((namespace + "Settings.L_Arm_FK_Orient"), 1)
+        
+    if mode == 3:
+        cmds.setAttr((namespace + "Settings.L_Arm_FK_Orient"), 2)
+        
+    if mode == 4:
+        cmds.setAttr((namespace + "Settings.L_Arm_FK_Orient"), 3)
+        
+def rFKArmOrient(*args):
+    namespace = cmds.optionMenu("selectedCharacterOptionMenu", q = True, v = True)
+    namespace += ":"
+    
+    #for the selected character, look up the current fk arm orient mode
+    mode = cmds.optionMenu("rClavOrientMenu", q= True, sl = True)
+    
+    
+    if mode == 1:
+        cmds.setAttr((namespace + "Settings.R_Arm_FK_Orient"), 0)
+        
+    if mode == 2:
+        cmds.setAttr((namespace + "Settings.R_Arm_FK_Orient"), 1)
+        
+    if mode == 3:
+        cmds.setAttr((namespace + "Settings.R_Arm_FK_Orient"), 2)
+        
+    if mode == 4:
+        cmds.setAttr((namespace + "Settings.R_Arm_FK_Orient"), 3)
+
+    
+
+    
+    
+    
+def LArmModeSwitchIK(*args):
+    namespace = cmds.optionMenu("selectedCharacterOptionMenu", q = True, v = True)
+    namespace += ":"
+    
+    
+    #call the match function
+    leftArmMatchFK(namespace)
+    
+    #switch to IK mode
+    cmds.setAttr((namespace + "Settings.L_Arm_Mode"), 10)
+    
+    #then key the settings locator
+    cmds.setKeyframe((namespace + "Settings.L_Arm_Mode"), ott = 'step')
+    
+    #now make sure that any previous keys are also stepped
+    cmds.selectKey(clear = True)
+    cmds.selectKey((namespace + "Settings.L_Arm_Mode"),add = True, k = True)
+    cmds.selectKey((namespace + "Settings.R_Arm_Mode"),add = True, k = True)
+    cmds.selectKey((namespace + "Settings.L_Arm_FK_Orient"),add = True, k = True)
+    cmds.selectKey((namespace + "Settings.R_Arm_FK_Orient"),add = True, k = True)
+    cmds.selectKey((namespace + "Settings.L_Leg_Mode"),add = True, k = True)
+    cmds.selectKey((namespace + "Settings.R_Leg_Mode"),add = True, k = True)
+    cmds.keyTangent(ott = 'step')
+    
+def RArmModeSwitchIK(*args):
+    namespace = cmds.optionMenu("selectedCharacterOptionMenu", q = True, v = True)
+    namespace += ":"
+    
+    
+    #call the match function
+    rightArmMatchFK(namespace)
+    
+    #switch to IK mode
+    cmds.setAttr((namespace + "Settings.R_Arm_Mode"), 10)
+    
+    #then key the settings locator
+    cmds.setKeyframe((namespace + "Settings.R_Arm_Mode"), ott = 'step')
+    
+    #now make sure that any previous keys are also stepped
+    cmds.selectKey(clear = True)
+    cmds.selectKey((namespace + "Settings.L_Arm_Mode"),add = True, k = True)
+    cmds.selectKey((namespace + "Settings.R_Arm_Mode"),add = True, k = True)
+    cmds.selectKey((namespace + "Settings.L_Arm_FK_Orient"),add = True, k = True)
+    cmds.selectKey((namespace + "Settings.R_Arm_FK_Orient"),add = True, k = True)
+    cmds.selectKey((namespace + "Settings.L_Leg_Mode"),add = True, k = True)
+    cmds.selectKey((namespace + "Settings.R_Leg_Mode"),add = True, k = True)
+    cmds.keyTangent(ott = 'step')
+        
+
+        
+def LArmModeSwitchFK(*args):
+    namespace = cmds.optionMenu("selectedCharacterOptionMenu", q = True, v = True)
+    namespace += ":"
+    
+    
+    #call the match function
+    leftArmMatchIK(namespace)
+    
+    #switch to FK mode
+    cmds.setAttr((namespace + "Settings.L_Arm_Mode"), 0)
+    
+    #then key the settings locator
+    cmds.setKeyframe((namespace + "Settings.L_Arm_Mode"), ott = 'step')
+    
+    #now make sure that any previous keys are also stepped
+    cmds.selectKey(clear = True)
+    cmds.selectKey((namespace + "Settings.L_Arm_Mode"),add = True, k = True)
+    cmds.selectKey((namespace + "Settings.R_Arm_Mode"),add = True, k = True)
+    cmds.selectKey((namespace + "Settings.L_Arm_FK_Orient"),add = True, k = True)
+    cmds.selectKey((namespace + "Settings.R_Arm_FK_Orient"),add = True, k = True)
+    cmds.selectKey((namespace + "Settings.L_Leg_Mode"),add = True, k = True)
+    cmds.selectKey((namespace + "Settings.R_Leg_Mode"),add = True, k = True)
+    cmds.keyTangent(ott = 'step')
+    
+def RArmModeSwitchFK(*args):
+    namespace = cmds.optionMenu("selectedCharacterOptionMenu", q = True, v = True)
+    namespace += ":"
+    
+    
+    #call the match function
+    rightArmMatchIK(namespace)
+    
+    #switch to FK mode
+    cmds.setAttr((namespace + "Settings.R_Arm_Mode"), 0)
+    
+    #then key the settings locator
+    cmds.setKeyframe((namespace + "Settings.R_Arm_Mode"), ott = 'step')
+    
+    #now make sure that any previous keys are also stepped
+    cmds.selectKey(clear = True)
+    cmds.selectKey((namespace + "Settings.L_Arm_Mode"),add = True, k = True)
+    cmds.selectKey((namespace + "Settings.R_Arm_Mode"),add = True, k = True)
+    cmds.selectKey((namespace + "Settings.L_Arm_FK_Orient"),add = True, k = True)
+    cmds.selectKey((namespace + "Settings.R_Arm_FK_Orient"),add = True, k = True)
+    cmds.selectKey((namespace + "Settings.L_Leg_Mode"),add = True, k = True)
+    cmds.selectKey((namespace + "Settings.R_Leg_Mode"),add = True, k = True)
+    cmds.keyTangent(ott = 'step')
+    
+
+        
+
+
+        
+        
+def leftArmMatchFK(namespace):
+    
+
+
+    #unlock the containers
+    containers = [(namespace + "Rig_Container")]
+    for container in containers:
+        cmds.lockNode(container, lock = False, lockUnpublished = False)
+        
+    #first mastch the hand
+    constraint = cmds.parentConstraint((namespace + "L_IK_Match"),(namespace + "L_IK_Hand_Anim"))[0]
+    cmds.setKeyframe((namespace + "L_IK_Hand_Anim"))
+    cmds.delete(constraint)
+
+    
+
+    #then match the elbow
+    constraint = cmds.pointConstraint((namespace + "FK_b_MF_Forearm_L"),(namespace + "L_IK_Elbow_Anim"))[0]
+    cmds.setKeyframe((namespace + "L_IK_Elbow_Anim"),ott = 'spline', itt = 'spline')
+    cmds.delete(constraint)
+    
+
+    
+    for container in containers:
+        cmds.lockNode(container, lock = False, lockUnpublished = False)
+        
+    cmds.select((namespace + "L_IK_Hand_Anim"),replace = True)
+
+def rightArmMatchFK(namespace):
+    
+
+
+    #unlock the containers
+    containers = [(namespace + "Rig_Container")]
+    for container in containers:
+        cmds.lockNode(container, lock = False, lockUnpublished = False)
+        
+    #first mastch the hand
+    constraint = cmds.parentConstraint((namespace + "R_IK_Match"),(namespace + "R_IK_Hand_Anim"))[0]
+    cmds.setKeyframe((namespace + "R_IK_Hand_Anim"))
+    cmds.delete(constraint)
+
+    
+
+    #then match the elbow
+    constraint = cmds.pointConstraint((namespace + "FK_b_MF_Forearm_R"),(namespace + "R_IK_Elbow_Anim"))[0]
+    cmds.setKeyframe((namespace + "R_IK_Elbow_Anim"),ott = 'spline', itt = 'spline')
+    cmds.delete(constraint)
+    
+
+    
+    for container in containers:
+        cmds.lockNode(container, lock = False, lockUnpublished = False)
+        
+    cmds.select((namespace + "R_IK_Hand_Anim"),replace = True)
+    
+def leftArmMatchIK(namespace):
+    
+    #unlock the containers
+    containers = [(namespace + "Rig_Container")]
+    for container in containers:
+        cmds.lockNode(container, lock = False, lockUnpublished = False)
+        
+        
+    #first match the hand
+    constraint1 = cmds.orientConstraint((namespace + "IK_b_MF_UpperArm_L"),(namespace + "FK_b_MF_UpperArm_L"))[0]
+    constraint2 = cmds.orientConstraint((namespace + "IK_b_MF_Forearm_L"),(namespace + "FK_b_MF_Forearm_L"))[0]
+    constraint3 = cmds.orientConstraint((namespace + "IK_b_MF_Hand_L"),(namespace + "FK_b_MF_Hand_L"))[0]
+
+
+    #then key the fk arm
+    
+    cmds.setKeyframe((namespace + "FK_b_MF_UpperArm_L"),)
+    cmds.setKeyframe((namespace + "FK_b_MF_Forearm_L"))
+    cmds.setKeyframe((namespace + "FK_b_MF_Hand_L"))
+    
+    cmds.delete([constraint1, constraint2, constraint3])
+    
+    for container in containers:
+        cmds.lockNode(container, lock = False, lockUnpublished = False)
+    
+
+    cmds.select((namespace + "FK_b_MF_Hand_L"),replace = True)
+    
+    
+def rightArmMatchIK(namespace):
+    
+    #unlock the containers
+    containers = [(namespace + "Rig_Container")]
+    for container in containers:
+        cmds.lockNode(container, lock = False, lockUnpublished = False)
+        
+        
+    #first match the hand
+    constraint1 = cmds.orientConstraint((namespace + "IK_b_MF_UpperArm_R"),(namespace + "FK_b_MF_UpperArm_R"))[0]
+    constraint2 = cmds.orientConstraint((namespace + "IK_b_MF_Forearm_R"),(namespace + "FK_b_MF_Forearm_R"))[0]
+    constraint3 = cmds.orientConstraint((namespace + "IK_b_MF_Hand_R"),(namespace + "FK_b_MF_Hand_R"))[0]
+
+
+    #then key the fk arm
+    
+    cmds.setKeyframe((namespace + "FK_b_MF_UpperArm_R"),)
+    cmds.setKeyframe((namespace + "FK_b_MF_Forearm_R"))
+    cmds.setKeyframe((namespace + "FK_b_MF_Hand_R"))
+    
+    cmds.delete([constraint1, constraint2, constraint3])
+    
+    for container in containers:
+        cmds.lockNode(container, lock = False, lockUnpublished = False)
+    
+
+    cmds.select((namespace + "FK_b_MF_Hand_R"),replace = True)
+    
+
+
+def LLegModeSwitchIK(*args):
+    
+    namespace = cmds.optionMenu("selectedCharacterOptionMenu", q = True, v = True)
+    namespace += ":"
+    
+    
+    #call the match function
+    LLegMatchFK(namespace)
+    
+    #switch to IK mode
+    cmds.setAttr((namespace + "Settings.L_Leg_Mode"), 10)
+    cmds.setAttr((namespace + "Left_Foot_Anim.roll"), 0)
+    cmds.setAttr((namespace + "Left_Foot_Anim.side"), 0)
+    cmds.setAttr((namespace + "Left_Foot_Anim.lean"), 0)
+    cmds.setAttr((namespace + "Left_Foot_Anim.toe_spin"), 0)
+    cmds.setAttr((namespace + "Left_Foot_Anim.toe_wiggle"), 0)
+    
+    #then key the settings locator
+    cmds.setKeyframe((namespace + "Settings.L_Leg_Mode"), ott = 'step')
+    
+    #now make sure that any previous keys are also stepped
+    cmds.selectKey(clear = True)
+    cmds.selectKey((namespace + "Settings.L_Arm_Mode"),add = True, k = True)
+    cmds.selectKey((namespace + "Settings.R_Arm_Mode"),add = True, k = True)
+    cmds.selectKey((namespace + "Settings.L_Arm_FK_Orient"),add = True, k = True)
+    cmds.selectKey((namespace + "Settings.R_Arm_FK_Orient"),add = True, k = True)
+    cmds.selectKey((namespace + "Settings.L_Leg_Mode"),add = True, k = True)
+    cmds.selectKey((namespace + "Settings.R_Leg_Mode"),add = True, k = True)
+    cmds.keyTangent(ott = 'step')
+    
+    
+def LLegModeSwitchFK(*args):
+    namespace = cmds.optionMenu("selectedCharacterOptionMenu", q = True, v = True)
+    namespace += ":"
+    
+    
+    #call the match function
+    LLegMatchIK(namespace)
+    
+    #switch to FK mode
+    cmds.setAttr((namespace + "Settings.L_Leg_Mode"), 0)
+    
+    #then key the settings locator
+    cmds.setKeyframe((namespace + "Settings.L_Leg_Mode"), ott = 'step')
+    
+    #now make sure that any previous keys are also stepped
+    cmds.selectKey(clear = True)
+    cmds.selectKey((namespace + "Settings.L_Arm_Mode"),add = True, k = True)
+    cmds.selectKey((namespace + "Settings.R_Arm_Mode"),add = True, k = True)
+    cmds.selectKey((namespace + "Settings.L_Arm_FK_Orient"),add = True, k = True)
+    cmds.selectKey((namespace + "Settings.R_Arm_FK_Orient"),add = True, k = True)
+    cmds.selectKey((namespace + "Settings.L_Leg_Mode"),add = True, k = True)
+    cmds.selectKey((namespace + "Settings.R_Leg_Mode"),add = True, k = True)
+    cmds.keyTangent(ott = 'step')
+    
+    
+def RLegModeSwitchIK(*args):
+    
+    namespace = cmds.optionMenu("selectedCharacterOptionMenu", q = True, v = True)
+    namespace += ":"
+    
+    
+    #call the match function
+    RLegMatchFK(namespace)
+    #switch to IK mode
+    cmds.setAttr((namespace + "Settings.R_Leg_Mode"), 10)
+    cmds.setAttr((namespace + "Right_Foot_Anim.roll"), 0)
+    cmds.setAttr((namespace + "Right_Foot_Anim.side"), 0)
+    cmds.setAttr((namespace + "Right_Foot_Anim.lean"), 0)
+    cmds.setAttr((namespace + "Right_Foot_Anim.toe_spin"), 0)
+    cmds.setAttr((namespace + "Right_Foot_Anim.toe_wiggle"), 0)
+    
+    #then key the settings locator
+    cmds.setKeyframe((namespace + "Settings.R_Leg_Mode"), ott = 'step')
+    
+    #now make sure that any previous keys are also stepped
+    cmds.selectKey(clear = True)
+    cmds.selectKey((namespace + "Settings.L_Arm_Mode"),add = True, k = True)
+    cmds.selectKey((namespace + "Settings.R_Arm_Mode"),add = True, k = True)
+    cmds.selectKey((namespace + "Settings.L_Arm_FK_Orient"),add = True, k = True)
+    cmds.selectKey((namespace + "Settings.R_Arm_FK_Orient"),add = True, k = True)
+    cmds.selectKey((namespace + "Settings.L_Leg_Mode"),add = True, k = True)
+    cmds.selectKey((namespace + "Settings.R_Leg_Mode"),add = True, k = True)
+    cmds.keyTangent(ott = 'step')
+    
+    
+def RLegModeSwitchFK(*args):
+    namespace = cmds.optionMenu("selectedCharacterOptionMenu", q = True, v = True)
+    namespace += ":"
+    
+    
+    #call the match function
+    RLegMatchIK(namespace)
+    #switch to FK mode
+    cmds.setAttr((namespace + "Settings.R_Leg_Mode"), 0)
+    
+    #then key the settings locator
+    cmds.setKeyframe((namespace + "Settings.R_Leg_Mode"), ott = 'step')
+    
+    #now make sure that any previous keys are also stepped
+    cmds.selectKey(clear = True)
+    cmds.selectKey((namespace + "Settings.L_Arm_Mode"),add = True, k = True)
+    cmds.selectKey((namespace + "Settings.R_Arm_Mode"),add = True, k = True)
+    cmds.selectKey((namespace + "Settings.L_Arm_FK_Orient"),add = True, k = True)
+    cmds.selectKey((namespace + "Settings.R_Arm_FK_Orient"),add = True, k = True)
+    cmds.selectKey((namespace + "Settings.L_Leg_Mode"),add = True, k = True)
+    cmds.selectKey((namespace + "Settings.R_Leg_Mode"),add = True, k = True)
+    cmds.keyTangent(ott = 'step')
+
+
+def LLegMatchFK(namespace):
+    
+
+
+    #unlock the containers
+    containers = [(namespace + "Rig_Container")]
+    for container in containers:
+        cmds.lockNode(container, lock = False, lockUnpublished = False)
+        
+    #first mastch the hand
+    constraint = cmds.parentConstraint((namespace + "L_IK_Leg_Match"),(namespace + "Left_Foot_Anim"))[0]
+    cmds.setKeyframe((namespace + "Left_Foot_Anim"))
+    cmds.delete(constraint)
+
+    
+
+    #then match the elbow
+    constraint = cmds.pointConstraint((namespace + "FK_b_MF_Calf_L"),(namespace + "Left_Knee_Anim"))[0]
+    cmds.setKeyframe((namespace + "Left_Knee_Anim"),ott = 'spline', itt = 'spline')
+    cmds.delete(constraint)
+    
+
+    
+    for container in containers:
+        cmds.lockNode(container, lock = False, lockUnpublished = False)
+        
+    cmds.select((namespace + "Left_Foot_Anim"),replace = True)
+    
+    
+def RLegMatchFK(namespace):
+    
+
+
+    #unlock the containers
+    containers = [(namespace + "Rig_Container")]
+    for container in containers:
+        cmds.lockNode(container, lock = False, lockUnpublished = False)
+        
+    #first mastch the hand
+    constraint = cmds.parentConstraint((namespace + "R_IK_Leg_Match"),(namespace + "Right_Foot_Anim"))[0]
+    cmds.setKeyframe((namespace + "Right_Foot_Anim"))
+    cmds.delete(constraint)
+
+    
+
+    #then match the elbow
+    constraint = cmds.pointConstraint((namespace + "FK_b_MF_Calf_R"),(namespace + "Right_Knee_Anim"))[0]
+    cmds.setKeyframe((namespace + "Right_Knee_Anim"),ott = 'spline', itt = 'spline')
+    cmds.delete(constraint)
+    
+
+    
+    for container in containers:
+        cmds.lockNode(container, lock = False, lockUnpublished = False)
+        
+    cmds.select((namespace + "Right_Foot_Anim"),replace = True)
+    
+    
+def LLegMatchIK(namespace):
+    
+    #unlock the containers
+    containers = [(namespace + "Rig_Container")]
+    for container in containers:
+        cmds.lockNode(container, lock = False, lockUnpublished = False)
+        
+        
+    #first match the hand
+    constraint1 = cmds.orientConstraint((namespace + "IK_b_MF_Thigh_L"),(namespace + "FK_b_MF_Thigh_L"))[0]
+    constraint2 = cmds.orientConstraint((namespace + "IK_b_MF_Calf_L"),(namespace + "FK_b_MF_Calf_L"))[0]
+    constraint3 = cmds.orientConstraint((namespace + "IK_b_MF_Foot_L"),(namespace + "FK_b_MF_Foot_L"))[0]
+    constraint4 = cmds.orientConstraint((namespace + "IK_b_MF_Toe_L"),(namespace + "FK_b_MF_Toe_L"))[0]
+
+    #then key the fk arm
+    
+    cmds.setKeyframe((namespace + "FK_b_MF_Thigh_L"),)
+    cmds.setKeyframe((namespace + "FK_b_MF_Calf_L"))
+    cmds.setKeyframe((namespace + "FK_b_MF_Foot_L"))
+    cmds.setKeyframe((namespace + "FK_b_MF_Toe_L"))
+     
+        
+    cmds.delete([constraint1, constraint2, constraint3,constraint4])
+    
+    for container in containers:
+        cmds.lockNode(container, lock = False, lockUnpublished = False)
+    
+
+    cmds.select((namespace + "FK_b_MF_Thigh_L"),replace = True)
+    
+    
+def RLegMatchIK(namespace):
+    
+    #unlock the containers
+    containers = [(namespace + "Rig_Container")]
+    for container in containers:
+        cmds.lockNode(container, lock = False, lockUnpublished = False)
+        
+        
+    #first match the hand
+    constraint1 = cmds.orientConstraint((namespace + "IK_b_MF_Thigh_R"),(namespace + "FK_b_MF_Thigh_R"))[0]
+    constraint2 = cmds.orientConstraint((namespace + "IK_b_MF_Calf_R"),(namespace + "FK_b_MF_Calf_R"))[0]
+    constraint3 = cmds.orientConstraint((namespace + "IK_b_MF_Foot_R"),(namespace + "FK_b_MF_Foot_R"))[0]
+    constraint4 = cmds.orientConstraint((namespace + "IK_b_MF_Toe_R"),(namespace + "FK_b_MF_Toe_R"))[0]
+
+    #then key the fk arm
+    
+    cmds.setKeyframe((namespace + "FK_b_MF_Thigh_R"),)
+    cmds.setKeyframe((namespace + "FK_b_MF_Calf_R"))
+    cmds.setKeyframe((namespace + "FK_b_MF_Foot_R"))
+    cmds.setKeyframe((namespace + "FK_b_MF_Toe_R"))
+     
+        
+    cmds.delete([constraint1, constraint2, constraint3,constraint4])
+    
+    for container in containers:
+        cmds.lockNode(container, lock = False, lockUnpublished = False)
+    
+
+    cmds.select((namespace + "FK_b_MF_Thigh_R"),replace = True)
+    
+
+    
+def selectFinger(control,button,*args):
+    namespace = cmds.optionMenu("selectedCharacterOptionMenu", q = True, v = True)
+    namespace += ":"
+    mods = cmds.getModifiers()
+
+    if (mods & 1) > 0:
+        cmds.select((namespace + control), tgl = True)
+        
+        buttonVal = cmds.button(button, q = True, bgc = True)
+        
+        if buttonVal[1] == 1.0:
+            cmds.button(button, edit = True, bgc = [.196, .475, .98])
+            
+        else:
+            cmds.button(button, edit = True, bgc = [1,1,1])
+        
+            controls = [(namespace + control)]
+            buttons = [button]
+        
+            createScriptJob(controls, buttons)
+        
+        
+    if (mods & 1) == 0:
+        cmds.select((namespace + control), replace = True)
+        cmds.button(button, edit = True, bgc = [1,1,1])
+        
+        controls = [(namespace + control)]
+        buttons = [button]
+        
+        createScriptJob(controls, buttons)
+        
+        
+def selectMultiFingers(controls,buttons,*args):
+    
+    namespace = cmds.optionMenu("selectedCharacterOptionMenu", q = True, v = True)
+    namespace += ":"
+    mods = cmds.getModifiers()
+
+    controlsToSelect = []
+    buttonsToChange = []
+    
+    if (mods & 1) > 0:
+        for i in range(len(controls)):
+        
+            controlsToSelect.append(namespace + controls[i])
+            buttonsToChange.append(buttons[i])
+            
+        for each in controlsToSelect:
+            cmds.select(each, tgl = True)
+            
+        for each in buttonsToChange:
+            buttonVal = cmds.button(each, q = True, bgc = True)
+        
+            if buttonVal[1] == 1.0:
+                cmds.button(each, edit = True, bgc = [.196, .475, .98])
+            
+            else:
+                cmds.button(each, edit = True, bgc = [1,1,1])
+        
+                controls = controlsToSelect
+                buttons = buttonsToChange
+        
+                createScriptJob(controls, buttons)
+                
+        
+        
+    if (mods & 1) == 0:
+        
+        for i in range(len(controls)):
+            
+            controlsToSelect.append(namespace + controls[i])
+            buttonsToChange.append(buttons[i])
+            
+            
+        cmds.select(controlsToSelect[0], replace = True)
+        
+        newControls = list(controlsToSelect)
+        newControls.pop(0)
+        
+        for each in newControls:
+            cmds.select(each, add = True)
+            
+        for each in buttonsToChange:
+            cmds.button(each, edit = True, bgc = [1,1,1])
+        
+            controls = controlsToSelect
+            buttons = buttonsToChange
+        
+            createScriptJob(controls, buttons)
+            
+        
+        
+    
+        
+        
+
+
+def selectBody(*args):
+    namespace = cmds.optionMenu("selectedCharacterOptionMenu", q = True, v = True)
+    mods = cmds.getModifiers()
+    
+    if (mods & 1) > 0:
+        cmds.select((namespace + ":" + "Body_Anim"), tgl = True)
+        
+        buttonVal = cmds.button("bodyButton", q = True, bgc = True)
+        
+        if buttonVal[1] == 1.0:
+            cmds.button("bodyButton", edit = True, bgc = [.968,.639,0])
+            
+        else:
+            cmds.button("bodyButton", edit = True, bgc = [1,1,1])
+        
+            controls = [(namespace + ":" + "Body_Anim")]
+            buttons = ["bodyButton"]
+        
+            createScriptJob(controls, buttons)
+        
+        
+    if (mods & 1) == 0:
+        cmds.select((namespace + ":" + "Body_Anim"), replace = True)
+        cmds.button("bodyButton", edit = True, bgc = [1,1,1])
+        
+        controls = [(namespace + ":" + "Body_Anim")]
+        buttons = ["bodyButton"]
+        
+        createScriptJob(controls, buttons)
+        
+def selectSettings(*args):
+    namespace = cmds.optionMenu("selectedCharacterOptionMenu", q = True, v = True)
+
+    
+    mods = cmds.getModifiers()
+    
+    if (mods & 1) > 0:
+        cmds.select((namespace + ":" + "Settings"), tgl = True)
+        
+        buttonVal = cmds.button("settingsButtonBGC", q = True, bgc = True)
+        
+        if buttonVal[1] == 1.0:
+            cmds.button("settingsButtonBGC", edit = True, bgc = [0.4,0.4,0.4])
+            
+        else:
+            cmds.button("settingsButtonBGC", edit = True, bgc = [1,1,1])
+        
+            controls = [(namespace + ":" + "Settings")]
+            buttons = ["settingsButtonBGC"]
+        
+            createScriptJob(controls, buttons)
+        
+        
+    if (mods & 1) == 0:
+        
+        cmds.select((namespace + ":" + "Settings"), replace = True)
+        cmds.button("settingsButtonBGC", edit = True, bgc = [1,1,1])
+        
+        controls = [(namespace + ":" + "Settings")]
+        buttons = ["settingsButtonBGC"]
+        
+        createScriptJob(controls, buttons)
+
+        
+def selectMaster(*args):
+    namespace = cmds.optionMenu("selectedCharacterOptionMenu", q = True, v = True)
+    mods = cmds.getModifiers()
+    
+    if (mods & 1) > 0:
+        cmds.select((namespace + ":" + "Master_Anim"), tgl = True)
+        
+        buttonVal = cmds.button("masterButton", q = True, bgc = True)
+        
+        if buttonVal[1] == 1.0:
+            cmds.button("masterButton", edit = True, bgc = [.968,.639,0])
+            
+        else:
+            cmds.button("masterButton", edit = True, bgc = [1,1,1])
+        
+            controls = [(namespace + ":" + "Master_Anim")]
+            buttons = ["masterButton"]
+        
+            createScriptJob(controls, buttons)
+        
+        
+    if (mods & 1) == 0:
+        cmds.select((namespace + ":" + "Master_Anim"), replace = True)
+        cmds.button("masterButton", edit = True, bgc = [1,1,1])
+        
+        controls = [(namespace + ":" + "Master_Anim")]
+        buttons = ["masterButton"]
+        
+        createScriptJob(controls, buttons)
+        
+def selectRoot(*args):
+    namespace = cmds.optionMenu("selectedCharacterOptionMenu", q = True, v = True)
+    mods = cmds.getModifiers()
+    
+    if (mods & 1) > 0:
+        cmds.select((namespace + ":" + "Root_Anim"), tgl = True)
+        
+        buttonVal = cmds.button("rootButton", q = True, bgc = True)
+        
+        if buttonVal[1] == 1.0:
+            cmds.button("rootButton", edit = True, bgc = [.196, .475, .98])
+            
+        else:
+            cmds.button("rootButton", edit = True, bgc = [1,1,1])
+        
+            controls = [(namespace + ":" + "Root_Anim")]
+            buttons = ["rootButton"]
+        
+            createScriptJob(controls, buttons)
+        
+        
+    if (mods & 1) == 0:
+        cmds.select((namespace + ":" + "Root_Anim"), replace = True)
+        cmds.button("rootButton", edit = True, bgc = [1,1,1])
+        
+        controls = [(namespace + ":" + "Root_Anim")]
+        buttons = ["rootButton"]
+        
+        createScriptJob(controls, buttons)
+        
+def selectAttach(*args):
+    namespace = cmds.optionMenu("selectedCharacterOptionMenu", q = True, v = True)
+    mods = cmds.getModifiers()
+    
+    if (mods & 1) > 0:
+        cmds.select((namespace + ":" + "Attach_Anim"), tgl = True)
+        
+        buttonVal = cmds.button("attachButton", q = True, bgc = True)
+        
+        if buttonVal[1] == 1.0:
+            cmds.button("attachButton", edit = True, bgc = [.557, 0, .988])
+            
+        else:
+            cmds.button("attachButton", edit = True, bgc = [1,1,1])
+        
+            controls = [(namespace + ":" + "Attach_Anim")]
+            buttons = ["attachButton"]
+        
+            createScriptJob(controls, buttons)
+        
+        
+    if (mods & 1) == 0:
+        cmds.select((namespace + ":" + "Attach_Anim"), replace = True)
+        cmds.button("attachButton", edit = True, bgc = [1,1,1])
+        
+        controls = [(namespace + ":" + "Attach_Anim")]
+        buttons = ["attachButton"]
+        
+        createScriptJob(controls, buttons)
+        
+def selectLShoulder(*args):
+    namespace = cmds.optionMenu("selectedCharacterOptionMenu", q = True, v = True)
+    mods = cmds.getModifiers()
+    
+    if (mods & 1) > 0:
+        cmds.select((namespace + ":" + "L_Shoulder_Anim"), tgl = True)
+        
+        buttonVal = cmds.button("lShoulderButton", q = True, bgc = True)
+        
+        if buttonVal[1] == 1.0:
+            cmds.button("lShoulderButton", edit = True, bgc = [.196, .475, .98])
+            
+        else:
+            cmds.button("lShoulderButton", edit = True, bgc = [1,1,1])
+        
+            controls = [(namespace + ":" + "L_Shoulder_Anim")]
+            buttons = ["lShoulderButton"]
+        
+            createScriptJob(controls, buttons)
+        
+        
+    if (mods & 1) == 0:
+        cmds.select((namespace + ":" + "L_Shoulder_Anim"), replace = True)
+        cmds.button("lShoulderButton", edit = True, bgc = [1,1,1])
+        
+        controls = [(namespace + ":" + "L_Shoulder_Anim")]
+        buttons = ["lShoulderButton"]
+        
+        createScriptJob(controls, buttons)
+        
+def selectRShoulder(*args):
+    namespace = cmds.optionMenu("selectedCharacterOptionMenu", q = True, v = True)
+    mods = cmds.getModifiers()
+    
+    if (mods & 1) > 0:
+        cmds.select((namespace + ":" + "R_Shoulder_Anim"), tgl = True)
+        
+        buttonVal = cmds.button("rShoulderButton", q = True, bgc = True)
+        
+        if buttonVal[1] == 1.0:
+            cmds.button("rShoulderButton", edit = True, bgc = [.196, .475, .98])
+            
+        else:
+            cmds.button("rShoulderButton", edit = True, bgc = [1,1,1])
+        
+            controls = [(namespace + ":" + "R_Shoulder_Anim")]
+            buttons = ["rShoulderButton"]
+        
+            createScriptJob(controls, buttons)
+        
+        
+    if (mods & 1) == 0:
+        cmds.select((namespace + ":" + "R_Shoulder_Anim"), replace = True)
+        cmds.button("rShoulderButton", edit = True, bgc = [1,1,1])
+        
+        controls = [(namespace + ":" + "R_Shoulder_Anim")]
+        buttons = ["rShoulderButton"]
+        
+        createScriptJob(controls, buttons)
+        
+def selectLElbow(*args):
+    namespace = cmds.optionMenu("selectedCharacterOptionMenu", q = True, v = True)
+    mods = cmds.getModifiers()
+    
+    if (mods & 1) > 0:
+        cmds.select((namespace + ":" + "L_IK_Elbow_Anim"), tgl = True)
+        
+        buttonVal = cmds.button("lElbowButton", q = True, bgc = True)
+        
+        if buttonVal[1] == 1.0:
+            cmds.button("lElbowButton", edit = True, bgc = [1, 0, 0])
+            
+        else:
+            cmds.button("lElbowButton", edit = True, bgc = [1,1,1])
+        
+            controls = [(namespace + ":" + "L_IK_Elbow_Anim")]
+            buttons = ["lElbowButton"]
+        
+            createScriptJob(controls, buttons)
+        
+        
+    if (mods & 1) == 0:
+        cmds.select((namespace + ":" + "L_IK_Elbow_Anim"), replace = True)
+        cmds.button("lElbowButton", edit = True, bgc = [1,1,1])
+        
+        controls = [(namespace + ":" + "L_IK_Elbow_Anim")]
+        buttons = ["lElbowButton"]
+        
+        createScriptJob(controls, buttons)
+        
+def selectRElbow(*args):
+    namespace = cmds.optionMenu("selectedCharacterOptionMenu", q = True, v = True)
+    mods = cmds.getModifiers()
+    
+    if (mods & 1) > 0:
+        cmds.select((namespace + ":" + "R_IK_Elbow_Anim"), tgl = True)
+        
+        buttonVal = cmds.button("rElbowButton", q = True, bgc = True)
+        
+        if buttonVal[1] == 1.0:
+            cmds.button("rElbowButton", edit = True, bgc = [1, 0, 0])
+            
+        else:
+            cmds.button("rElbowButton", edit = True, bgc = [1,1,1])
+        
+            controls = [(namespace + ":" + "R_IK_Elbow_Anim")]
+            buttons = ["rElbowButton"]
+        
+            createScriptJob(controls, buttons)
+        
+        
+    if (mods & 1) == 0:
+        cmds.select((namespace + ":" + "R_IK_Elbow_Anim"), replace = True)
+        cmds.button("rElbowButton", edit = True, bgc = [1,1,1])
+        
+        controls = [(namespace + ":" + "R_IK_Elbow_Anim")]
+        buttons = ["rElbowButton"]
+        
+        createScriptJob(controls, buttons)
+        
+def selectLIKHand(*args):
+    namespace = cmds.optionMenu("selectedCharacterOptionMenu", q = True, v = True)
+    mods = cmds.getModifiers()
+    
+    if (mods & 1) > 0:
+        cmds.select((namespace + ":" + "L_IK_Hand_Anim"), tgl = True)
+        
+        buttonVal = cmds.button("lIKHandButton", q = True, bgc = True)
+        
+        if buttonVal[1] == 1.0:
+            cmds.button("lIKHandButton", edit = True, bgc = [1, 0, 0])
+            
+        else:
+            cmds.button("lIKHandButton", edit = True, bgc = [1,1,1])
+        
+            controls = [(namespace + ":" + "L_IK_Hand_Anim")]
+            buttons = ["lIKHandButton"]
+        
+            createScriptJob(controls, buttons)
+        
+        
+    if (mods & 1) == 0:
+        cmds.select((namespace + ":" + "L_IK_Hand_Anim"), replace = True)
+        cmds.button("lIKHandButton", edit = True, bgc = [1,1,1])
+        
+        controls = [(namespace + ":" + "L_IK_Hand_Anim")]
+        buttons = ["lIKHandButton"]
+        
+        createScriptJob(controls, buttons)
+        
+        
+def selectRIKHand(*args):
+    namespace = cmds.optionMenu("selectedCharacterOptionMenu", q = True, v = True)
+    mods = cmds.getModifiers()
+    
+    if (mods & 1) > 0:
+        cmds.select((namespace + ":" + "R_IK_Hand_Anim"), tgl = True)
+        
+        buttonVal = cmds.button("rIKHandButton", q = True, bgc = True)
+        
+        if buttonVal[1] == 1.0:
+            cmds.button("rIKHandButton", edit = True, bgc = [1, 0, 0])
+            
+        else:
+            cmds.button("rIKHandButton", edit = True, bgc = [1,1,1])
+        
+            controls = [(namespace + ":" + "R_IK_Hand_Anim")]
+            buttons = ["rIKHandButton"]
+        
+            createScriptJob(controls, buttons)
+        
+        
+    if (mods & 1) == 0:
+        cmds.select((namespace + ":" + "R_IK_Hand_Anim"), replace = True)
+        cmds.button("rIKHandButton", edit = True, bgc = [1,1,1])
+        
+        controls = [(namespace + ":" + "R_IK_Hand_Anim")]
+        buttons = ["rIKHandButton"]
+        
+        createScriptJob(controls, buttons)
+        
+def selectLFKHand(*args):
+    namespace = cmds.optionMenu("selectedCharacterOptionMenu", q = True, v = True)
+    mods = cmds.getModifiers()
+    
+    if (mods & 1) > 0:
+        cmds.select((namespace + ":" + "FK_b_MF_Hand_L"), tgl = True)
+        
+        buttonVal = cmds.button("lFKHandButton", q = True, bgc = True)
+        
+        if buttonVal[1] == 1.0:
+            cmds.button("lFKHandButton", edit = True, bgc = [.196, .475, .98])
+            
+        else:
+            cmds.button("lFKHandButton", edit = True, bgc = [1,1,1])
+        
+            controls = [(namespace + ":" + "FK_b_MF_Hand_L")]
+            buttons = ["lFKHandButton"]
+        
+            createScriptJob(controls, buttons)
+        
+        
+    if (mods & 1) == 0:
+        cmds.select((namespace + ":" + "FK_b_MF_Hand_L"), replace = True)
+        cmds.button("lFKHandButton", edit = True, bgc = [1,1,1])
+        
+        controls = [(namespace + ":" + "FK_b_MF_Hand_L")]
+        buttons = ["lFKHandButton"]
+        
+        createScriptJob(controls, buttons)
+        
+def selectRFKHand(*args):
+    namespace = cmds.optionMenu("selectedCharacterOptionMenu", q = True, v = True)
+    mods = cmds.getModifiers()
+    
+    if (mods & 1) > 0:
+        cmds.select((namespace + ":" + "FK_b_MF_Hand_R"), tgl = True)
+        
+        buttonVal = cmds.button("rFKHandButton", q = True, bgc = True)
+        
+        if buttonVal[1] == 1.0:
+            cmds.button("rFKHandButton", edit = True, bgc = [.196, .475, .98])
+            
+        else:
+            cmds.button("rFKHandButton", edit = True, bgc = [1,1,1])
+        
+            controls = [(namespace + ":" + "FK_b_MF_Hand_R")]
+            buttons = ["rFKHandButton"]
+        
+            createScriptJob(controls, buttons)
+        
+        
+    if (mods & 1) == 0:
+        cmds.select((namespace + ":" + "FK_b_MF_Hand_R"), replace = True)
+        cmds.button("rFKHandButton", edit = True, bgc = [1,1,1])
+        
+        controls = [(namespace + ":" + "FK_b_MF_Hand_R")]
+        buttons = ["rFKHandButton"]
+        
+        createScriptJob(controls, buttons)
+        
+def selectLFKElbow(*args):
+    namespace = cmds.optionMenu("selectedCharacterOptionMenu", q = True, v = True)
+    mods = cmds.getModifiers()
+    
+    if (mods & 1) > 0:
+        cmds.select((namespace + ":" + "FK_b_MF_Forearm_L"), tgl = True)
+        
+        buttonVal = cmds.button("lFKElbowButton", q = True, bgc = True)
+        
+        if buttonVal[1] == 1.0:
+            cmds.button("lFKElbowButton", edit = True, bgc = [.196, .475, .98])
+            
+        else:
+            cmds.button("lFKElbowButton", edit = True, bgc = [1,1,1])
+        
+            controls = [(namespace + ":" + "FK_b_MF_Forearm_L")]
+            buttons = ["lFKElbowButton"]
+        
+            createScriptJob(controls, buttons)
+        
+        
+    if (mods & 1) == 0:
+        cmds.select((namespace + ":" + "FK_b_MF_Forearm_L"), replace = True)
+        cmds.button("lFKElbowButton", edit = True, bgc = [1,1,1])
+        
+        controls = [(namespace + ":" + "FK_b_MF_Forearm_L")]
+        buttons = ["lFKElbowButton"]
+        
+        createScriptJob(controls, buttons)
+        
+def selectRFKElbow(*args):
+    namespace = cmds.optionMenu("selectedCharacterOptionMenu", q = True, v = True)
+    mods = cmds.getModifiers()
+    
+    if (mods & 1) > 0:
+        cmds.select((namespace + ":" + "FK_b_MF_Forearm_R"), tgl = True)
+        
+        buttonVal = cmds.button("rFKElbowButton", q = True, bgc = True)
+        
+        if buttonVal[1] == 1.0:
+            cmds.button("rFKElbowButton", edit = True, bgc = [.196, .475, .98])
+            
+        else:
+            cmds.button("rFKElbowButton", edit = True, bgc = [1,1,1])
+        
+            controls = [(namespace + ":" + "FK_b_MF_Forearm_R")]
+            buttons = ["rFKElbowButton"]
+        
+            createScriptJob(controls, buttons)
+        
+        
+    if (mods & 1) == 0:
+        cmds.select((namespace + ":" + "FK_b_MF_Forearm_R"), replace = True)
+        cmds.button("rFKElbowButton", edit = True, bgc = [1,1,1])
+        
+        controls = [(namespace + ":" + "FK_b_MF_Forearm_R")]
+        buttons = ["rFKElbowButton"]
+        
+        createScriptJob(controls, buttons)
+        
+        
+def selectLFKArm(*args):
+    namespace = cmds.optionMenu("selectedCharacterOptionMenu", q = True, v = True)
+    mods = cmds.getModifiers()
+    
+    if (mods & 1) > 0:
+        cmds.select((namespace + ":" + "FK_b_MF_UpperArm_L"), tgl = True)
+        
+        buttonVal = cmds.button("lFKArmButton", q = True, bgc = True)
+        
+        if buttonVal[1] == 1.0:
+            cmds.button("lFKArmButton", edit = True, bgc = [.196, .475, .98])
+            
+        else:
+            cmds.button("lFKArmButton", edit = True, bgc = [1,1,1])
+        
+            controls = [(namespace + ":" + "FK_b_MF_UpperArm_L")]
+            buttons = ["lFKArmButton"]
+        
+            createScriptJob(controls, buttons)
+        
+        
+    if (mods & 1) == 0:
+        cmds.select((namespace + ":" + "FK_b_MF_UpperArm_L"), replace = True)
+        cmds.button("lFKArmButton", edit = True, bgc = [1,1,1])
+        
+        controls = [(namespace + ":" + "FK_b_MF_UpperArm_L")]
+        buttons = ["lFKArmButton"]
+        
+        createScriptJob(controls, buttons)
+        
+def selectRFKArm(*args):
+    namespace = cmds.optionMenu("selectedCharacterOptionMenu", q = True, v = True)
+    mods = cmds.getModifiers()
+    
+    if (mods & 1) > 0:
+        cmds.select((namespace + ":" + "FK_b_MF_UpperArm_R"), tgl = True)
+        
+        buttonVal = cmds.button("rFKArmButton", q = True, bgc = True)
+        
+        if buttonVal[1] == 1.0:
+            cmds.button("rFKArmButton", edit = True, bgc = [.196, .475, .98])
+            
+        else:
+            cmds.button("rFKArmButton", edit = True, bgc = [1,1,1])
+        
+            controls = [(namespace + ":" + "FK_b_MF_UpperArm_R")]
+            buttons = ["rFKArmButton"]
+        
+            createScriptJob(controls, buttons)
+        
+        
+    if (mods & 1) == 0:
+        cmds.select((namespace + ":" + "FK_b_MF_UpperArm_R"), replace = True)
+        cmds.button("rFKArmButton", edit = True, bgc = [1,1,1])
+        
+        controls = [(namespace + ":" + "FK_b_MF_UpperArm_R")]
+        buttons = ["rFKArmButton"]
+        
+        createScriptJob(controls, buttons)
+        
+def selectLWep(*args):
+    namespace = cmds.optionMenu("selectedCharacterOptionMenu", q = True, v = True)
+    mods = cmds.getModifiers()
+    
+    if (mods & 1) > 0:
+        cmds.select((namespace + ":" + "L_Weapon_Anim"), tgl = True)
+        
+        buttonVal = cmds.button("lWepButton", q = True, bgc = True)
+        
+        if buttonVal[1] == 1.0:
+            cmds.button("lWepButton", edit = True, bgc = [.557, 0, .988])
+            
+        else:
+            cmds.button("lWepButton", edit = True, bgc = [1,1,1])
+        
+            controls = [(namespace + ":" + "L_Weapon_Anim")]
+            buttons = ["lWepButton"]
+        
+            createScriptJob(controls, buttons)
+        
+        
+    if (mods & 1) == 0:
+        cmds.select((namespace + ":" + "L_Weapon_Anim"), replace = True)
+        cmds.button("lWepButton", edit = True, bgc = [1,1,1])
+        
+        controls = [(namespace + ":" + "L_Weapon_Anim")]
+        buttons = ["lWepButton"]
+        
+        createScriptJob(controls, buttons)
+        
+        
+        
+def selectLancerStock(*args):
+    namespace = cmds.optionMenu("selectedCharacterOptionMenu", q = True, v = True)
+    mods = cmds.getModifiers()
+    
+    
+    referenceNamespace = ""
+    
+    cmds.select((namespace + ":R_Weapon_Anim"))
+    connections = cmds.listConnections(source = True)[0]
+    if connections.find("b_AR_Root") == 0:
+        parent = cmds.listRelatives(connections, parent = True)[0]
+        referenceNamespace = parent.partition(":")[0]
+        
+    cmds.select(clear = True)
+        
+    
+                
+    if (mods & 1) > 0:
+        if cmds.objExists((referenceNamespace + ":" + "Lancer_Stock_Loc")):
+            cmds.select((referenceNamespace + ":" + "Lancer_Stock_Loc"), tgl = True)
+            
+            buttonVal = cmds.button("LLancerStockButton", q = True, bgc = True)
+            
+            if buttonVal[1] == 1.0:
+                cmds.button("LLancerStockButton", edit = True, bgc = [.557, 0, .988])
+                
+            else:
+                cmds.button("LLancerStockButton", edit = True, bgc = [1,1,1])
+            
+                controls = [(referenceNamespace + ":" + "Lancer_Stock_Loc")]
+                buttons = ["LLancerStockButton"]
+            
+                createScriptJob(controls, buttons)
+        
+        
+    if (mods & 1) == 0:
+        if cmds.objExists((referenceNamespace + ":" + "Lancer_Stock_Loc")):
+            
+            cmds.select((referenceNamespace + ":" + "Lancer_Stock_Loc"), replace = True)
+            cmds.button("LLancerStockButton", edit = True, bgc = [1,1,1])
+            
+            controls = [(referenceNamespace + ":" + "Lancer_Stock_Loc")]
+            buttons = ["LLancerStockButton"]
+            
+            createScriptJob(controls, buttons)
+            
+            
+def selectLWepData(*args):
+    namespace = cmds.optionMenu("selectedCharacterOptionMenu", q = True, v = True)
+    mods = cmds.getModifiers()
+    
+  
+    if (mods & 1) > 0:
+        if cmds.objExists((namespace + ":" + "L_Weapon_Loc")):
+            cmds.select((namespace + ":" + "L_Weapon_Loc"), tgl = True)
+            
+            buttonVal = cmds.button("lWepDataButton", q = True, bgc = True)
+            
+            if buttonVal[1] == 1.0:
+                cmds.button("lWepDataButton", edit = True, bgc = [.557, 0, .988])
+                
+            else:
+                cmds.button("lWepDataButton", edit = True, bgc = [1,1,1])
+            
+                controls = [(namespace + ":" + "L_Weapon_Loc")]
+                buttons = ["lWepDataButton"]
+            
+                createScriptJob(controls, buttons)
+        
+        
+    if (mods & 1) == 0:
+        if cmds.objExists((namespace + ":" + "L_Weapon_Loc")):
+            
+            cmds.select((namespace + ":" + "L_Weapon_Loc"), replace = True)
+            cmds.button("lWepDataButton", edit = True, bgc = [1,1,1])
+            
+            controls = [(namespace + ":" + "L_Weapon_Loc")]
+            buttons = ["lWepDataButton"]
+            
+            createScriptJob(controls, buttons)
+            
+            
+def selectRWepData(*args):
+    namespace = cmds.optionMenu("selectedCharacterOptionMenu", q = True, v = True)
+    mods = cmds.getModifiers()
+    
+  
+    if (mods & 1) > 0:
+        if cmds.objExists((namespace + ":" + "R_Weapon_Loc")):
+            cmds.select((namespace + ":" + "R_Weapon_Loc"), tgl = True)
+            
+            buttonVal = cmds.button("rWepDataButton", q = True, bgc = True)
+            
+            if buttonVal[1] == 1.0:
+                cmds.button("rWepDataButton", edit = True, bgc = [.557, 0, .988])
+                
+            else:
+                cmds.button("rWepDataButton", edit = True, bgc = [1,1,1])
+            
+                controls = [(namespace + ":" + "R_Weapon_Loc")]
+                buttons = ["rWepDataButton"]
+            
+                createScriptJob(controls, buttons)
+        
+        
+    if (mods & 1) == 0:
+        if cmds.objExists((namespace + ":" + "R_Weapon_Loc")):
+            
+            cmds.select((namespace + ":" + "R_Weapon_Loc"), replace = True)
+            cmds.button("rWepDataButton", edit = True, bgc = [1,1,1])
+            
+            controls = [(namespace + ":" + "R_Weapon_Loc")]
+            buttons = ["rWepDataButton"]
+            
+            createScriptJob(controls, buttons)
+            
+        
+        
+        
+def selectRWep(*args):
+    namespace = cmds.optionMenu("selectedCharacterOptionMenu", q = True, v = True)
+    mods = cmds.getModifiers()
+    
+    if (mods & 1) > 0:
+        cmds.select((namespace + ":" + "R_Weapon_Anim"), tgl = True)
+        
+        buttonVal = cmds.button("rWepButton", q = True, bgc = True)
+        
+        if buttonVal[1] == 1.0:
+            cmds.button("rWepButton", edit = True, bgc = [.557, 0, .988])
+            
+        else:
+            cmds.button("rWepButton", edit = True, bgc = [1,1,1])
+        
+            controls = [(namespace + ":" + "R_Weapon_Anim")]
+            buttons = ["rWepButton"]
+        
+            createScriptJob(controls, buttons)
+        
+        
+    if (mods & 1) == 0:
+        cmds.select((namespace + ":" + "R_Weapon_Anim"), replace = True)
+        cmds.button("rWepButton", edit = True, bgc = [1,1,1])
+        
+        controls = [(namespace + ":" + "R_Weapon_Anim")]
+        buttons = ["rWepButton"]
+        
+        createScriptJob(controls, buttons)
+        
+def selectNeck(*args):
+    namespace = cmds.optionMenu("selectedCharacterOptionMenu", q = True, v = True)
+    mods = cmds.getModifiers()
+    
+    if (mods & 1) > 0:
+        cmds.select((namespace + ":" + "Neck_Anim"), tgl = True)
+        
+        buttonVal = cmds.button("neckButton", q = True, bgc = True)
+        
+        if buttonVal[1] == 1.0:
+            cmds.button("neckButton", edit = True, bgc = [.196, .475, .98])
+            
+        else:
+            cmds.button("neckButton", edit = True, bgc = [1,1,1])
+        
+            controls = [(namespace + ":" + "Neck_Anim")]
+            buttons = ["neckButton"]
+        
+            createScriptJob(controls, buttons)
+        
+        
+    if (mods & 1) == 0:
+        cmds.select((namespace + ":" + "Neck_Anim"), replace = True)
+        cmds.button("neckButton", edit = True, bgc = [1,1,1])
+        
+        controls = [(namespace + ":" + "Neck_Anim")]
+        buttons = ["neckButton"]
+        
+        createScriptJob(controls, buttons)
+        
+def selectHead(*args):
+    namespace = cmds.optionMenu("selectedCharacterOptionMenu", q = True, v = True)
+    mods = cmds.getModifiers()
+    
+    if (mods & 1) > 0:
+        cmds.select((namespace + ":" + "Head_Anim"), tgl = True)
+        
+        buttonVal = cmds.button("headButton", q = True, bgc = True)
+        
+        if buttonVal[1] == 1.0:
+            cmds.button("headButton", edit = True, bgc = [.196, .475, .98])
+            
+        else:
+            cmds.button("headButton", edit = True, bgc = [1,1,1])
+        
+            controls = [(namespace + ":" + "Head_Anim")]
+            buttons = ["headButton"]
+            
+            createScriptJob(controls, buttons)
+        
+        
+    if (mods & 1) == 0:
+        cmds.select((namespace + ":" + "Head_Anim"), replace = True)
+        cmds.button("headButton", edit = True, bgc = [1,1,1])
+        
+        controls = [(namespace + ":" + "Head_Anim")]
+        buttons = ["headButton"]
+        
+        createScriptJob(controls, buttons)
+        
+        
+def selectSpine1(*args):
+    namespace = cmds.optionMenu("selectedCharacterOptionMenu", q = True, v = True)
+    mods = cmds.getModifiers()
+    
+    if (mods & 1) > 0:
+        cmds.select((namespace + ":" + "Spine_01_Anim"), tgl = True)
+        
+        buttonVal = cmds.button("spine1Button", q = True, bgc = True)
+        
+        if buttonVal[1] == 1.0:
+            cmds.button("spine1Button", edit = True, bgc = [.196, .475, .98])
+            
+        else:
+            cmds.button("spine1Button", edit = True, bgc = [1.0,1.0,1.0])
+
+            controls = [(namespace + ":" + "Spine_01_Anim")]
+            buttons = ["spine1Button"]
+            createScriptJob(controls, buttons)
+            
+        
+        
+    if (mods & 1) == 0:
+        
+        cmds.select((namespace + ":" + "Spine_01_Anim"), replace = True)
+        cmds.button("spine1Button", edit = True, bgc = [1,1,1])
+        
+        controls = [(namespace + ":" + "Spine_01_Anim")]
+        buttons = ["spine1Button"]
+        
+        createScriptJob(controls, buttons)
+        
+        
+def selectSpine2(*args):
+    namespace = cmds.optionMenu("selectedCharacterOptionMenu", q = True, v = True)
+    mods = cmds.getModifiers()
+    
+    if (mods & 1) > 0:
+        cmds.select((namespace + ":" + "Spine_02_Anim"), tgl = True)
+        
+        buttonVal = cmds.button("spine2Button", q = True, bgc = True)
+        
+        if buttonVal[1] == 1.0:
+            cmds.button("spine2Button", edit = True, bgc = [.196, .475, .98])
+            
+        else:
+            cmds.button("spine2Button", edit = True, bgc = [1,1,1])
+
+            controls = [(namespace + ":" + "Spine_02_Anim")]
+            buttons = ["spine2Button"]
+        
+            createScriptJob(controls, buttons)
+        
+    if (mods & 1) == 0:
+        cmds.select((namespace + ":" + "Spine_02_Anim"), replace = True)
+        cmds.button("spine2Button", edit = True, bgc = [1,1,1])
+        
+        controls = [(namespace + ":" + "Spine_02_Anim")]
+        buttons = ["spine2Button"]
+        
+        createScriptJob(controls, buttons)
+        
+def selectSpine3(*args):
+    namespace = cmds.optionMenu("selectedCharacterOptionMenu", q = True, v = True)
+    mods = cmds.getModifiers()
+    
+    if (mods & 1) > 0:
+        cmds.select((namespace + ":" + "Spine_03_Anim"), tgl = True)
+        
+        buttonVal = cmds.button("spine3Button", q = True, bgc = True)
+        
+        if buttonVal[1] == 1.0:
+            cmds.button("spine3Button", edit = True, bgc = [.196, .475, .98])
+            
+        else:
+            cmds.button("spine3Button", edit = True, bgc = [1,1,1])
+
+            controls = [(namespace + ":" + "Spine_03_Anim")]
+            buttons = ["spine3Button"]
+        
+            createScriptJob(controls, buttons)
+        
+    if (mods & 1) == 0:
+        cmds.select((namespace + ":" + "Spine_03_Anim"), replace = True)
+        cmds.button("spine3Button", edit = True, bgc = [1,1,1])
+        
+        controls = [(namespace + ":" + "Spine_03_Anim")]
+        buttons = ["spine3Button"]
+        
+        createScriptJob(controls, buttons)
+        
+def selectHips(*args):
+    
+    namespace = cmds.optionMenu("selectedCharacterOptionMenu", q = True, v = True)
+    mods = cmds.getModifiers()
+    
+    if (mods & 1) > 0:
+        cmds.select((namespace + ":" + "Hip_Anim"), tgl = True)
+        
+        buttonVal = cmds.button("hipButton", q = True, bgc = True)
+        
+        if buttonVal[1] == 1.0:
+            cmds.button("hipButton", edit = True, bgc = [.196, .475, .98])
+            
+        else:
+            cmds.button("hipButton", edit = True, bgc = [1,1,1])
+
+            controls = [(namespace + ":" + "Hip_Anim")]
+            buttons = ["hipButton"]
+        
+            createScriptJob(controls, buttons)
+        
+    if (mods & 1) == 0:
+        cmds.select((namespace + ":" + "Hip_Anim"), replace = True)
+        cmds.button("hipButton", edit = True, bgc = [1,1,1])
+        
+        controls = [(namespace + ":" + "Hip_Anim")]
+        buttons = ["hipButton"]
+        
+        createScriptJob(controls, buttons)
+        
+        
+def selectLFKThigh(*args):
+    namespace = cmds.optionMenu("selectedCharacterOptionMenu", q = True, v = True)
+    mods = cmds.getModifiers()
+    
+    if (mods & 1) > 0:
+        cmds.select((namespace + ":" + "FK_b_MF_Thigh_L"), tgl = True)
+        
+        buttonVal = cmds.button("lFKThighButton", q = True, bgc = True)
+        
+        if buttonVal[1] == 1.0:
+            cmds.button("lFKThighButton", edit = True, bgc = [.196, .475, .98])
+            
+        else:
+            cmds.button("lFKThighButton", edit = True, bgc = [1,1,1])
+        
+            controls = [(namespace + ":" + "FK_b_MF_Thigh_L")]
+            buttons = ["lFKThighButton"]
+        
+            createScriptJob(controls, buttons)
+        
+        
+    if (mods & 1) == 0:
+        cmds.select((namespace + ":" + "FK_b_MF_Thigh_L"), replace = True)
+        cmds.button("lFKThighButton", edit = True, bgc = [1,1,1])
+        
+        controls = [(namespace + ":" + "FK_b_MF_Thigh_L")]
+        buttons = ["lFKThighButton"]
+        
+        createScriptJob(controls, buttons)
+        
+def selectRFKThigh(*args):
+    namespace = cmds.optionMenu("selectedCharacterOptionMenu", q = True, v = True)
+    mods = cmds.getModifiers()
+    
+    if (mods & 1) > 0:
+        cmds.select((namespace + ":" + "FK_b_MF_Thigh_R"), tgl = True)
+        
+        buttonVal = cmds.button("rFKThighButton", q = True, bgc = True)
+        
+        if buttonVal[1] == 1.0:
+            cmds.button("rFKThighButton", edit = True, bgc = [.196, .475, .98])
+            
+        else:
+            cmds.button("rFKThighButton", edit = True, bgc = [1,1,1])
+        
+            controls = [(namespace + ":" + "FK_b_MF_Thigh_R")]
+            buttons = ["rFKThighButton"]
+        
+            createScriptJob(controls, buttons)
+        
+        
+    if (mods & 1) == 0:
+        cmds.select((namespace + ":" + "FK_b_MF_Thigh_R"), replace = True)
+        cmds.button("rFKThighButton", edit = True, bgc = [1,1,1])
+        
+        controls = [(namespace + ":" + "FK_b_MF_Thigh_R")]
+        buttons = ["rFKThighButton"]
+        
+        createScriptJob(controls, buttons)
+        
+        
+def selectLFKCalf(*args):
+    namespace = cmds.optionMenu("selectedCharacterOptionMenu", q = True, v = True)
+    mods = cmds.getModifiers()
+    
+    if (mods & 1) > 0:
+        cmds.select((namespace + ":" + "FK_b_MF_Calf_L"), tgl = True)
+        
+        buttonVal = cmds.button("lFKCalfButton", q = True, bgc = True)
+        
+        if buttonVal[1] == 1.0:
+            cmds.button("lFKCalfButton", edit = True, bgc = [.196, .475, .98])
+            
+        else:
+            cmds.button("lFKCalfButton", edit = True, bgc = [1,1,1])
+        
+            controls = [(namespace + ":" + "FK_b_MF_Calf_L")]
+            buttons = ["lFKCalfButton"]
+        
+            createScriptJob(controls, buttons)
+        
+        
+    if (mods & 1) == 0:
+        cmds.select((namespace + ":" + "FK_b_MF_Calf_L"), replace = True)
+        cmds.button("lFKCalfButton", edit = True, bgc = [1,1,1])
+        
+        controls = [(namespace + ":" + "FK_b_MF_Calf_L")]
+        buttons = ["lFKCalfButton"]
+        
+        createScriptJob(controls, buttons)
+        
+def selectRFKCalf(*args):
+    namespace = cmds.optionMenu("selectedCharacterOptionMenu", q = True, v = True)
+    mods = cmds.getModifiers()
+    
+    if (mods & 1) > 0:
+        cmds.select((namespace + ":" + "FK_b_MF_Calf_R"), tgl = True)
+        
+        buttonVal = cmds.button("rFKCalfButton", q = True, bgc = True)
+        
+        if buttonVal[1] == 1.0:
+            cmds.button("rFKCalfButton", edit = True, bgc = [.196, .475, .98])
+            
+        else:
+            cmds.button("rFKCalfButton", edit = True, bgc = [1,1,1])
+        
+            controls = [(namespace + ":" + "FK_b_MF_Calf_R")]
+            buttons = ["rFKCalfButton"]
+        
+            createScriptJob(controls, buttons)
+        
+        
+    if (mods & 1) == 0:
+        cmds.select((namespace + ":" + "FK_b_MF_Calf_R"), replace = True)
+        cmds.button("rFKCalfButton", edit = True, bgc = [1,1,1])
+        
+        controls = [(namespace + ":" + "FK_b_MF_Calf_R")]
+        buttons = ["rFKCalfButton"]
+        
+        createScriptJob(controls, buttons)
+        
+def selectLFKFoot(*args):
+    namespace = cmds.optionMenu("selectedCharacterOptionMenu", q = True, v = True)
+    mods = cmds.getModifiers()
+    
+    if (mods & 1) > 0:
+        cmds.select((namespace + ":" + "FK_b_MF_Foot_L"), tgl = True)
+        
+        buttonVal = cmds.button("lFKFootButton", q = True, bgc = True)
+        
+        if buttonVal[1] == 1.0:
+            cmds.button("lFKFootButton", edit = True, bgc = [.196, .475, .98])
+            
+        else:
+            cmds.button("lFKFootButton", edit = True, bgc = [1,1,1])
+        
+            controls = [(namespace + ":" + "FK_b_MF_Foot_L")]
+            buttons = ["lFKFootButton"]
+        
+            createScriptJob(controls, buttons)
+        
+        
+    if (mods & 1) == 0:
+        cmds.select((namespace + ":" + "FK_b_MF_Foot_L"), replace = True)
+        cmds.button("lFKFootButton", edit = True, bgc = [1,1,1])
+        
+        controls = [(namespace + ":" + "FK_b_MF_Foot_L")]
+        buttons = ["lFKFootButton"]
+        
+        createScriptJob(controls, buttons)
+        
+def selectRFKFoot(*args):
+    namespace = cmds.optionMenu("selectedCharacterOptionMenu", q = True, v = True)
+    mods = cmds.getModifiers()
+    
+    if (mods & 1) > 0:
+        cmds.select((namespace + ":" + "FK_b_MF_Foot_R"), tgl = True)
+        
+        buttonVal = cmds.button("rFKFootButton", q = True, bgc = True)
+        
+        if buttonVal[1] == 1.0:
+            cmds.button("rFKFootButton", edit = True, bgc = [.196, .475, .98])
+            
+        else:
+            cmds.button("rFKFootButton", edit = True, bgc = [1,1,1])
+        
+            controls = [(namespace + ":" + "FK_b_MF_Foot_R")]
+            buttons = ["rFKFootButton"]
+        
+            createScriptJob(controls, buttons)
+        
+        
+    if (mods & 1) == 0:
+        cmds.select((namespace + ":" + "FK_b_MF_Foot_R"), replace = True)
+        cmds.button("rFKFootButton", edit = True, bgc = [1,1,1])
+        
+        controls = [(namespace + ":" + "FK_b_MF_Foot_R")]
+        buttons = ["rFKFootButton"]
+        
+        createScriptJob(controls, buttons)
+        
+        
+def selectLFKToe(*args):
+    namespace = cmds.optionMenu("selectedCharacterOptionMenu", q = True, v = True)
+    mods = cmds.getModifiers()
+    
+    if (mods & 1) > 0:
+        cmds.select((namespace + ":" + "FK_b_MF_Toe_L"), tgl = True)
+        
+        buttonVal = cmds.button("lFKToeButton", q = True, bgc = True)
+        
+        if buttonVal[1] == 1.0:
+            cmds.button("lFKToeButton", edit = True, bgc = [.196, .475, .98])
+            
+        else:
+            cmds.button("lFKToeButton", edit = True, bgc = [1,1,1])
+        
+            controls = [(namespace + ":" + "FK_b_MF_Toe_L")]
+            buttons = ["lFKToeButton"]
+        
+            createScriptJob(controls, buttons)
+        
+        
+    if (mods & 1) == 0:
+        cmds.select((namespace + ":" + "FK_b_MF_Toe_L"), replace = True)
+        cmds.button("lFKToeButton", edit = True, bgc = [1,1,1])
+        
+        controls = [(namespace + ":" + "FK_b_MF_Toe_L")]
+        buttons = ["lFKToeButton"]
+        
+        createScriptJob(controls, buttons)
+        
+def selectRFKToe(*args):
+    namespace = cmds.optionMenu("selectedCharacterOptionMenu", q = True, v = True)
+    mods = cmds.getModifiers()
+    
+    if (mods & 1) > 0:
+        cmds.select((namespace + ":" + "FK_b_MF_Toe_R"), tgl = True)
+        
+        buttonVal = cmds.button("rFKToeButton", q = True, bgc = True)
+        
+        if buttonVal[1] == 1.0:
+            cmds.button("rFKToeButton", edit = True, bgc = [.196, .475, .98])
+            
+        else:
+            cmds.button("rFKToeButton", edit = True, bgc = [1,1,1])
+        
+            controls = [(namespace + ":" + "FK_b_MF_Toe_R")]
+            buttons = ["rFKToeButton"]
+        
+            createScriptJob(controls, buttons)
+        
+        
+    if (mods & 1) == 0:
+        cmds.select((namespace + ":" + "FK_b_MF_Toe_R"), replace = True)
+        cmds.button("rFKToeButton", edit = True, bgc = [1,1,1])
+        
+        controls = [(namespace + ":" + "FK_b_MF_Toe_R")]
+        buttons = ["rFKToeButton"]
+        
+        createScriptJob(controls, buttons)
+        
+        
+def selectLIKKnee(*args):
+    namespace = cmds.optionMenu("selectedCharacterOptionMenu", q = True, v = True)
+    mods = cmds.getModifiers()
+    
+    if (mods & 1) > 0:
+        cmds.select((namespace + ":" + "Left_Knee_Anim"), tgl = True)
+        
+        buttonVal = cmds.button("lIKKneeButton", q = True, bgc = True)
+        
+        if buttonVal[1] == 1.0:
+            cmds.button("lIKKneeButton", edit = True, bgc = [1, 0, 0])
+            
+        else:
+            cmds.button("lIKKneeButton", edit = True, bgc = [1,1,1])
+        
+            controls = [(namespace + ":" + "Left_Knee_Anim")]
+            buttons = ["lIKKneeButton"]
+        
+            createScriptJob(controls, buttons)
+        
+        
+    if (mods & 1) == 0:
+        cmds.select((namespace + ":" + "Left_Knee_Anim"), replace = True)
+        cmds.button("lIKKneeButton", edit = True, bgc = [1,1,1])
+        
+        controls = [(namespace + ":" + "Left_Knee_Anim")]
+        buttons = ["lIKKneeButton"]
+        
+        createScriptJob(controls, buttons)
+        
+def selectRIKKnee(*args):
+    namespace = cmds.optionMenu("selectedCharacterOptionMenu", q = True, v = True)
+    mods = cmds.getModifiers()
+    
+    if (mods & 1) > 0:
+        cmds.select((namespace + ":" + "Right_Knee_Anim"), tgl = True)
+        
+        buttonVal = cmds.button("rIKKneeButton", q = True, bgc = True)
+        
+        if buttonVal[1] == 1.0:
+            cmds.button("rIKKneeButton", edit = True, bgc = [1, 0, 0])
+            
+        else:
+            cmds.button("rIKKneeButton", edit = True, bgc = [1,1,1])
+        
+            controls = [(namespace + ":" + "Right_Knee_Anim")]
+            buttons = ["rIKKneeButton"]
+        
+            createScriptJob(controls, buttons)
+        
+        
+    if (mods & 1) == 0:
+        cmds.select((namespace + ":" + "Right_Knee_Anim"), replace = True)
+        cmds.button("rIKKneeButton", edit = True, bgc = [1,1,1])
+        
+        controls = [(namespace + ":" + "Right_Knee_Anim")]
+        buttons = ["rIKKneeButton"]
+        
+        createScriptJob(controls, buttons)
+        
+        
+def selectLIKFoot(*args):
+    namespace = cmds.optionMenu("selectedCharacterOptionMenu", q = True, v = True)
+    mods = cmds.getModifiers()
+    
+    if (mods & 1) > 0:
+        cmds.select((namespace + ":" + "Left_Foot_Anim"), tgl = True)
+        
+        buttonVal = cmds.button("lIKFootButton", q = True, bgc = True)
+        
+        if buttonVal[1] == 1.0:
+            cmds.button("lIKFootButton", edit = True, bgc = [1, 0, 0])
+            
+        else:
+            cmds.button("lIKFootButton", edit = True, bgc = [1,1,1])
+        
+            controls = [(namespace + ":" + "Left_Foot_Anim")]
+            buttons = ["lIKFootButton"]
+        
+            createScriptJob(controls, buttons)
+        
+        
+    if (mods & 1) == 0:
+        cmds.select((namespace + ":" + "Left_Foot_Anim"), replace = True)
+        cmds.button("lIKFootButton", edit = True, bgc = [1,1,1])
+        
+        controls = [(namespace + ":" + "Left_Foot_Anim")]
+        buttons = ["lIKFootButton"]
+        
+        createScriptJob(controls, buttons)
+        
+        
+def selectRIKFoot(*args):
+    namespace = cmds.optionMenu("selectedCharacterOptionMenu", q = True, v = True)
+    mods = cmds.getModifiers()
+    
+    if (mods & 1) > 0:
+        cmds.select((namespace + ":" + "Right_Foot_Anim"), tgl = True)
+        
+        buttonVal = cmds.button("rIKFootButton", q = True, bgc = True)
+        
+        if buttonVal[1] == 1.0:
+            cmds.button("rIKFootButton", edit = True, bgc = [1, 0, 0])
+            
+        else:
+            cmds.button("rIKFootButton", edit = True, bgc = [1,1,1])
+        
+            controls = [(namespace + ":" + "Right_Foot_Anim")]
+            buttons = ["rIKFootButton"]
+        
+            createScriptJob(controls, buttons)
+        
+        
+    if (mods & 1) == 0:
+        cmds.select((namespace + ":" + "Right_Foot_Anim"), replace = True)
+        cmds.button("rIKFootButton", edit = True, bgc = [1,1,1])
+        
+        controls = [(namespace + ":" + "Right_Foot_Anim")]
+        buttons = ["rIKFootButton"]
+        
+        createScriptJob(controls, buttons)
+        
+        
+        
+def selectAllTorso(*args):
+    namespace = cmds.optionMenu("selectedCharacterOptionMenu", q = True, v = True)
+    mods = cmds.getModifiers()
+    
+    #if holding down shift
+    if (mods & 1) > 0:
+        cmds.select((namespace + ":" + "Spine_01_Anim"), tgl = True)
+        cmds.select((namespace + ":" + "Spine_02_Anim"), tgl = True)
+        cmds.select((namespace + ":" + "Spine_03_Anim"), tgl = True)
+        cmds.select((namespace + ":" + "Body_Anim"), tgl = True)
+        cmds.select((namespace + ":" + "Hip_Anim"), tgl = True)
+        
+        
+        cmds.button("spine1Button", edit = True, bgc = [1,1,1])
+        cmds.button("spine2Button", edit = True, bgc = [1,1,1])
+        cmds.button("spine3Button", edit = True, bgc = [1,1,1])
+        cmds.button("bodyButton", edit = True, bgc = [1,1,1])
+        cmds.button("hipButton", edit = True, bgc = [1,1,1])
+        
+        
+        #
+        controls = [(namespace + ":" + "Spine_01_Anim"),(namespace + ":" + "Spine_02_Anim"),(namespace + ":" + "Spine_03_Anim"),(namespace + ":" + "Body_Anim"),(namespace + ":" + "Hip_Anim")]
+        buttons = ["spine1Button","spine2Button","spine3Button","bodyButton","hipButton"]
+        
+        createScriptJob(controls, buttons)
+        
+        
+    #if not holding down a modifier
+    if (mods & 1) == 0:
+        cmds.select((namespace + ":" + "Spine_01_Anim"), replace = True)
+        cmds.select((namespace + ":" + "Spine_02_Anim"), add = True)
+        cmds.select((namespace + ":" + "Spine_03_Anim"), add = True)
+        cmds.select((namespace + ":" + "Body_Anim"), add = True)
+        cmds.select((namespace + ":" + "Hip_Anim"), add = True)
+        
+        
+        cmds.button("spine1Button", edit = True, bgc = [1,1,1])
+        cmds.button("spine2Button", edit = True, bgc = [1,1,1])
+        cmds.button("spine3Button", edit = True, bgc = [1,1,1])
+        cmds.button("bodyButton", edit = True, bgc = [1,1,1])
+        cmds.button("hipButton", edit = True, bgc = [1,1,1])
+        
+        
+        #
+        controls = [(namespace + ":" + "Spine_01_Anim"),(namespace + ":" + "Spine_02_Anim"),(namespace + ":" + "Spine_03_Anim"),(namespace + ":" + "Body_Anim"),(namespace + ":" + "Hip_Anim")]
+        buttons = ["spine1Button","spine2Button","spine3Button","bodyButton","hipButton"]
+        
+        createScriptJob(controls, buttons)
+        
+
+        
+
+                                
+                            
+def selectAllHead(*args):
+    namespace = cmds.optionMenu("selectedCharacterOptionMenu", q = True, v = True)
+    mods = cmds.getModifiers()
+    
+    #if holding down shift
+    if (mods & 1) > 0:
+        cmds.select((namespace + ":" + "Neck_Anim"), tgl = True)
+        cmds.select((namespace + ":" + "Head_Anim"), tgl = True)
+
+        
+        cmds.button("neckButton", edit = True, bgc = [1,1,1])
+        cmds.button("headButton", edit = True, bgc = [1,1,1])
+
+        
+        controls = [(namespace + ":" + "Neck_Anim"), (namespace + ":" + "Head_Anim")]
+        buttons = ["neckButton", "headButton"]
+        
+        createScriptJob(controls, buttons)
+        
+
+    #if not holding down a modifer
+    if (mods & 1) == 0:
+        cmds.select((namespace + ":" + "Neck_Anim"), replace = True)
+        cmds.select((namespace + ":" + "Head_Anim"), add = True)
+
+        
+        cmds.button("neckButton", edit = True, bgc = [1,1,1])
+        cmds.button("headButton", edit = True, bgc = [1,1,1])
+        
+        controls = [(namespace + ":" + "Neck_Anim"), (namespace + ":" + "Head_Anim")]
+        buttons = ["neckButton", "headButton"]
+        
+        createScriptJob(controls, buttons)
+        
+def selectAllLeftArm(*args):
+    namespace = cmds.optionMenu("selectedCharacterOptionMenu", q = True, v = True)
+    mods = cmds.getModifiers()
+    
+    #if holding down shift
+    if (mods & 1) > 0:
+        cmds.select((namespace + ":" + "L_Shoulder_Anim"), tgl = True)
+        cmds.select((namespace + ":" + "L_IK_Elbow_Anim"), tgl = True)
+        cmds.select((namespace + ":" + "L_IK_Hand_Anim"), tgl = True)
+        cmds.select((namespace + ":" + "FK_b_MF_UpperArm_L"), tgl = True)
+        cmds.select((namespace + ":" + "FK_b_MF_Forearm_L"), tgl = True)
+        cmds.select((namespace + ":" + "FK_b_MF_Hand_L"), tgl = True)
+        cmds.select((namespace + ":" + "L_Weapon_Anim"), tgl = True)
+        
+        cmds.button("lFKHandButton", edit = True, bgc = [1,1,1])
+        cmds.button("lFKElbowButton", edit = True, bgc = [1,1,1])
+        cmds.button("lFKArmButton", edit = True, bgc = [1,1,1])
+        cmds.button("lShoulderButton", edit = True, bgc = [1,1,1])
+        cmds.button("lElbowButton", edit = True, bgc = [1,1,1])
+        cmds.button("lIKHandButton", edit = True, bgc = [1,1,1])
+        cmds.button("lWepButton", edit = True, bgc = [1,1,1])
+
+        
+        
+        controls = [(namespace + ":" + "FK_b_MF_Hand_L"),(namespace + ":" + "FK_b_MF_Forearm_L"), (namespace + ":" + "FK_b_MF_UpperArm_L"),(namespace + ":" + "L_Shoulder_Anim"), (namespace + ":" + "L_IK_Elbow_Anim"),(namespace + ":" + "L_IK_Hand_Anim"),(namespace + ":" + "L_Weapon_Anim")]
+        buttons = ["lFKHandButton", "lFKElbowButton","lFKArmButton","lShoulderButton","lElbowButton","lIKHandButton","lWepButton"]
+        
+        createScriptJob(controls, buttons)
+        
+
+    #if not holding down a modifer
+    if (mods & 1) == 0:
+        cmds.select((namespace + ":" + "L_Shoulder_Anim"), replace = True)
+        cmds.select((namespace + ":" + "L_IK_Elbow_Anim"), add = True)
+        cmds.select((namespace + ":" + "L_IK_Hand_Anim"), add = True)
+        cmds.select((namespace + ":" + "FK_b_MF_UpperArm_L"), add = True)
+        cmds.select((namespace + ":" + "FK_b_MF_Forearm_L"), add = True)
+        cmds.select((namespace + ":" + "FK_b_MF_Hand_L"), add = True)
+        cmds.select((namespace + ":" + "L_Weapon_Anim"), add = True)
+        
+        cmds.button("lFKHandButton", edit = True, bgc = [1,1,1])
+        cmds.button("lFKElbowButton", edit = True, bgc = [1,1,1])
+        cmds.button("lFKArmButton", edit = True, bgc = [1,1,1])
+        cmds.button("lShoulderButton", edit = True, bgc = [1,1,1])
+        cmds.button("lElbowButton", edit = True, bgc = [1,1,1])
+        cmds.button("lIKHandButton", edit = True, bgc = [1,1,1])
+        cmds.button("lWepButton", edit = True, bgc = [1,1,1])
+
+        
+        
+        controls = [(namespace + ":" + "FK_b_MF_Hand_L"),(namespace + ":" + "FK_b_MF_Forearm_L"), (namespace + ":" + "FK_b_MF_UpperArm_L"),(namespace + ":" + "L_Shoulder_Anim"), (namespace + ":" + "L_IK_Elbow_Anim"),(namespace + ":" + "L_IK_Hand_Anim"),(namespace + ":" + "L_Weapon_Anim")]
+        buttons = ["lFKHandButton", "lFKElbowButton","lFKArmButton","lShoulderButton","lElbowButton","lIKHandButton","lWepButton"]
+        
+        createScriptJob(controls, buttons)
+        
+        
+def selectAllRightArm(*args):
+    namespace = cmds.optionMenu("selectedCharacterOptionMenu", q = True, v = True)
+    mods = cmds.getModifiers()
+    
+    #if holding down shift
+    if (mods & 1) > 0:
+        cmds.select((namespace + ":" + "R_Shoulder_Anim"), tgl = True)
+        cmds.select((namespace + ":" + "R_IK_Elbow_Anim"), tgl = True)
+        cmds.select((namespace + ":" + "R_IK_Hand_Anim"), tgl = True)
+        cmds.select((namespace + ":" + "FK_b_MF_UpperArm_R"), tgl = True)
+        cmds.select((namespace + ":" + "FK_b_MF_Forearm_R"), tgl = True)
+        cmds.select((namespace + ":" + "FK_b_MF_Hand_R"), tgl = True)
+        cmds.select((namespace + ":" + "R_Weapon_Anim"), tgl = True)
+        
+        cmds.button("rFKHandButton", edit = True, bgc = [1,1,1])
+        cmds.button("rFKElbowButton", edit = True, bgc = [1,1,1])
+        cmds.button("rFKArmButton", edit = True, bgc = [1,1,1])
+        cmds.button("rShoulderButton", edit = True, bgc = [1,1,1])
+        cmds.button("rElbowButton", edit = True, bgc = [1,1,1])
+        cmds.button("rIKHandButton", edit = True, bgc = [1,1,1])
+        cmds.button("rWepButton", edit = True, bgc = [1,1,1])
+
+        
+        
+        controls = [(namespace + ":" + "FK_b_MF_Hand_R"),(namespace + ":" + "FK_b_MF_Forearm_R"), (namespace + ":" + "FK_b_MF_UpperArm_R"),(namespace + ":" + "R_Shoulder_Anim"), (namespace + ":" + "R_IK_Elbow_Anim"),(namespace + ":" + "R_IK_Hand_Anim"),(namespace + ":" + "R_Weapon_Anim")]
+        buttons = ["rFKHandButton", "rFKElbowButton","rFKArmButton","rShoulderButton","rElbowButton","rIKHandButton","rWepButton"]
+        
+        createScriptJob(controls, buttons)
+        
+
+    #if not holding down a modifer
+    if (mods & 1) == 0:
+        cmds.select((namespace + ":" + "R_Shoulder_Anim"), replace = True)
+        cmds.select((namespace + ":" + "R_IK_Elbow_Anim"), add = True)
+        cmds.select((namespace + ":" + "R_IK_Hand_Anim"), add = True)
+        cmds.select((namespace + ":" + "FK_b_MF_UpperArm_R"), add = True)
+        cmds.select((namespace + ":" + "FK_b_MF_Forearm_R"), add = True)
+        cmds.select((namespace + ":" + "FK_b_MF_Hand_R"), add = True)
+        cmds.select((namespace + ":" + "R_Weapon_Anim"), add = True)
+        
+        cmds.button("rFKHandButton", edit = True, bgc = [1,1,1])
+        cmds.button("rFKElbowButton", edit = True, bgc = [1,1,1])
+        cmds.button("rFKArmButton", edit = True, bgc = [1,1,1])
+        cmds.button("rShoulderButton", edit = True, bgc = [1,1,1])
+        cmds.button("rElbowButton", edit = True, bgc = [1,1,1])
+        cmds.button("rIKHandButton", edit = True, bgc = [1,1,1])
+        cmds.button("rWepButton", edit = True, bgc = [1,1,1])
+
+        
+        
+        controls = [(namespace + ":" + "FK_b_MF_Hand_R"),(namespace + ":" + "FK_b_MF_Forearm_R"), (namespace + ":" + "FK_b_MF_UpperArm_R"),(namespace + ":" + "R_Shoulder_Anim"), (namespace + ":" + "R_IK_Elbow_Anim"),(namespace + ":" + "R_IK_Hand_Anim"),(namespace + ":" + "R_Weapon_Anim")]
+        buttons = ["rFKHandButton", "rFKElbowButton","rFKArmButton","rShoulderButton","rElbowButton","rIKHandButton","rWepButton"]
+        
+        createScriptJob(controls, buttons)
+        
+        
+def selectAllLeftLeg(*args):
+    namespace = cmds.optionMenu("selectedCharacterOptionMenu", q = True, v = True)
+    mods = cmds.getModifiers()
+    
+    #if holding down shift
+    if (mods & 1) > 0:
+        cmds.select((namespace + ":" + "FK_b_MF_Thigh_L"), tgl = True)
+        cmds.select((namespace + ":" + "FK_b_MF_Calf_L"), tgl = True)
+        cmds.select((namespace + ":" + "FK_b_MF_Foot_L"), tgl = True)
+        cmds.select((namespace + ":" + "FK_b_MF_Toe_L"), tgl = True)
+        cmds.select((namespace + ":" + "Left_Knee_Anim"), tgl = True)
+        cmds.select((namespace + ":" + "Left_Foot_Anim"), tgl = True)
+        
+        cmds.button("lFKThighButton", edit = True, bgc = [1,1,1])
+        cmds.button("lFKCalfButton", edit = True, bgc = [1,1,1])
+        cmds.button("lFKFootButton", edit = True, bgc = [1,1,1])
+        cmds.button("lFKToeButton", edit = True, bgc = [1,1,1])
+        cmds.button("lIKKneeButton", edit = True, bgc = [1,1,1])
+        cmds.button("lIKFootButton", edit = True, bgc = [1,1,1])
+
+        
+        
+        controls = [(namespace + ":" + "FK_b_MF_Thigh_L"),(namespace + ":" + "FK_b_MF_Calf_L"), (namespace + ":" + "FK_b_MF_Foot_L"),(namespace + ":" + "FK_b_MF_Toe_L"),(namespace + ":" + "Left_Knee_Anim"),(namespace + ":" + "Left_Foot_Anim")]
+        buttons = ["lFKThighButton", "lFKCalfButton","lFKFootButton","lFKToeButton","lIKKneeButton","lIKFootButton"]
+        
+        createScriptJob(controls, buttons)
+        
+
+    #if not holding down a modifer
+    if (mods & 1) == 0:
+        
+        cmds.select((namespace + ":" + "FK_b_MF_Thigh_L"), replace = True)
+        cmds.select((namespace + ":" + "FK_b_MF_Calf_L"), add = True)
+        cmds.select((namespace + ":" + "FK_b_MF_Foot_L"), add = True)
+        cmds.select((namespace + ":" + "FK_b_MF_Toe_L"), add = True)
+        cmds.select((namespace + ":" + "Left_Knee_Anim"), add = True)
+        cmds.select((namespace + ":" + "Left_Foot_Anim"), add = True)
+        
+        cmds.button("lFKThighButton", edit = True, bgc = [1,1,1])
+        cmds.button("lFKCalfButton", edit = True, bgc = [1,1,1])
+        cmds.button("lFKFootButton", edit = True, bgc = [1,1,1])
+        cmds.button("lFKToeButton", edit = True, bgc = [1,1,1])
+        cmds.button("lIKKneeButton", edit = True, bgc = [1,1,1])
+        cmds.button("lIKFootButton", edit = True, bgc = [1,1,1])
+
+        
+        
+        controls = [(namespace + ":" + "FK_b_MF_Thigh_L"),(namespace + ":" + "FK_b_MF_Calf_L"), (namespace + ":" + "FK_b_MF_Foot_L"),(namespace + ":" + "FK_b_MF_Toe_L"),(namespace + ":" + "Left_Knee_Anim"),(namespace + ":" + "Left_Foot_Anim")]
+        buttons = ["lFKThighButton", "lFKCalfButton","lFKFootButton","lFKToeButton","lIKKneeButton","lIKFootButton"]
+        
+        createScriptJob(controls, buttons)
+        
+        
+def selectAllRightLeg(*args):
+    namespace = cmds.optionMenu("selectedCharacterOptionMenu", q = True, v = True)
+    mods = cmds.getModifiers()
+    
+    #if holding down shift
+    if (mods & 1) > 0:
+        cmds.select((namespace + ":" + "FK_b_MF_Thigh_R"), tgl = True)
+        cmds.select((namespace + ":" + "FK_b_MF_Calf_R"), tgl = True)
+        cmds.select((namespace + ":" + "FK_b_MF_Foot_R"), tgl = True)
+        cmds.select((namespace + ":" + "FK_b_MF_Toe_R"), tgl = True)
+        cmds.select((namespace + ":" + "Right_Knee_Anim"), tgl = True)
+        cmds.select((namespace + ":" + "Right_Foot_Anim"), tgl = True)
+        
+        cmds.button("rFKThighButton", edit = True, bgc = [1,1,1])
+        cmds.button("rFKCalfButton", edit = True, bgc = [1,1,1])
+        cmds.button("rFKFootButton", edit = True, bgc = [1,1,1])
+        cmds.button("rFKToeButton", edit = True, bgc = [1,1,1])
+        cmds.button("rIKKneeButton", edit = True, bgc = [1,1,1])
+        cmds.button("rIKFootButton", edit = True, bgc = [1,1,1])
+
+        
+        
+        controls = [(namespace + ":" + "FK_b_MF_Thigh_R"),(namespace + ":" + "FK_b_MF_Calf_R"), (namespace + ":" + "FK_b_MF_Foot_R"),(namespace + ":" + "FK_b_MF_Toe_R"),(namespace + ":" + "Right_Knee_Anim"),(namespace + ":" + "Right_Foot_Anim")]
+        buttons = ["rFKThighButton", "rFKCalfButton","rFKFootButton","rFKToeButton","rIKKneeButton","rIKFootButton"]
+        
+        createScriptJob(controls, buttons)
+        
+
+    #if not holding down a modifer
+    if (mods & 1) == 0:
+        
+        cmds.select((namespace + ":" + "FK_b_MF_Thigh_R"), replace = True)
+        cmds.select((namespace + ":" + "FK_b_MF_Calf_R"), add = True)
+        cmds.select((namespace + ":" + "FK_b_MF_Foot_R"), add = True)
+        cmds.select((namespace + ":" + "FK_b_MF_Toe_R"), add = True)
+        cmds.select((namespace + ":" + "Right_Knee_Anim"), add = True)
+        cmds.select((namespace + ":" + "Right_Foot_Anim"), add = True)
+        
+        cmds.button("rFKThighButton", edit = True, bgc = [1,1,1])
+        cmds.button("rFKCalfButton", edit = True, bgc = [1,1,1])
+        cmds.button("rFKFootButton", edit = True, bgc = [1,1,1])
+        cmds.button("rFKToeButton", edit = True, bgc = [1,1,1])
+        cmds.button("rIKKneeButton", edit = True, bgc = [1,1,1])
+        cmds.button("rIKFootButton", edit = True, bgc = [1,1,1])
+
+        
+        
+        controls = [(namespace + ":" + "FK_b_MF_Thigh_R"),(namespace + ":" + "FK_b_MF_Calf_R"), (namespace + ":" + "FK_b_MF_Foot_R"),(namespace + ":" + "FK_b_MF_Toe_R"),(namespace + ":" + "Right_Knee_Anim"),(namespace + ":" + "Right_Foot_Anim")]
+        buttons = ["rFKThighButton", "rFKCalfButton","rFKFootButton","rFKToeButton","rIKKneeButton","rIKFootButton"]
+        
+        createScriptJob(controls, buttons)
+        
+        
+def selectAll(*args):
+    namespace = cmds.optionMenu("selectedCharacterOptionMenu", q = True, v = True)
+    mods = cmds.getModifiers()
+    
+    #if holding down shift
+    if (mods & 1) > 0:
+        
+        #CONTROLS
+        cmds.select((namespace + ":" + "FK_b_MF_Thigh_R"), tgl = True)
+        cmds.select((namespace + ":" + "FK_b_MF_Calf_R"), tgl = True)
+        cmds.select((namespace + ":" + "FK_b_MF_Foot_R"), tgl = True)
+        cmds.select((namespace + ":" + "FK_b_MF_Toe_R"), tgl = True)
+        cmds.select((namespace + ":" + "Right_Knee_Anim"), tgl = True)
+        cmds.select((namespace + ":" + "Right_Foot_Anim"), tgl = True)
+        
+        cmds.select((namespace + ":" + "FK_b_MF_Thigh_L"), tgl = True)
+        cmds.select((namespace + ":" + "FK_b_MF_Calf_L"), tgl = True)
+        cmds.select((namespace + ":" + "FK_b_MF_Foot_L"), tgl = True)
+        cmds.select((namespace + ":" + "FK_b_MF_Toe_L"), tgl = True)
+        cmds.select((namespace + ":" + "Left_Knee_Anim"), tgl = True)
+        cmds.select((namespace + ":" + "Left_Foot_Anim"), tgl = True)
+        
+        cmds.select((namespace + ":" + "R_Shoulder_Anim"), tgl = True)
+        cmds.select((namespace + ":" + "R_IK_Elbow_Anim"), tgl = True)
+        cmds.select((namespace + ":" + "R_IK_Hand_Anim"), tgl = True)
+        cmds.select((namespace + ":" + "FK_b_MF_UpperArm_R"), tgl = True)
+        cmds.select((namespace + ":" + "FK_b_MF_Forearm_R"), tgl = True)
+        cmds.select((namespace + ":" + "FK_b_MF_Hand_R"), tgl = True)
+        cmds.select((namespace + ":" + "R_Weapon_Anim"), tgl = True)
+        
+        cmds.select((namespace + ":" + "L_Shoulder_Anim"), tgl = True)
+        cmds.select((namespace + ":" + "L_IK_Elbow_Anim"), tgl = True)
+        cmds.select((namespace + ":" + "L_IK_Hand_Anim"), tgl = True)
+        cmds.select((namespace + ":" + "FK_b_MF_UpperArm_L"), tgl = True)
+        cmds.select((namespace + ":" + "FK_b_MF_Forearm_L"), tgl = True)
+        cmds.select((namespace + ":" + "FK_b_MF_Hand_L"), tgl = True)
+        cmds.select((namespace + ":" + "L_Weapon_Anim"), tgl = True)
+        
+        cmds.select((namespace + ":" + "Neck_Anim"), tgl = True)
+        cmds.select((namespace + ":" + "Head_Anim"), tgl = True)
+        
+        cmds.select((namespace + ":" + "Spine_01_Anim"), tgl = True)
+        cmds.select((namespace + ":" + "Spine_02_Anim"), tgl = True)
+        cmds.select((namespace + ":" + "Spine_03_Anim"), tgl = True)
+        cmds.select((namespace + ":" + "Body_Anim"), tgl = True)
+        cmds.select((namespace + ":" + "Hip_Anim"), tgl = True)
+        
+        cmds.select((namespace + ":" + "Master_Anim"), tgl = True)
+        cmds.select((namespace + ":" + "Root_Anim"), tgl = True)
+        cmds.select((namespace + ":" + "Attach_Anim"), tgl = True)
+        
+        cmds.select((namespace + ":" + "Settings"), tgl = True)
+        
+        #BUTTONS
+        cmds.button("rFKThighButton", edit = True, bgc = [1,1,1])
+        cmds.button("rFKCalfButton", edit = True, bgc = [1,1,1])
+        cmds.button("rFKFootButton", edit = True, bgc = [1,1,1])
+        cmds.button("rFKToeButton", edit = True, bgc = [1,1,1])
+        cmds.button("rIKKneeButton", edit = True, bgc = [1,1,1])
+        cmds.button("rIKFootButton", edit = True, bgc = [1,1,1])
+        
+        cmds.button("lFKThighButton", edit = True, bgc = [1,1,1])
+        cmds.button("lFKCalfButton", edit = True, bgc = [1,1,1])
+        cmds.button("lFKFootButton", edit = True, bgc = [1,1,1])
+        cmds.button("lFKToeButton", edit = True, bgc = [1,1,1])
+        cmds.button("lIKKneeButton", edit = True, bgc = [1,1,1])
+        cmds.button("lIKFootButton", edit = True, bgc = [1,1,1])
+        
+        cmds.button("rFKHandButton", edit = True, bgc = [1,1,1])
+        cmds.button("rFKElbowButton", edit = True, bgc = [1,1,1])
+        cmds.button("rFKArmButton", edit = True, bgc = [1,1,1])
+        cmds.button("rShoulderButton", edit = True, bgc = [1,1,1])
+        cmds.button("rElbowButton", edit = True, bgc = [1,1,1])
+        cmds.button("rIKHandButton", edit = True, bgc = [1,1,1])
+        cmds.button("rWepButton", edit = True, bgc = [1,1,1])
+        
+        cmds.button("lFKHandButton", edit = True, bgc = [1,1,1])
+        cmds.button("lFKElbowButton", edit = True, bgc = [1,1,1])
+        cmds.button("lFKArmButton", edit = True, bgc = [1,1,1])
+        cmds.button("lShoulderButton", edit = True, bgc = [1,1,1])
+        cmds.button("lElbowButton", edit = True, bgc = [1,1,1])
+        cmds.button("lIKHandButton", edit = True, bgc = [1,1,1])
+        cmds.button("lWepButton", edit = True, bgc = [1,1,1])
+        
+        cmds.button("neckButton", edit = True, bgc = [1,1,1])
+        cmds.button("headButton", edit = True, bgc = [1,1,1])
+        
+        cmds.button("spine1Button", edit = True, bgc = [1,1,1])
+        cmds.button("spine2Button", edit = True, bgc = [1,1,1])
+        cmds.button("spine3Button", edit = True, bgc = [1,1,1])
+        cmds.button("bodyButton", edit = True, bgc = [1,1,1])
+        cmds.button("hipButton", edit = True, bgc = [1,1,1])
+        
+        cmds.button("masterButton", edit = True, bgc = [1,1,1])
+        cmds.button("rootButton", edit = True, bgc = [1,1,1])
+        cmds.button("attachButton", edit = True, bgc = [1,1,1])
+        
+        cmds.button("settingsButtonBGC", edit = True, bgc = [1,1,1])
+        
+
+    
+        controls = [(namespace + ":" + "FK_b_MF_Thigh_R"),(namespace + ":" + "FK_b_MF_Calf_R"), (namespace + ":" + "FK_b_MF_Foot_R"),(namespace + ":" + "FK_b_MF_Toe_R"),(namespace + ":" + "Right_Knee_Anim"),(namespace + ":" + "Right_Foot_Anim"),(namespace + ":" + "FK_b_MF_Thigh_L"),(namespace + ":" + "FK_b_MF_Calf_L"), (namespace + ":" + "FK_b_MF_Foot_L"),(namespace + ":" + "FK_b_MF_Toe_L"),(namespace + ":" + "Left_Knee_Anim"),(namespace + ":" + "Left_Foot_Anim"),(namespace + ":" + "FK_b_MF_Hand_R"),(namespace + ":" + "FK_b_MF_Forearm_R"), (namespace + ":" + "FK_b_MF_UpperArm_R"),(namespace + ":" + "R_Shoulder_Anim"), (namespace + ":" + "R_IK_Elbow_Anim"),(namespace + ":" + "R_IK_Hand_Anim"),(namespace + ":" + "R_Weapon_Anim"),(namespace + ":" + "FK_b_MF_Hand_L"),(namespace + ":" + "FK_b_MF_Forearm_L"), (namespace + ":" + "FK_b_MF_UpperArm_L"),(namespace + ":" + "L_Shoulder_Anim"), (namespace + ":" + "L_IK_Elbow_Anim"),(namespace + ":" + "L_IK_Hand_Anim"),(namespace + ":" + "L_Weapon_Anim"),(namespace + ":" + "Neck_Anim"), (namespace + ":" + "Head_Anim"),(namespace + ":" + "Spine_01_Anim"),(namespace + ":" + "Spine_02_Anim"),(namespace + ":" + "Spine_03_Anim"),(namespace + ":" + "Body_Anim"),(namespace + ":" + "Hip_Anim"),(namespace + ":" + "Master_Anim"),(namespace + ":" + "Root_Anim"),(namespace + ":" + "Attach_Anim"),(namespace + ":" + "Settings")]
+        buttons = ["rFKThighButton", "rFKCalfButton","rFKFootButton","rFKToeButton","rIKKneeButton","rIKFootButton","lFKThighButton", "lFKCalfButton","lFKFootButton","lFKToeButton","lIKKneeButton","lIKFootButton","rFKHandButton", "rFKElbowButton","rFKArmButton","rShoulderButton","rElbowButton","rIKHandButton","rWepButton","lFKHandButton", "lFKElbowButton","lFKArmButton","lShoulderButton","lElbowButton","lIKHandButton","lWepButton","neckButton", "headButton","spine1Button","spine2Button","spine3Button","bodyButton","hipButton","masterButton","rootButton","attachButton","settingsButtonBGC"]
+        
+        createScriptJob(controls, buttons)
+        
+
+    #if not holding down a modifer
+    if (mods & 1) == 0:
+        
+        #CONTROLS
+        cmds.select((namespace + ":" + "FK_b_MF_Thigh_R"), replace = True)
+        cmds.select((namespace + ":" + "FK_b_MF_Calf_R"), add = True)
+        cmds.select((namespace + ":" + "FK_b_MF_Foot_R"), add = True)
+        cmds.select((namespace + ":" + "FK_b_MF_Toe_R"), add = True)
+        cmds.select((namespace + ":" + "Right_Knee_Anim"), add = True)
+        cmds.select((namespace + ":" + "Right_Foot_Anim"), add = True)
+        
+        cmds.select((namespace + ":" + "FK_b_MF_Thigh_L"), add = True)
+        cmds.select((namespace + ":" + "FK_b_MF_Calf_L"), add = True)
+        cmds.select((namespace + ":" + "FK_b_MF_Foot_L"), add = True)
+        cmds.select((namespace + ":" + "FK_b_MF_Toe_L"), add = True)
+        cmds.select((namespace + ":" + "Left_Knee_Anim"), add = True)
+        cmds.select((namespace + ":" + "Left_Foot_Anim"), add = True)
+        
+        cmds.select((namespace + ":" + "R_Shoulder_Anim"), add = True)
+        cmds.select((namespace + ":" + "R_IK_Elbow_Anim"), add = True)
+        cmds.select((namespace + ":" + "R_IK_Hand_Anim"), add = True)
+        cmds.select((namespace + ":" + "FK_b_MF_UpperArm_R"), add = True)
+        cmds.select((namespace + ":" + "FK_b_MF_Forearm_R"), add = True)
+        cmds.select((namespace + ":" + "FK_b_MF_Hand_R"), add = True)
+        cmds.select((namespace + ":" + "R_Weapon_Anim"), add = True)
+        
+        cmds.select((namespace + ":" + "L_Shoulder_Anim"), add = True)
+        cmds.select((namespace + ":" + "L_IK_Elbow_Anim"), add = True)
+        cmds.select((namespace + ":" + "L_IK_Hand_Anim"), add = True)
+        cmds.select((namespace + ":" + "FK_b_MF_UpperArm_L"), add = True)
+        cmds.select((namespace + ":" + "FK_b_MF_Forearm_L"), add = True)
+        cmds.select((namespace + ":" + "FK_b_MF_Hand_L"), add = True)
+        cmds.select((namespace + ":" + "L_Weapon_Anim"), add = True)
+        
+        cmds.select((namespace + ":" + "Neck_Anim"), add = True)
+        cmds.select((namespace + ":" + "Head_Anim"), add = True)
+        
+        cmds.select((namespace + ":" + "Spine_01_Anim"), add = True)
+        cmds.select((namespace + ":" + "Spine_02_Anim"), add = True)
+        cmds.select((namespace + ":" + "Spine_03_Anim"), add = True)
+        cmds.select((namespace + ":" + "Body_Anim"), add = True)
+        cmds.select((namespace + ":" + "Hip_Anim"), add = True)
+        
+        cmds.select((namespace + ":" + "Master_Anim"), add = True)
+        cmds.select((namespace + ":" + "Root_Anim"), add = True)
+        cmds.select((namespace + ":" + "Attach_Anim"), add = True)
+        
+        cmds.select((namespace + ":" + "Settings"), tgl = True)
+        
+        #BUTTONS
+        cmds.button("rFKThighButton", edit = True, bgc = [1,1,1])
+        cmds.button("rFKCalfButton", edit = True, bgc = [1,1,1])
+        cmds.button("rFKFootButton", edit = True, bgc = [1,1,1])
+        cmds.button("rFKToeButton", edit = True, bgc = [1,1,1])
+        cmds.button("rIKKneeButton", edit = True, bgc = [1,1,1])
+        cmds.button("rIKFootButton", edit = True, bgc = [1,1,1])
+        
+        cmds.button("lFKThighButton", edit = True, bgc = [1,1,1])
+        cmds.button("lFKCalfButton", edit = True, bgc = [1,1,1])
+        cmds.button("lFKFootButton", edit = True, bgc = [1,1,1])
+        cmds.button("lFKToeButton", edit = True, bgc = [1,1,1])
+        cmds.button("lIKKneeButton", edit = True, bgc = [1,1,1])
+        cmds.button("lIKFootButton", edit = True, bgc = [1,1,1])
+        
+        cmds.button("rFKHandButton", edit = True, bgc = [1,1,1])
+        cmds.button("rFKElbowButton", edit = True, bgc = [1,1,1])
+        cmds.button("rFKArmButton", edit = True, bgc = [1,1,1])
+        cmds.button("rShoulderButton", edit = True, bgc = [1,1,1])
+        cmds.button("rElbowButton", edit = True, bgc = [1,1,1])
+        cmds.button("rIKHandButton", edit = True, bgc = [1,1,1])
+        cmds.button("rWepButton", edit = True, bgc = [1,1,1])
+        
+        cmds.button("lFKHandButton", edit = True, bgc = [1,1,1])
+        cmds.button("lFKElbowButton", edit = True, bgc = [1,1,1])
+        cmds.button("lFKArmButton", edit = True, bgc = [1,1,1])
+        cmds.button("lShoulderButton", edit = True, bgc = [1,1,1])
+        cmds.button("lElbowButton", edit = True, bgc = [1,1,1])
+        cmds.button("lIKHandButton", edit = True, bgc = [1,1,1])
+        cmds.button("lWepButton", edit = True, bgc = [1,1,1])
+        
+        cmds.button("neckButton", edit = True, bgc = [1,1,1])
+        cmds.button("headButton", edit = True, bgc = [1,1,1])
+        
+        cmds.button("spine1Button", edit = True, bgc = [1,1,1])
+        cmds.button("spine2Button", edit = True, bgc = [1,1,1])
+        cmds.button("spine3Button", edit = True, bgc = [1,1,1])
+        cmds.button("bodyButton", edit = True, bgc = [1,1,1])
+        cmds.button("hipButton", edit = True, bgc = [1,1,1])
+        
+        cmds.button("masterButton", edit = True, bgc = [1,1,1])
+        cmds.button("rootButton", edit = True, bgc = [1,1,1])
+        cmds.button("attachButton", edit = True, bgc = [1,1,1])
+        
+        cmds.button("settingsButtonBGC", edit = True, bgc = [1,1,1])
+        
+        controls = [(namespace + ":" + "FK_b_MF_Thigh_R"),(namespace + ":" + "FK_b_MF_Calf_R"), (namespace + ":" + "FK_b_MF_Foot_R"),(namespace + ":" + "FK_b_MF_Toe_R"),(namespace + ":" + "Right_Knee_Anim"),(namespace + ":" + "Right_Foot_Anim"),(namespace + ":" + "FK_b_MF_Thigh_L"),(namespace + ":" + "FK_b_MF_Calf_L"), (namespace + ":" + "FK_b_MF_Foot_L"),(namespace + ":" + "FK_b_MF_Toe_L"),(namespace + ":" + "Left_Knee_Anim"),(namespace + ":" + "Left_Foot_Anim"),(namespace + ":" + "FK_b_MF_Hand_R"),(namespace + ":" + "FK_b_MF_Forearm_R"), (namespace + ":" + "FK_b_MF_UpperArm_R"),(namespace + ":" + "R_Shoulder_Anim"), (namespace + ":" + "R_IK_Elbow_Anim"),(namespace + ":" + "R_IK_Hand_Anim"),(namespace + ":" + "R_Weapon_Anim"),(namespace + ":" + "FK_b_MF_Hand_L"),(namespace + ":" + "FK_b_MF_Forearm_L"), (namespace + ":" + "FK_b_MF_UpperArm_L"),(namespace + ":" + "L_Shoulder_Anim"), (namespace + ":" + "L_IK_Elbow_Anim"),(namespace + ":" + "L_IK_Hand_Anim"),(namespace + ":" + "L_Weapon_Anim"),(namespace + ":" + "Neck_Anim"), (namespace + ":" + "Head_Anim"),(namespace + ":" + "Spine_01_Anim"),(namespace + ":" + "Spine_02_Anim"),(namespace + ":" + "Spine_03_Anim"),(namespace + ":" + "Body_Anim"),(namespace + ":" + "Hip_Anim"),(namespace + ":" + "Master_Anim"),(namespace + ":" + "Root_Anim"),(namespace + ":" + "Attach_Anim"),(namespace + ":" + "Settings")]
+        buttons = ["rFKThighButton", "rFKCalfButton","rFKFootButton","rFKToeButton","rIKKneeButton","rIKFootButton","lFKThighButton", "lFKCalfButton","lFKFootButton","lFKToeButton","lIKKneeButton","lIKFootButton","rFKHandButton", "rFKElbowButton","rFKArmButton","rShoulderButton","rElbowButton","rIKHandButton","rWepButton","lFKHandButton", "lFKElbowButton","lFKArmButton","lShoulderButton","lElbowButton","lIKHandButton","lWepButton","neckButton", "headButton","spine1Button","spine2Button","spine3Button","bodyButton","hipButton","masterButton","rootButton","attachButton","settingsButtonBGC"]        
+        
+        createScriptJob(controls, buttons)
+        
+
+
+
+        
+        
+def createScriptJob(controls, buttons):
+    
+    numJob = cmds.scriptJob( runOnce=True, event=['SelectionChanged', partial(revertButton, controls, buttons)], kws = True)
+    
+    
+def revertButton(controls, buttons):
+    
+    masterButtons = ["bodyButton","masterButton"]
+    fkButtons = ["rShoulderButton", "lShoulderButton", "spine1Button","spine2Button","spine3Button", "hipButton", "neckButton", "headButton", "lFKHandButton", "lFKElbowButton", "lFKArmButton","rFKHandButton", "rFKElbowButton", "rFKArmButton","lFKThighButton", "lFKCalfButton", "lFKFootButton","lFKToeButton","rFKThighButton", "rFKCalfButton", "rFKFootButton","rFKToeButton", "rootButton", "Rpinky3Button", "Rpinky2Button", "Rpinky1Button", "Rring3Button", "Rring2Button", "Rring1Button", "Rmiddle3Button", "Rmiddle2Button", "Rmiddle1Button", "Rindex3Button", "Rindex2Button", "Rindex1Button", "Rthumb3Button", "Rthumb2Button", "Rthumb1Button", "Lpinky3Button", "Lpinky2Button", "Lpinky1Button", "Lring3Button", "Lring2Button", "Lring1Button", "Lmiddle3Button", "Lmiddle2Button", "Lmiddle1Button", "Lindex3Button", "Lindex2Button", "Lindex1Button", "Lthumb3Button", "Lthumb2Button", "Lthumb1Button"]
+    ikButtons = ["lElbowButton","lIKHandButton","rElbowButton","rIKHandButton","lIKKneeButton", "lIKFootButton","rIKKneeButton", "rIKFootButton"]
+    miscButtons = ["lWepButton", "rWepButton", "attachButton","LLancerStockButton", "lWepDataButton", "rWepDataButton"]
+    otherButtons = ["settingsButtonBGC"]
+    #get the current selection
+    selection = cmds.ls(sl = True, transforms = True)
+    
+    #create empty lists for later use
+    selectedControls = []
+    selectedButtons = []
+    
+    #get the size of the current selection
+    size = len(controls)
+    
+
+    for i in range(0, size):
+        if controls[i] not in selection:
+            
+                #if the specified control is not in your current selection, then modify that control's correlating button
+                if buttons[i] in masterButtons:
+                    cmds.button(buttons[i], edit = True, bgc = [.968,.639,0])
+                    
+                if buttons[i] in fkButtons:
+                    cmds.button(buttons[i], edit = True, bgc = [.196, .475, .98])
+                    
+                if buttons[i] in ikButtons:
+                    cmds.button(buttons[i], edit = True, bgc = [1, 0, 0])
+                    
+                if buttons[i] in miscButtons:
+                    cmds.button(buttons[i], edit = True, bgc = [.557, 0, .988])
+                    
+                if buttons[i] in otherButtons:
+                    cmds.button(buttons[i], edit = True, bgc = [.4, .4, .4])
+                    
+        else:
+            selectedControls.append(controls[i])
+            selectedButtons.append(buttons[i])
+            
+            
+    #if the controls is in the selection, create a new scriptJob that has that control in it      
+    if len(selectedControls) > 0:
+        
+        controls = selectedControls
+        buttons = selectedButtons
+        createScriptJob(controls, buttons)
+                    
+
+def uberCineSnapper(*args):
+
+    import autoRig_WeaponBatcher as arWB
+    reload(arWB)
+    arWB.selectCineCharacters_UI()
+    
+def shotSplitter(*args):
+
+    import autoRig_shotSplitter as arSS
+    reload(arSS)
+    arSS.UI()
+
+def importShotRefVid(*args):
+    import je_shotReferenceTool as jeSRT
+    reload(jeSRT)
+    jeSRT.UI()
+    
+def refVidSettings(*args):
+    import je_shotReferenceTool as jeSRT
+    reload(jeSRT)
+    jeSRT.listAllReferenceMovies_UI()
+    
+
+
+def UI():
+    
+    #if the window exists, delete it
+    if cmds.window("autoRig_CharacterUI", exists = True):
+        cmds.deleteUI("autoRig_CharacterUI")
+        
+    if cmds.dockControl("Auto_Rig_UI", exists = True):
+        cmds.deleteUI("Auto_Rig_UI")
+        
+    #create the main window
+    window = cmds.window("autoRig_CharacterUI", title = "autoRig_CharacterUI", width = 400, height = 1000, sizeable = False)
+    
+    #create a help menu/menuBar
+    menuBarLayout = cmds.menuBarLayout()
+    cmds.menu( label='Help' )
+    cmds.menuItem( label='Help Documentation' )
+    cmds.menuItem( label='Current Export Joints List', c = getExportJoints_UI)
+    
+    cmds.menu("ToolsMenu",  label = 'Tools' )
+    cmds.menuItem( label = 'Match Over Frame Range' , c = matchOverFrameRange_UI)
+    cmds.menuItem( label = 'Select Export Joints Tool' , c = selectExportJointsUI)
+    cmds.menuItem( label = 'Uber Cine Setup' , c = uberCineSnapper)
+    cmds.menuItem( label = 'Shot Splitter' , c = shotSplitter)
+    cmds.menuItem( label = 'Import Shot Reference Video' , c = importShotRefVid)
+    cmds.menuItem( label = 'Reference Video Settings' , c = refVidSettings)
+    
+    
+    #create the main layout
+    scrollLayout = cmds.scrollLayout(horizontalScrollBarThickness = 0, h = 1000)
+    mainLayout = cmds.columnLayout(parent = scrollLayout)
+    cmds.separator(w = 400, style = 'out', parent = mainLayout)
+    
+    #create the main tabLayout
+    tabLayout = cmds.tabLayout("autoRigUI_TabLayout", parent = mainLayout)
+    
+    #create the first tab
+    tab1 = cmds.columnLayout(adj = True, parent = tabLayout)
+    
+    #create a text field to show selected character
+    cmds.optionMenu("selectedCharacterOptionMenu", label = "Active Character :", cc = activeCharacterChangeCommand, parent = tab1)
+
+    #and now create the formLayout to house the picker
+    pickerFormLayout = cmds.formLayout(parent = tab1, h = 525, w = 400 )
+    
+    #add the background image
+    
+    image = cmds.internalVar(upd = True)
+    image += "Icons/Tools/autoRigCharacterUI.jpg"
+    cmds.image(w = 400, h = 525, image = image, parent = pickerFormLayout)
+    
+    
+    #edit the tabLayout
+    cmds.tabLayout( tabLayout, edit=True, tabLabel=[(tab1, 'Character Picker')])
+    
+    #create a dictionary for the picker buttons
+    pickerButtons = {}
+    
+    
+    #and start creating the picker buttons
+    pickerButtons["bodyButton"] = cmds.button("bodyButton", bgc = [.968,.639,0], width = 90, height = 20, label = "", c = selectBody)
+    pickerButtons["pelvisButton"] = cmds.button("hipButton", bgc = [.196, .475, .98], width = 60, height = 20, label = "", c = selectHips)
+    
+    pickerButtons["spine1Button"] = cmds.button("spine1Button", bgc = [.196, .475, .98], width = 60, height = 20, label = "", c = selectSpine1)
+    pickerButtons["spine2Button"] = cmds.button("spine2Button", bgc = [.196, .475, .98], width = 60, height = 20, label = "", c = selectSpine2)
+    pickerButtons["spine3Button"] = cmds.button("spine3Button", bgc = [.196, .475, .98], width = 60, height = 20, label = "", c = selectSpine3)
+    
+    pickerButtons["neckButton"] = cmds.button("neckButton", bgc = [.196, .475, .98], width = 30, height = 30, label = "", c =selectNeck)
+    pickerButtons["headButton"] = cmds.button("headButton", bgc = [.196, .475, .98], width = 50, height = 50, label = "", c = selectHead)
+    
+    pickerButtons["selectAllCoreButton"] = cmds.button(bgc = [.33, .98, .172], width = 20, height = 20, label = "", c = selectAllTorso)
+    pickerButtons["selectAllHeadButton"] = cmds.button(bgc = [.33, .98, .172], width = 20, height = 20, label = "",c = selectAllHead)
+    
+    pickerButtons["LfkLegButton"] = cmds.button("lFKThighButton", bgc = [.196, .475, .98], width = 20, height = 60, label = "", c = selectLFKThigh)
+    pickerButtons["LikLegButton"] = cmds.button("lIKKneeButton", bgc = [1, 0, 0], width = 20, height = 20, label = "", c = selectLIKKnee)
+    pickerButtons["LfkLowLegButton"] = cmds.button("lFKCalfButton", bgc = [.196, .475, .98], width = 20, height = 60, label = "", c = selectLFKCalf)
+    pickerButtons["LikFootButton"] = cmds.button("lIKFootButton", bgc = [1, 0, 0], width = 20, height = 20, label = "", c = selectLIKFoot)
+    pickerButtons["LfkFootButton"] = cmds.button("lFKFootButton", bgc = [.196, .475, .98], width = 40, height = 20, label = "", c = selectLFKFoot)
+    pickerButtons["LfkToeButton"] = cmds.button("lFKToeButton", bgc = [.196, .475, .98], width = 20, height = 20, label = "", c = selectLFKToe)
+    pickerButtons["LselectAllLegButton"] = cmds.button(bgc = [.33, .98, .172], width = 20, height = 20, label = "", c = selectAllLeftLeg)
+    
+    
+    pickerButtons["RfkLegButton"] = cmds.button("rFKThighButton", bgc = [.196, .475, .98], width = 20, height = 60, label = "", c = selectRFKThigh)
+    pickerButtons["RikLegButton"] = cmds.button("rIKKneeButton", bgc = [1, 0, 0], width = 20, height = 20, label = "", c = selectRIKKnee)
+    pickerButtons["RfkLowLegButton"] = cmds.button("rFKCalfButton", bgc = [.196, .475, .98], width = 20, height = 60, label = "", c = selectRFKCalf)
+    pickerButtons["RikFootButton"] = cmds.button("rIKFootButton", bgc = [1, 0, 0], width = 20, height = 20, label = "", c = selectRIKFoot)
+    pickerButtons["RfkFootButton"] = cmds.button("rFKFootButton", bgc = [.196, .475, .98], width = 40, height = 20, label = "", c = selectRFKFoot)
+    pickerButtons["RfkToeButton"] = cmds.button("rFKToeButton", bgc = [.196, .475, .98], width = 20, height = 20, label = "", c = selectRFKToe)
+    pickerButtons["RselectAllLegButton"] = cmds.button(bgc = [.33, .98, .172], width = 20, height = 20, label = "", c = selectAllRightLeg)
+    
+    pickerButtons["LClavButton"] = cmds.button("lShoulderButton", bgc = [.196, .475, .98], width = 40, height = 20, label = "", c = selectLShoulder)
+    pickerButtons["LfkArmButton"] = cmds.button("lFKArmButton", bgc = [.196, .475, .98], width = 60, height = 20, label = "", c = selectLFKArm)
+    pickerButtons["LikElbowButton"] = cmds.button("lElbowButton", bgc = [1, 0, 0], width = 20, height = 20, label = "",c = selectLElbow)
+    pickerButtons["LfkLowArmButton"] = cmds.button("lFKElbowButton", bgc = [.196, .475, .98], width = 20, height = 80, label = "", c = selectLFKElbow)
+    pickerButtons["LikHandButton"] = cmds.button("lIKHandButton", bgc = [1, 0, 0], width = 20, height = 20, label = "", c = selectLIKHand)
+    pickerButtons["LfkHandButton"] = cmds.button("lFKHandButton", bgc = [.196, .475, .98], width = 40, height = 40, label = "", c = selectLFKHand)
+    pickerButtons["LWepButton"] = cmds.button("lWepButton", bgc = [.557, 0, .988], width = 20, height = 20, label = "", c = selectLWep)
+    pickerButtons["LselectAllArmButton"] = cmds.button(bgc = [.33, .98, .172], width = 20, height = 20, label = "", c = selectAllLeftArm)
+    
+    pickerButtons["LLancerStockButton"] = cmds.button("LLancerStockButton", bgc = [.557, 0, .988], width = 20, height = 20, label = "", c = selectLancerStock)
+
+
+    pickerButtons["RClavButton"] = cmds.button("rShoulderButton", bgc = [.196, .475, .98], width = 40, height = 20, label = "", c = selectRShoulder)
+    pickerButtons["RfkArmButton"] = cmds.button("rFKArmButton", bgc = [.196, .475, .98], width = 60, height = 20, label = "", c = selectRFKArm)
+    pickerButtons["RikElbowButton"] = cmds.button("rElbowButton", bgc = [1, 0, 0], width = 20, height = 20, label = "",c = selectRElbow)
+    pickerButtons["RfkLowArmButton"] = cmds.button("rFKElbowButton", bgc = [.196, .475, .98], width = 20, height = 80, label = "", c = selectRFKElbow)
+    pickerButtons["RikHandButton"] = cmds.button("rIKHandButton", bgc = [1, 0, 0], width = 20, height = 20, label = "", c = selectRIKHand)
+    pickerButtons["RfkHandButton"] = cmds.button("rFKHandButton", bgc = [.196, .475, .98], width = 40, height = 40, label = "", c = selectRFKHand)
+    pickerButtons["RWepButton"] = cmds.button("rWepButton", bgc = [.557, 0, .988], width = 20, height = 20, label = "", c = selectRWep)
+    pickerButtons["RselectAllArmButton"] = cmds.button(bgc = [.33, .98, .172], width = 20, height = 20, label = "", c = selectAllRightArm)
+    pickerButtons["selectSettingsBGC"] = cmds.button("settingsButtonBGC", width = 70, height = 40, label = "", isObscured = True, enable = False, bgc = [.4, .4,.4])
+    pickerButtons["selectSettings"] = cmds.button("settingsButton", width = 60, height = 30, label = "Settings", c = selectSettings)
+    
+    pickerButtons["hideControls"] = cmds.button("hideControlsButton", width = 40, height = 65, label = "HIDE\nCTRLS", c = hideControls)
+    pickerButtons["showControls"] = cmds.button("showControlsButton", width = 40, height = 65, label = "SHOW\nCTRLS", c = showControls)
+    
+    
+    pickerButtons["selectAllButton"] = cmds.button(bgc = [.33, .98, .172], width = 20, height = 20, label = "", c = selectAll)
+    
+    pickerButtons["RootOffsetButton"] = cmds.button("rootButton", bgc = [.196, .475, .98], width = 20, height = 20, label = "", c = selectRoot)
+    pickerButtons["AttachButton"] = cmds.button("attachButton", bgc = [.557, 0, .988], width = 20, height = 20, label = "", c = selectAttach)
+    pickerButtons["MasterButton"] = cmds.button("masterButton", bgc = [.968,.639,0], width = 20, height = 20, label = "", c = selectMaster)
+    
+    pickerButtons["lWepDataButton"] = cmds.button("lWepDataButton", bgc = [.557, 0, .988], width = 15, height = 15, label = "", c = selectLWepData)
+    pickerButtons["rWepDataButton"] = cmds.button("rWepDataButton", bgc = [.557, 0, .988], width = 15, height = 15, label = "", c = selectRWepData)
+
+
+
+    pickerSep = cmds.separator(w = 400, style = 'out', parent = pickerFormLayout)
+    
+    pickerButtons["pickerSeparatorBottom"] = cmds.separator(w = 400, style = 'out', parent = pickerFormLayout)
+    
+    pickerButtons["keyText"] = cmds.button(width = 30, height = 30, label = "KEY:", enable = False)
+    pickerButtons["keyAll"] = cmds.button(width = 30, height = 30, label = "ALL", c = keyAll)
+    pickerButtons["keySel"] = cmds.button(width = 30, height = 30, label = "SEL", c = keySelection)
+    
+    pickerButtons["resetText"] = cmds.button(width = 30, height = 30, label = "RES:", enable = False)
+    pickerButtons["resetAll"] = cmds.button(width = 30, height = 30, label = "ALL", c = resetAll)
+    pickerButtons["resetSel"] = cmds.button(width = 30, height = 30, label = "SEL", c = resetSelection)
+    
+    pickerButtons["exportButton"] = cmds.button(width = 95, height = 30, label = "EXPORT", c = export)
+    
+    pickerButtons["addMocapButton"] = cmds.button(width = 95, height = 30, label = "Add Mocap", c = addMocap)
+    pickerButtons["addWeaponButton"] = cmds.button(width = 95, height = 30, label = "Add Weapon", c = addWeapon)
+    pickerButtons["addFaceRigButton"] = cmds.button(width = 95, height = 30, label = "Add Face Rig", enable = True, c = addFaceRig)
+    
+    pickerButtons["LArmFKButton"] = cmds.button(width = 20, height = 20, label = "FK", c = LArmModeSwitchFK)
+    pickerButtons["LArmIKButton"] = cmds.button(width = 20, height = 20, label = "IK", c = LArmModeSwitchIK)
+    
+    pickerButtons["RArmFKButton"] = cmds.button(width = 20, height = 20, label = "FK", c = RArmModeSwitchFK)
+    pickerButtons["RArmIKButton"] = cmds.button(width = 20, height = 20, label = "IK", c = RArmModeSwitchIK)
+    
+    pickerButtons["LLegFKButton"] = cmds.button(width = 20, height = 20, label = "FK", c = LLegModeSwitchFK)
+    pickerButtons["LLegIKButton"] = cmds.button(width = 20, height = 20, label = "IK", c = LLegModeSwitchIK)
+ 
+    pickerButtons["RLegFKButton"] = cmds.button(width = 20, height = 20, label = "FK", c = RLegModeSwitchFK)
+    pickerButtons["RLegIKButton"] = cmds.button(width = 20, height = 20, label = "IK", c = RLegModeSwitchIK)
+    
+    pickerButtons["lClavOrientMenu"] = cmds.optionMenu("lClavOrientMenu", label = 'FK Orient:', cc = lFKArmOrient)
+    cmds.menuItem(label = 'FK')
+    cmds.menuItem(label = 'Torso')
+    cmds.menuItem(label = 'Body')
+    cmds.menuItem(label = 'World')
+    
+
+    pickerButtons["rClavOrientMenu"] = cmds.optionMenu("rClavOrientMenu", label = 'FK Orient:', cc = rFKArmOrient)
+    cmds.menuItem(label = 'FK')
+    cmds.menuItem(label = 'Torso')
+    cmds.menuItem(label = 'Body')
+    cmds.menuItem(label = 'World')
+    
+
+
+    #edit the position on the buttons
+    cmds.formLayout(pickerFormLayout, edit = True, af = [(pickerButtons["selectSettingsBGC"], 'bottom', 75), (pickerButtons["selectSettingsBGC"], 'right', 0)])
+    cmds.formLayout(pickerFormLayout, edit = True, af = [(pickerButtons["selectSettings"], 'bottom', 80), (pickerButtons["selectSettings"], 'right', 5)])
+
+    cmds.formLayout(pickerFormLayout, edit = True, af = [(pickerButtons["bodyButton"], 'top', 215), (pickerButtons["bodyButton"], 'left', 150)])
+    cmds.formLayout(pickerFormLayout, edit = True, af = [(pickerButtons["pelvisButton"], 'top', 242), (pickerButtons["pelvisButton"], 'left', 165)])
+    cmds.formLayout(pickerFormLayout, edit = True, af = [(pickerButtons["spine1Button"], 'top', 190), (pickerButtons["spine1Button"], 'left', 165)])
+    cmds.formLayout(pickerFormLayout, edit = True, af = [(pickerButtons["spine2Button"], 'top', 165), (pickerButtons["spine2Button"], 'left', 165)])
+    cmds.formLayout(pickerFormLayout, edit = True, af = [(pickerButtons["spine3Button"], 'top', 140), (pickerButtons["spine3Button"], 'left', 165)])
+
+    cmds.formLayout(pickerFormLayout, edit = True, af = [(pickerButtons["neckButton"], 'top', 105), (pickerButtons["neckButton"], 'left', 180)])
+    cmds.formLayout(pickerFormLayout, edit = True, af = [(pickerButtons["headButton"], 'top', 50), (pickerButtons["headButton"], 'left', 170)])
+
+    cmds.formLayout(pickerFormLayout, edit = True, af = [(pickerButtons["selectAllCoreButton"], 'top', 265), (pickerButtons["selectAllCoreButton"], 'left', 185)])
+    cmds.formLayout(pickerFormLayout, edit = True, af = [(pickerButtons["selectAllHeadButton"], 'top', 25), (pickerButtons["selectAllHeadButton"], 'left', 185)])
+
+    cmds.formLayout(pickerFormLayout, edit = True, af = [(pickerButtons["LfkLegButton"], 'top', 270), (pickerButtons["LfkLegButton"], 'left', 225)])
+    cmds.formLayout(pickerFormLayout, edit = True, af = [(pickerButtons["LikLegButton"], 'top', 335), (pickerButtons["LikLegButton"], 'left', 225)])
+    cmds.formLayout(pickerFormLayout, edit = True, af = [(pickerButtons["LfkLowLegButton"], 'top', 360), (pickerButtons["LfkLowLegButton"], 'left', 225)])
+    cmds.formLayout(pickerFormLayout, edit = True, af = [(pickerButtons["LikFootButton"], 'top', 425), (pickerButtons["LikFootButton"], 'left', 225)])
+    cmds.formLayout(pickerFormLayout, edit = True, af = [(pickerButtons["LfkFootButton"], 'top', 425), (pickerButtons["LfkFootButton"], 'left', 250)])
+    cmds.formLayout(pickerFormLayout, edit = True, af = [(pickerButtons["LfkToeButton"], 'top', 425), (pickerButtons["LfkToeButton"], 'left', 295)])
+    cmds.formLayout(pickerFormLayout, edit = True, af = [(pickerButtons["LselectAllLegButton"], 'top', 245), (pickerButtons["LselectAllLegButton"], 'left', 225)])
+
+    cmds.formLayout(pickerFormLayout, edit = True, af = [(pickerButtons["RfkLegButton"], 'top', 270), (pickerButtons["RfkLegButton"], 'left', 145)])
+    cmds.formLayout(pickerFormLayout, edit = True, af = [(pickerButtons["RikLegButton"], 'top', 335), (pickerButtons["RikLegButton"], 'left', 145)])
+    cmds.formLayout(pickerFormLayout, edit = True, af = [(pickerButtons["RfkLowLegButton"], 'top', 360), (pickerButtons["RfkLowLegButton"], 'left', 145)])
+    cmds.formLayout(pickerFormLayout, edit = True, af = [(pickerButtons["RikFootButton"], 'top', 425), (pickerButtons["RikFootButton"], 'left', 145)])
+    cmds.formLayout(pickerFormLayout, edit = True, af = [(pickerButtons["RfkFootButton"], 'top', 425), (pickerButtons["RfkFootButton"], 'left', 100)])
+    cmds.formLayout(pickerFormLayout, edit = True, af = [(pickerButtons["RfkToeButton"], 'top', 425), (pickerButtons["RfkToeButton"], 'left', 75)])
+    cmds.formLayout(pickerFormLayout, edit = True, af = [(pickerButtons["RselectAllLegButton"], 'top', 245), (pickerButtons["RselectAllLegButton"], 'left', 145)])
+    
+    cmds.formLayout(pickerFormLayout, edit = True, af = [(pickerButtons["LClavButton"], 'top', 115), (pickerButtons["LClavButton"], 'left', 215)])
+    cmds.formLayout(pickerFormLayout, edit = True, af = [(pickerButtons["LfkArmButton"], 'top', 115), (pickerButtons["LfkArmButton"], 'left', 260)])
+    cmds.formLayout(pickerFormLayout, edit = True, af = [(pickerButtons["LikElbowButton"], 'top', 115), (pickerButtons["LikElbowButton"], 'left', 325)])
+    cmds.formLayout(pickerFormLayout, edit = True, af = [(pickerButtons["LfkLowArmButton"], 'top', 140), (pickerButtons["LfkLowArmButton"], 'left', 325)])
+    cmds.formLayout(pickerFormLayout, edit = True, af = [(pickerButtons["LikHandButton"], 'top', 225), (pickerButtons["LikHandButton"], 'left', 325)])
+    cmds.formLayout(pickerFormLayout, edit = True, af = [(pickerButtons["LfkHandButton"], 'top', 250), (pickerButtons["LfkHandButton"], 'left', 315)])
+    cmds.formLayout(pickerFormLayout, edit = True, af = [(pickerButtons["LWepButton"], 'top', 225), (pickerButtons["LWepButton"], 'left', 300)])
+    
+    cmds.formLayout(pickerFormLayout, edit = True, af = [(pickerButtons["LLancerStockButton"], 'top', 225), (pickerButtons["LLancerStockButton"], 'left', 370)])
+
+    cmds.formLayout(pickerFormLayout, edit = True, af = [(pickerButtons["LselectAllArmButton"], 'top', 90), (pickerButtons["LselectAllArmButton"], 'left', 225)])
+
+    cmds.formLayout(pickerFormLayout, edit = True, af = [(pickerButtons["RClavButton"], 'top', 115), (pickerButtons["RClavButton"], 'left', 135)])
+    cmds.formLayout(pickerFormLayout, edit = True, af = [(pickerButtons["RfkArmButton"], 'top', 115), (pickerButtons["RfkArmButton"], 'left', 70)])
+    cmds.formLayout(pickerFormLayout, edit = True, af = [(pickerButtons["RikElbowButton"], 'top', 115), (pickerButtons["RikElbowButton"], 'left', 45)])
+    cmds.formLayout(pickerFormLayout, edit = True, af = [(pickerButtons["RfkLowArmButton"], 'top', 140), (pickerButtons["RfkLowArmButton"], 'left', 45)])
+    cmds.formLayout(pickerFormLayout, edit = True, af = [(pickerButtons["RikHandButton"], 'top', 225), (pickerButtons["RikHandButton"], 'left', 45)])
+    cmds.formLayout(pickerFormLayout, edit = True, af = [(pickerButtons["RfkHandButton"], 'top', 250), (pickerButtons["RfkHandButton"], 'left', 35)])
+    cmds.formLayout(pickerFormLayout, edit = True, af = [(pickerButtons["RWepButton"], 'top', 225), (pickerButtons["RWepButton"], 'left', 70)])
+    cmds.formLayout(pickerFormLayout, edit = True, af = [(pickerButtons["RselectAllArmButton"], 'top', 90), (pickerButtons["RselectAllArmButton"], 'left', 145)])
+    
+    
+    cmds.formLayout(pickerFormLayout, edit = True, af = [(pickerButtons["selectAllButton"], 'top', 425), (pickerButtons["selectAllButton"], 'left', 185)])
+
+    cmds.formLayout(pickerFormLayout, edit = True, af = [(pickerButtons["RootOffsetButton"], 'top', 375), (pickerButtons["RootOffsetButton"], 'left', 185)])
+    cmds.formLayout(pickerFormLayout, edit = True, af = [(pickerButtons["AttachButton"], 'top', 350), (pickerButtons["AttachButton"], 'left', 185)])
+    cmds.formLayout(pickerFormLayout, edit = True, af = [(pickerButtons["MasterButton"], 'top', 400), (pickerButtons["MasterButton"], 'left', 185)])
+    
+    cmds.formLayout(pickerFormLayout, edit = True, af = [(pickerButtons["lWepDataButton"], 'top', 400), (pickerButtons["lWepDataButton"], 'left', 210)])
+    cmds.formLayout(pickerFormLayout, edit = True, af = [(pickerButtons["rWepDataButton"], 'top', 400), (pickerButtons["rWepDataButton"], 'left', 165)])
+
+    
+    cmds.formLayout(pickerFormLayout, edit = True, af = [(pickerButtons["pickerSeparatorBottom"], 'top', 450), (pickerButtons["pickerSeparatorBottom"], 'left', 0)])
+    cmds.formLayout(pickerFormLayout, edit = True, af = [(pickerButtons["keyText"], 'top', 455), (pickerButtons["keyText"], 'left', 5)])
+    cmds.formLayout(pickerFormLayout, edit = True, af = [(pickerButtons["keyAll"], 'top', 455), (pickerButtons["keyAll"], 'left', 40)])
+    cmds.formLayout(pickerFormLayout, edit = True, af = [(pickerButtons["keySel"], 'top', 455), (pickerButtons["keySel"], 'left', 70)])
+
+    cmds.formLayout(pickerFormLayout, edit = True, af = [(pickerButtons["resetText"], 'top', 455), (pickerButtons["resetText"], 'left', 155)])
+    cmds.formLayout(pickerFormLayout, edit = True, af = [(pickerButtons["resetAll"], 'top', 455), (pickerButtons["resetAll"], 'left', 190)])
+    cmds.formLayout(pickerFormLayout, edit = True, af = [(pickerButtons["resetSel"], 'top', 455), (pickerButtons["resetSel"], 'left', 220)])
+    
+    cmds.formLayout(pickerFormLayout, edit = True, af = [(pickerButtons["hideControls"], 'top', 455), (pickerButtons["hideControls"], 'left', 255)])
+    cmds.formLayout(pickerFormLayout, edit = True, af = [(pickerButtons["showControls"], 'top', 455), (pickerButtons["showControls"], 'left', 107)])
+
+
+    cmds.formLayout(pickerFormLayout, edit = True, af = [(pickerButtons["exportButton"], 'top', 455), (pickerButtons["exportButton"], 'left', 300)])
+    cmds.formLayout(pickerFormLayout, edit = True, af = [(pickerButtons["addMocapButton"], 'top', 490), (pickerButtons["addMocapButton"], 'left', 5)])
+    cmds.formLayout(pickerFormLayout, edit = True, af = [(pickerButtons["addWeaponButton"], 'top', 490), (pickerButtons["addWeaponButton"], 'left', 155)])
+    cmds.formLayout(pickerFormLayout, edit = True, af = [(pickerButtons["addFaceRigButton"], 'top', 490), (pickerButtons["addFaceRigButton"], 'left', 300)])
+
+    cmds.formLayout(pickerFormLayout, edit = True, af = [(pickerButtons["LArmFKButton"], 'top', 50), (pickerButtons["LArmFKButton"], 'left', 302)])
+    cmds.formLayout(pickerFormLayout, edit = True, af = [(pickerButtons["LArmIKButton"], 'top', 50), (pickerButtons["LArmIKButton"], 'left', 325)])
+    
+    cmds.formLayout(pickerFormLayout, edit = True, af = [(pickerButtons["RArmFKButton"], 'top', 50), (pickerButtons["RArmFKButton"], 'left', 67)])
+    cmds.formLayout(pickerFormLayout, edit = True, af = [(pickerButtons["RArmIKButton"], 'top', 50), (pickerButtons["RArmIKButton"], 'left', 45)])
+ 
+    cmds.formLayout(pickerFormLayout, edit = True, af = [(pickerButtons["LLegFKButton"], 'top', 335), (pickerButtons["LLegFKButton"], 'left', 302)])
+    cmds.formLayout(pickerFormLayout, edit = True, af = [(pickerButtons["LLegIKButton"], 'top', 335), (pickerButtons["LLegIKButton"], 'left', 325)])
+
+    cmds.formLayout(pickerFormLayout, edit = True, af = [(pickerButtons["RLegFKButton"], 'top', 335), (pickerButtons["RLegFKButton"], 'left', 67)])
+    cmds.formLayout(pickerFormLayout, edit = True, af = [(pickerButtons["RLegIKButton"], 'top', 335), (pickerButtons["RLegIKButton"], 'left', 45)])
+
+    cmds.formLayout(pickerFormLayout, edit = True, af = [(pickerButtons["lClavOrientMenu"], 'top', 5), (pickerButtons["lClavOrientMenu"], 'right', 5)])
+    cmds.formLayout(pickerFormLayout, edit = True, af = [(pickerButtons["rClavOrientMenu"], 'top', 5), (pickerButtons["rClavOrientMenu"], 'left', 5)])
+
+    
+
+
+    #and now create the formLayout to house the hand/fingers picker
+    cmds.separator(w = 400, style = 'out', parent = tab1)
+    handPickerFormLayout = cmds.formLayout(parent = tab1, h = 125, w = 400 )
+    
+    #create the buttons for the picker
+    
+        
+    pickerButtons["RPinkyColButton"] = cmds.button(bgc = [.33, .98, .172], width = 20, height = 20, label = "", c = partial(selectMultiFingers, ["R_Pinky_3_Anim","R_Pinky_2_Anim","R_Pinky_1_Anim"],["Rpinky3Button","Rpinky2Button","Rpinky1Button"]))
+    pickerButtons["RRingColButton"] = cmds.button(bgc = [.33, .98, .172], width = 20, height = 20, label = "", c = partial(selectMultiFingers, ["R_Ring_3_Anim","R_Ring_2_Anim","R_Ring_1_Anim"],["Rring3Button","Rring2Button","Rring1Button"]))
+    pickerButtons["RMiddleColButton"] = cmds.button(bgc = [.33, .98, .172], width = 20, height = 20, label = "", c = partial(selectMultiFingers, ["R_Middle_3_Anim","R_Middle_2_Anim","R_Middle_1_Anim"],["Rmiddle3Button","Rmiddle2Button","Rmiddle1Button"]))
+    pickerButtons["RIndexColButton"] = cmds.button(bgc = [.33, .98, .172], width = 20, height = 20, label = "", c = partial(selectMultiFingers, ["R_Index_3_Anim","R_Index_2_Anim","R_Index_1_Anim"],["Rindex3Button","Rindex2Button","Rindex1Button"]))
+    
+    pickerButtons["RTopFingerRow"] = cmds.button(bgc = [.33, .98, .172], width = 20, height = 20, label = "", c = partial(selectMultiFingers, ["R_Pinky_3_Anim","R_Ring_3_Anim","R_Middle_3_Anim","R_Index_3_Anim"],["Rpinky3Button","Rring3Button","Rmiddle3Button","Rindex3Button"]))
+    pickerButtons["RMidFingerRow"] = cmds.button(bgc = [.33, .98, .172], width = 20, height = 20, label = "", c = partial(selectMultiFingers, ["R_Pinky_2_Anim","R_Ring_2_Anim","R_Middle_2_Anim","R_Index_2_Anim"],["Rpinky2Button","Rring2Button","Rmiddle2Button","Rindex2Button"]))
+    pickerButtons["RBottomFingerRow"] = cmds.button(bgc = [.33, .98, .172], width = 20, height = 20, label = "", c = partial(selectMultiFingers, ["R_Pinky_1_Anim","R_Ring_1_Anim","R_Middle_1_Anim","R_Index_1_Anim"],["Rpinky1Button","Rring1Button","Rmiddle1Button","Rindex1Button"]))
+    
+    pickerButtons["RPinky3Button"] = cmds.button("Rpinky3Button", bgc = [.196, .475, .98], width = 20, height = 20, label = "", c = partial(selectFinger, "R_Pinky_3_Anim", "Rpinky3Button"))
+    pickerButtons["RPinky2Button"] = cmds.button("Rpinky2Button", bgc = [.196, .475, .98], width = 20, height = 20, label = "", c = partial(selectFinger, "R_Pinky_2_Anim", "Rpinky2Button"))
+    pickerButtons["RPinky1Button"] = cmds.button("Rpinky1Button", bgc = [.196, .475, .98], width = 20, height = 20, label = "", c = partial(selectFinger, "R_Pinky_1_Anim", "Rpinky1Button"))
+
+    pickerButtons["RRing3Button"] = cmds.button("Rring3Button", bgc = [.196, .475, .98], width = 20, height = 20, label = "", c = partial(selectFinger, "R_Ring_3_Anim", "Rring3Button"))
+    pickerButtons["RRing2Button"] = cmds.button("Rring2Button", bgc = [.196, .475, .98], width = 20, height = 20, label = "", c = partial(selectFinger, "R_Ring_2_Anim", "Rring2Button"))
+    pickerButtons["RRing1Button"] = cmds.button("Rring1Button", bgc = [.196, .475, .98], width = 20, height = 20, label = "", c = partial(selectFinger, "R_Ring_1_Anim", "Rring1Button"))
+    
+    pickerButtons["RMiddle3Button"] = cmds.button("Rmiddle3Button", bgc = [.196, .475, .98], width = 20, height = 20, label = "", c = partial(selectFinger, "R_Middle_3_Anim", "Rmiddle3Button"))
+    pickerButtons["RMiddle2Button"] = cmds.button("Rmiddle2Button", bgc = [.196, .475, .98], width = 20, height = 20, label = "", c = partial(selectFinger, "R_Middle_2_Anim", "Rmiddle2Button"))
+    pickerButtons["RMiddle1Button"] = cmds.button("Rmiddle1Button", bgc = [.196, .475, .98], width = 20, height = 20, label = "", c = partial(selectFinger, "R_Middle_1_Anim", "Rmiddle1Button"))
+    
+    pickerButtons["RIndex3Button"] = cmds.button("Rindex3Button", bgc = [.196, .475, .98], width = 20, height = 20, label = "", c = partial(selectFinger, "R_Index_3_Anim", "Rindex3Button"))
+    pickerButtons["RIndex2Button"] = cmds.button("Rindex2Button", bgc = [.196, .475, .98], width = 20, height = 20, label = "", c = partial(selectFinger, "R_Index_2_Anim", "Rindex2Button"))
+    pickerButtons["RIndex1Button"] = cmds.button("Rindex1Button", bgc = [.196, .475, .98], width = 20, height = 20, label = "", c = partial(selectFinger, "R_Index_1_Anim", "Rindex1Button"))
+    
+    pickerButtons["RThumb3Button"] = cmds.button("Rthumb3Button", bgc = [.196, .475, .98], width = 20, height = 20, label = "", c = partial(selectFinger, "R_Thumb_3_Anim", "Rthumb3Button"))
+    pickerButtons["RThumb2Button"] = cmds.button("Rthumb2Button", bgc = [.196, .475, .98], width = 20, height = 20, label = "", c = partial(selectFinger, "R_Thumb_2_Anim", "Rthumb2Button"))
+    pickerButtons["RThumb1Button"] = cmds.button("Rthumb1Button", bgc = [.196, .475, .98], width = 20, height = 20, label = "", c = partial(selectFinger, "R_Thumb_1_Anim", "Rthumb1Button"))
+    pickerButtons["RThumbAll"] = cmds.button(bgc = [.33, .98, .172], width = 20, height = 20, label = "", c = partial(selectMultiFingers, ["R_Thumb_3_Anim","R_Thumb_2_Anim","R_Thumb_1_Anim"],["Rthumb3Button","Rthumb2Button","Rthumb1Button"]))
+    
+    
+    
+        
+    pickerButtons["LPinkyColButton"] = cmds.button(bgc = [.33, .98, .172], width = 20, height = 20, label = "", c = partial(selectMultiFingers, ["L_Pinky_3_Anim","L_Pinky_2_Anim","L_Pinky_1_Anim"],["Lpinky3Button","Lpinky2Button","Lpinky1Button"]))
+    pickerButtons["LRingColButton"] = cmds.button(bgc = [.33, .98, .172], width = 20, height = 20, label = "", c = partial(selectMultiFingers, ["L_Ring_3_Anim","L_Ring_2_Anim","L_Ring_1_Anim"],["Lring3Button","Lring2Button","Lring1Button"]))
+    pickerButtons["LMiddleColButton"] = cmds.button(bgc = [.33, .98, .172], width = 20, height = 20, label = "", c = partial(selectMultiFingers, ["L_Middle_3_Anim","L_Middle_2_Anim","L_Middle_1_Anim"],["Lmiddle3Button","Lmiddle2Button","Lmiddle1Button"]))
+    pickerButtons["LIndexColButton"] = cmds.button(bgc = [.33, .98, .172], width = 20, height = 20, label = "", c = partial(selectMultiFingers, ["L_Index_3_Anim","L_Index_2_Anim","L_Index_1_Anim"],["Lindex3Button","Lindex2Button","Lindex1Button"]))
+    
+    pickerButtons["LTopFingerRow"] = cmds.button(bgc = [.33, .98, .172], width = 20, height = 20, label = "", c = partial(selectMultiFingers, ["L_Pinky_3_Anim","L_Ring_3_Anim","L_Middle_3_Anim","L_Index_3_Anim"],["Lpinky3Button","Lring3Button","Lmiddle3Button","Lindex3Button"]))
+    pickerButtons["LMidFingerRow"] = cmds.button(bgc = [.33, .98, .172], width = 20, height = 20, label = "", c = partial(selectMultiFingers, ["L_Pinky_2_Anim","L_Ring_2_Anim","L_Middle_2_Anim","L_Index_2_Anim"],["Lpinky2Button","Lring2Button","Lmiddle2Button","Lindex2Button"]))
+    pickerButtons["LBottomFingerRow"] = cmds.button(bgc = [.33, .98, .172], width = 20, height = 20, label = "", c = partial(selectMultiFingers, ["L_Pinky_1_Anim","L_Ring_1_Anim","L_Middle_1_Anim","L_Index_1_Anim"],["Lpinky1Button","Lring1Button","Lmiddle1Button","Lindex1Button"]))
+    
+    pickerButtons["LPinky3Button"] = cmds.button("Lpinky3Button", bgc = [.196, .475, .98], width = 20, height = 20, label = "", c = partial(selectFinger, "L_Pinky_3_Anim", "Lpinky3Button"))
+    pickerButtons["LPinky2Button"] = cmds.button("Lpinky2Button", bgc = [.196, .475, .98], width = 20, height = 20, label = "", c = partial(selectFinger, "L_Pinky_2_Anim", "Lpinky2Button"))
+    pickerButtons["LPinky1Button"] = cmds.button("Lpinky1Button", bgc = [.196, .475, .98], width = 20, height = 20, label = "", c = partial(selectFinger, "L_Pinky_1_Anim", "Lpinky1Button"))
+
+    pickerButtons["LRing3Button"] = cmds.button("Lring3Button", bgc = [.196, .475, .98], width = 20, height = 20, label = "", c = partial(selectFinger, "L_Ring_3_Anim", "Lring3Button"))
+    pickerButtons["LRing2Button"] = cmds.button("Lring2Button", bgc = [.196, .475, .98], width = 20, height = 20, label = "", c = partial(selectFinger, "L_Ring_2_Anim", "Lring2Button"))
+    pickerButtons["LRing1Button"] = cmds.button("Lring1Button", bgc = [.196, .475, .98], width = 20, height = 20, label = "", c = partial(selectFinger, "L_Ring_1_Anim", "Lring1Button"))
+    
+    pickerButtons["LMiddle3Button"] = cmds.button("Lmiddle3Button", bgc = [.196, .475, .98], width = 20, height = 20, label = "", c = partial(selectFinger, "L_Middle_3_Anim", "Lmiddle3Button"))
+    pickerButtons["LMiddle2Button"] = cmds.button("Lmiddle2Button", bgc = [.196, .475, .98], width = 20, height = 20, label = "", c = partial(selectFinger, "L_Middle_2_Anim", "Lmiddle2Button"))
+    pickerButtons["LMiddle1Button"] = cmds.button("Lmiddle1Button", bgc = [.196, .475, .98], width = 20, height = 20, label = "", c = partial(selectFinger, "L_Middle_1_Anim", "Lmiddle1Button"))
+    
+    pickerButtons["LIndex3Button"] = cmds.button("Lindex3Button", bgc = [.196, .475, .98], width = 20, height = 20, label = "", c = partial(selectFinger, "L_Index_3_Anim", "Lindex3Button"))
+    pickerButtons["LIndex2Button"] = cmds.button("Lindex2Button", bgc = [.196, .475, .98], width = 20, height = 20, label = "", c = partial(selectFinger, "L_Index_2_Anim", "Lindex2Button"))
+    pickerButtons["LIndex1Button"] = cmds.button("Lindex1Button", bgc = [.196, .475, .98], width = 20, height = 20, label = "", c = partial(selectFinger, "L_Index_1_Anim", "Lindex1Button"))
+    
+    pickerButtons["LThumb3Button"] = cmds.button("Lthumb3Button", bgc = [.196, .475, .98], width = 20, height = 20, label = "", c = partial(selectFinger, "L_Thumb_3_Anim", "Lthumb3Button"))
+    pickerButtons["LThumb2Button"] = cmds.button("Lthumb2Button", bgc = [.196, .475, .98], width = 20, height = 20, label = "", c = partial(selectFinger, "L_Thumb_2_Anim", "Lthumb2Button"))
+    pickerButtons["LThumb1Button"] = cmds.button("Lthumb1Button", bgc = [.196, .475, .98], width = 20, height = 20, label = "", c = partial(selectFinger, "L_Thumb_1_Anim", "Lthumb1Button"))
+    pickerButtons["LThumbAll"] = cmds.button(bgc = [.33, .98, .172], width = 20, height = 20, label = "", c = partial(selectMultiFingers, ["L_Thumb_3_Anim","L_Thumb_2_Anim","L_Thumb_1_Anim"],["Lthumb3Button","Lthumb2Button","Lthumb1Button"]))
+    
+    
+    pickerButtons["leftText"] = cmds.button(label = "Left", width = 145, enable = True, bgc = [.33, .98, .172], c = partial(selectMultiFingers, ["L_Thumb_3_Anim","L_Thumb_2_Anim","L_Thumb_1_Anim","L_Pinky_3_Anim","L_Pinky_2_Anim","L_Pinky_1_Anim", "L_Ring_3_Anim","L_Ring_2_Anim","L_Ring_1_Anim", "L_Middle_3_Anim","L_Middle_2_Anim","L_Middle_1_Anim", "L_Index_3_Anim","L_Index_2_Anim","L_Index_1_Anim"],["Lthumb3Button","Lthumb2Button","Lthumb1Button","Lpinky3Button","Lpinky2Button","Lpinky1Button","Lring3Button","Lring2Button","Lring1Button","Lmiddle3Button","Lmiddle2Button","Lmiddle1Button","Lindex3Button","Lindex2Button","Lindex1Button"]))
+    pickerButtons["rightText"] = cmds.button(label = "Right", width = 145, enable = True, bgc = [.33, .98, .172], c = partial(selectMultiFingers, ["R_Thumb_3_Anim","R_Thumb_2_Anim","R_Thumb_1_Anim","R_Pinky_3_Anim","R_Pinky_2_Anim","R_Pinky_1_Anim", "R_Ring_3_Anim","R_Ring_2_Anim","R_Ring_1_Anim", "R_Middle_3_Anim","R_Middle_2_Anim","R_Middle_1_Anim", "R_Index_3_Anim","R_Index_2_Anim","R_Index_1_Anim"],["Rthumb3Button","Rthumb2Button","Rthumb1Button","Rpinky3Button","Rpinky2Button","Rpinky1Button","Rring3Button","Rring2Button","Rring1Button","Rmiddle3Button","Rmiddle2Button","Rmiddle1Button","Rindex3Button","Rindex2Button","Rindex1Button"]))
+    
+    
+    
+    #place the buttons in the form layout
+    cmds.formLayout(handPickerFormLayout, edit = True, af = [(pickerButtons["RPinkyColButton"], 'top', 5), (pickerButtons["RPinkyColButton"], 'left', 30)])
+    cmds.formLayout(handPickerFormLayout, edit = True, af = [(pickerButtons["RRingColButton"], 'top', 5), (pickerButtons["RRingColButton"], 'left', 55)])
+    cmds.formLayout(handPickerFormLayout, edit = True, af = [(pickerButtons["RMiddleColButton"], 'top', 5), (pickerButtons["RMiddleColButton"], 'left', 80)])
+    cmds.formLayout(handPickerFormLayout, edit = True, af = [(pickerButtons["RIndexColButton"], 'top', 5), (pickerButtons["RIndexColButton"], 'left', 105)])
+   
+    cmds.formLayout(handPickerFormLayout, edit = True, af = [(pickerButtons["RTopFingerRow"], 'top', 30), (pickerButtons["RTopFingerRow"], 'left', 5)])
+    cmds.formLayout(handPickerFormLayout, edit = True, af = [(pickerButtons["RMidFingerRow"], 'top', 55), (pickerButtons["RMidFingerRow"], 'left', 5)])
+    cmds.formLayout(handPickerFormLayout, edit = True, af = [(pickerButtons["RBottomFingerRow"], 'top', 80), (pickerButtons["RBottomFingerRow"], 'left', 5)])
+   
+    cmds.formLayout(handPickerFormLayout, edit = True, af = [(pickerButtons["RPinky3Button"], 'top', 30), (pickerButtons["RPinky3Button"], 'left', 30)])
+    cmds.formLayout(handPickerFormLayout, edit = True, af = [(pickerButtons["RPinky2Button"], 'top', 55), (pickerButtons["RPinky2Button"], 'left', 30)])
+    cmds.formLayout(handPickerFormLayout, edit = True, af = [(pickerButtons["RPinky1Button"], 'top', 80), (pickerButtons["RPinky1Button"], 'left', 30)])
+    
+    cmds.formLayout(handPickerFormLayout, edit = True, af = [(pickerButtons["RRing3Button"], 'top', 30), (pickerButtons["RRing3Button"], 'left', 55)])
+    cmds.formLayout(handPickerFormLayout, edit = True, af = [(pickerButtons["RRing2Button"], 'top', 55), (pickerButtons["RRing2Button"], 'left', 55)])
+    cmds.formLayout(handPickerFormLayout, edit = True, af = [(pickerButtons["RRing1Button"], 'top', 80), (pickerButtons["RRing1Button"], 'left', 55)])
+    
+    cmds.formLayout(handPickerFormLayout, edit = True, af = [(pickerButtons["RMiddle3Button"], 'top', 30), (pickerButtons["RMiddle3Button"], 'left', 80)])
+    cmds.formLayout(handPickerFormLayout, edit = True, af = [(pickerButtons["RMiddle2Button"], 'top', 55), (pickerButtons["RMiddle2Button"], 'left', 80)])
+    cmds.formLayout(handPickerFormLayout, edit = True, af = [(pickerButtons["RMiddle1Button"], 'top', 80), (pickerButtons["RMiddle1Button"], 'left', 80)])
+    
+    cmds.formLayout(handPickerFormLayout, edit = True, af = [(pickerButtons["RIndex3Button"], 'top', 30), (pickerButtons["RIndex3Button"], 'left', 105)])
+    cmds.formLayout(handPickerFormLayout, edit = True, af = [(pickerButtons["RIndex2Button"], 'top', 55), (pickerButtons["RIndex2Button"], 'left', 105)])
+    cmds.formLayout(handPickerFormLayout, edit = True, af = [(pickerButtons["RIndex1Button"], 'top', 80), (pickerButtons["RIndex1Button"], 'left', 105)])
+    
+    cmds.formLayout(handPickerFormLayout, edit = True, af = [(pickerButtons["RThumb3Button"], 'top', 30), (pickerButtons["RThumb3Button"], 'left', 160)])
+    cmds.formLayout(handPickerFormLayout, edit = True, af = [(pickerButtons["RThumb2Button"], 'top', 55), (pickerButtons["RThumb2Button"], 'left', 145)])
+    cmds.formLayout(handPickerFormLayout, edit = True, af = [(pickerButtons["RThumb1Button"], 'top', 80), (pickerButtons["RThumb1Button"], 'left', 130)])
+    cmds.formLayout(handPickerFormLayout, edit = True, af = [(pickerButtons["RThumbAll"], 'top', 5), (pickerButtons["RThumbAll"], 'left', 160)])
+    
+    
+    
+    
+    cmds.formLayout(handPickerFormLayout, edit = True, af = [(pickerButtons["LPinkyColButton"], 'top', 5), (pickerButtons["LPinkyColButton"], 'right', 30)])
+    cmds.formLayout(handPickerFormLayout, edit = True, af = [(pickerButtons["LRingColButton"], 'top', 5), (pickerButtons["LRingColButton"], 'right', 55)])
+    cmds.formLayout(handPickerFormLayout, edit = True, af = [(pickerButtons["LMiddleColButton"], 'top', 5), (pickerButtons["LMiddleColButton"], 'right', 80)])
+    cmds.formLayout(handPickerFormLayout, edit = True, af = [(pickerButtons["LIndexColButton"], 'top', 5), (pickerButtons["LIndexColButton"], 'right', 105)])
+   
+    cmds.formLayout(handPickerFormLayout, edit = True, af = [(pickerButtons["LTopFingerRow"], 'top', 30), (pickerButtons["LTopFingerRow"], 'right', 5)])
+    cmds.formLayout(handPickerFormLayout, edit = True, af = [(pickerButtons["LMidFingerRow"], 'top', 55), (pickerButtons["LMidFingerRow"], 'right', 5)])
+    cmds.formLayout(handPickerFormLayout, edit = True, af = [(pickerButtons["LBottomFingerRow"], 'top', 80), (pickerButtons["LBottomFingerRow"], 'right', 5)])
+   
+    cmds.formLayout(handPickerFormLayout, edit = True, af = [(pickerButtons["LPinky3Button"], 'top', 30), (pickerButtons["LPinky3Button"], 'right', 30)])
+    cmds.formLayout(handPickerFormLayout, edit = True, af = [(pickerButtons["LPinky2Button"], 'top', 55), (pickerButtons["LPinky2Button"], 'right', 30)])
+    cmds.formLayout(handPickerFormLayout, edit = True, af = [(pickerButtons["LPinky1Button"], 'top', 80), (pickerButtons["LPinky1Button"], 'right', 30)])
+    
+    cmds.formLayout(handPickerFormLayout, edit = True, af = [(pickerButtons["LRing3Button"], 'top', 30), (pickerButtons["LRing3Button"], 'right', 55)])
+    cmds.formLayout(handPickerFormLayout, edit = True, af = [(pickerButtons["LRing2Button"], 'top', 55), (pickerButtons["LRing2Button"], 'right', 55)])
+    cmds.formLayout(handPickerFormLayout, edit = True, af = [(pickerButtons["LRing1Button"], 'top', 80), (pickerButtons["LRing1Button"], 'right', 55)])
+    
+    cmds.formLayout(handPickerFormLayout, edit = True, af = [(pickerButtons["LMiddle3Button"], 'top', 30), (pickerButtons["LMiddle3Button"], 'right', 80)])
+    cmds.formLayout(handPickerFormLayout, edit = True, af = [(pickerButtons["LMiddle2Button"], 'top', 55), (pickerButtons["LMiddle2Button"], 'right', 80)])
+    cmds.formLayout(handPickerFormLayout, edit = True, af = [(pickerButtons["LMiddle1Button"], 'top', 80), (pickerButtons["LMiddle1Button"], 'right', 80)])
+    
+    cmds.formLayout(handPickerFormLayout, edit = True, af = [(pickerButtons["LIndex3Button"], 'top', 30), (pickerButtons["LIndex3Button"], 'right', 105)])
+    cmds.formLayout(handPickerFormLayout, edit = True, af = [(pickerButtons["LIndex2Button"], 'top', 55), (pickerButtons["LIndex2Button"], 'right', 105)])
+    cmds.formLayout(handPickerFormLayout, edit = True, af = [(pickerButtons["LIndex1Button"], 'top', 80), (pickerButtons["LIndex1Button"], 'right', 105)])
+    
+    cmds.formLayout(handPickerFormLayout, edit = True, af = [(pickerButtons["LThumb3Button"], 'top', 30), (pickerButtons["LThumb3Button"], 'right', 160)])
+    cmds.formLayout(handPickerFormLayout, edit = True, af = [(pickerButtons["LThumb2Button"], 'top', 55), (pickerButtons["LThumb2Button"], 'right', 145)])
+    cmds.formLayout(handPickerFormLayout, edit = True, af = [(pickerButtons["LThumb1Button"], 'top', 80), (pickerButtons["LThumb1Button"], 'right', 130)])
+    cmds.formLayout(handPickerFormLayout, edit = True, af = [(pickerButtons["LThumbAll"], 'top', 5), (pickerButtons["LThumbAll"], 'right', 160)])
+    
+    
+    cmds.formLayout(handPickerFormLayout, edit = True, af = [(pickerButtons["leftText"], 'bottom', 0), (pickerButtons["leftText"], 'right', 5)])
+    cmds.formLayout(handPickerFormLayout, edit = True, af = [(pickerButtons["rightText"], 'bottom', 0), (pickerButtons["rightText"], 'left', 5)])
+    
+    #and now create the formLayout to house the channel box     
+    channelForm = cmds.formLayout(parent= tab1, h = 200,w = 400)
+    channelBox = cmds.channelBox(cat = False)
+    cmds.formLayout(channelForm, e=True, af=((channelBox, 'top', 0), (channelBox, 'left', 0), (channelBox, 'right', 0), (channelBox, 'bottom', 0)) )
+
+
+    
+    
+    
+    #show the window
+    string = "int $mayaVersion = `getApplicationVersionAsFloat` ; "
+    version = mel.eval(string)
+    
+    if version > 2010:
+        cmds.dockControl("Auto_Rig_UI",area = 'right', allowedArea = 'right', content = window, label = "Auto Rig UI")
+        
+    else:
+        cmds.showWindow(window)
+        
+    
+    #populate the UI option box
+    cmds.select("*:Rig_Container", replace = True)
+    selection = cmds.ls(sl = True)
+    
+    for each in selection:
+        name = each.partition(":")[0]
+        cmds.menuItem(label = name, parent = "selectedCharacterOptionMenu")
+        
+    #make sure fk arm orients are set to the correct setting
+    activeCharacterChangeCommand()
+        
+
+        
+        
+UI()
+
