Branching

diff --git a/Engine/Extras/Maya_AnimationRiggingTools/ARTv1/General/Scripts/jeRigTools/jeRigTools.py b/Engine/Extras/Maya_AnimationRiggingTools/ARTv1/General/Scripts/jeRigTools/jeRigTools.py
new file mode 100644
--- /dev/null
+++ b/Engine/Extras/Maya_AnimationRiggingTools/ARTv1/General/Scripts/jeRigTools/jeRigTools.py
@@ -0,0 +1,1268 @@
+import maya.cmds as cmds
+import maya.OpenMayaUI as mui
+import maya.mel as mel
+from PySide import QtGui, QtCore
+import shiboken, os, inspect, sys
+from functools import partial
+import json
+
+
+
+# # # # # # # # # # # # # # # # # # # # # # # # # # # # # # # # # # # # # # # # # # # # # # # # # # #
+def getMainWindow():
+    pointer = mui.MQtUtil.mainWindow()
+    return shiboken.wrapInstance(long(pointer), QtGui.QWidget)
+
+
+# # # # # # # # # # # # # # # # # # # # # # # # # # # # # # # # # # # # # # # # # # # # # # # # # # # #
+def returnNicePath(toolsPath, imagePath):
+    image =  os.path.normpath(os.path.join(toolsPath, imagePath))
+    if image.partition("\\")[2] != "":
+        image = image.replace("\\", "/")
+    return image
+
+# # # # # # # # # # # # # # # # # # # # # # # # # # # # # # # # # # # # # # # # # # # # # # # # # # # #
+def writeQSS(filePath):
+
+    #this function takes the qss file given, and finds and replaces any image path URLs using the user's settings for the icons path and changes the file on disk
+    iconPath = os.path.dirname(os.path.normpath(__file__))
+    iconPath = returnNicePath(iconPath, "icons")
+
+    print iconPath
+
+    f = open(filePath, "r")
+    lines = f.readlines()
+    f.close()
+
+    newLines = []
+    for line in lines:
+        if line.find("url(") != -1:
+            oldPath = line.partition("(")[2].rpartition("/")[0]
+
+
+            newLine =  line.replace(oldPath, iconPath)
+            newLines.append(newLine)
+        else:
+            newLines.append(line)
+
+    f = open(filePath, "w")
+    for line in newLines:
+        f.write(line)
+    f.close()
+
+
+
+
+
+class RigTools(QtGui.QMainWindow):
+    #Original Author: [REDACTED_PERSON]
+
+    def __init__(self, parent = None):
+
+        super(RigTools, self).__init__(parent)
+
+        filePath = os.path.normpath(__file__)
+        stylesheet = os.path.join(os.path.dirname(filePath), "mainScheme.qss")
+        writeQSS(stylesheet)
+
+
+        self.buildUI()
+
+# # # # # # # # # # # # # # # # # # # # # # # # # # # # # # # # # # # # # # # # # # # # # # # # # # # # # # # #
+# # # # # # # # # # # # # # # # # # # # # # # # # # # # # # # # # # # # # # # # # # # # # # # # # # # # # # # #
+# # # # # # # # # # # # # # # # # # # # # # # # # # # # # # # # # # # # # # # # # # # # # # # # # # # # # # # #
+# # # # # # # # # # # # # # # # # # # # # # # # # # # # # # # # # # # # # # # # # # # # # # # # # # # # # # # #
+# # # # # # # # # # # # # # # # # # # # # # # # # # # # # # # # # # # # # # # # # # # # # # # # # # # # # # # #
+# # # # # # # # # # # # # # # # # # # # # # # # # # # # # # # # # # # # # # # # # # # # # # # # # # # # # # # #
+    def buildUI(self):
+
+
+        #create the main widget
+        self.mainWidget = QtGui.QWidget()
+        self.setCentralWidget(self.mainWidget)
+
+        #load stylesheet
+        filePath = os.path.normpath(__file__)
+        stylesheet = os.path.join(os.path.dirname(filePath), "mainScheme.qss")
+        f = open(stylesheet, "r")
+        style = f.read()
+        f.close()
+
+        iconPath = os.path.dirname(os.path.normpath(__file__))
+        iconPath = returnNicePath(iconPath, "icons")
+
+        self.mainWidget.setStyleSheet(style)
+
+        #set qt object name
+        self.setObjectName("jeRigToolsWIN")
+        self.setWindowTitle("Rigging Tools")
+
+        #create the mainLayout for the rig creator UI
+        self.layout = QtGui.QVBoxLayout(self.mainWidget)
+
+        self.setMinimumSize(QtCore.QSize( 220, 300 ))
+        self.setMaximumSize(QtCore.QSize( 220, 1100 ))
+        self.setStyleSheet("""QToolTip { background-color: black; color: white;border: black solid 1px}""")
+        self.resize(220, 600)
+
+        #jump to section
+        layout = QtGui.QHBoxLayout()
+        self.layout.addLayout(layout)
+
+        label = QtGui.QLabel("Jump To:")
+        label.setStyleSheet("background: transparent;")
+        layout.addWidget(label)
+
+        self.section = QtGui.QComboBox()
+
+
+        layout.addWidget(self.section)
+        self.section.addItem("Constraints")
+        self.section.addItem("Common")
+        self.section.addItem("Freeze Xforms")
+        self.section.addItem("Lock/Unlock")
+        self.section.addItem("Rigging Tools")
+        self.section.addItem("Rig Controls")
+
+        #scroll layout
+        self.scrollArea = QtGui.QScrollArea()
+        self.scrollArea.setMinimumSize(QtCore.QSize(210, 600))
+        self.scrollArea.setMaximumSize(QtCore.QSize(210, 1200))
+        self.scrollArea.setSizePolicy(QtGui.QSizePolicy(QtGui.QSizePolicy.Fixed, QtGui.QSizePolicy.Expanding))
+        self.scrollArea.setHorizontalScrollBarPolicy(QtCore.Qt.ScrollBarAlwaysOff)
+        self.layout.addWidget(self.scrollArea)
+
+        #frame
+        self.scrollContents = QtGui.QFrame()
+        self.scrollContents.setSizePolicy(QtGui.QSizePolicy(QtGui.QSizePolicy.Fixed, QtGui.QSizePolicy.Expanding))
+        self.scrollContents.setObjectName("mid")
+        self.scrollContents.setMinimumSize(QtCore.QSize(210, 1200))
+        self.scrollContents.setMaximumSize(QtCore.QSize(210, 1200))
+        self.scrollArea.setWidget(self.scrollContents)
+
+        self.mainLayout = QtGui.QVBoxLayout(self.scrollContents)
+
+
+        #=======================================================================
+        # #section 1: constraints
+        #=======================================================================
+        self.constraintFrame = QtGui.QGroupBox("Constraints:")
+        self.constraintFrame.setMinimumSize(180, 250)
+        self.constraintFrame.setMaximumSize(180, 250)
+        self.mainLayout.addWidget(self.constraintFrame)
+
+        constraintLayout = QtGui.QVBoxLayout(self.constraintFrame)
+
+        self.constraintLayoutH = QtGui.QHBoxLayout()
+        constraintLayout.addLayout(self.constraintLayoutH)
+        self.pointConst = QtGui.QCheckBox("Position")
+        self.orientConst = QtGui.QCheckBox("Rotation")
+        self.constraintLayoutH.addWidget(self.pointConst)
+        self.constraintLayoutH.addWidget(self.orientConst)
+
+
+        label = QtGui.QLabel("Skip Translate:")
+        label.setStyleSheet("background: transparent;")
+        constraintLayout.addWidget(label)
+
+
+        self.constraintSkipLayout = QtGui.QHBoxLayout()
+        constraintLayout.addLayout(self.constraintSkipLayout)
+
+        self.skipTX = QtGui.QCheckBox("X")
+        self.skipTY = QtGui.QCheckBox("Y")
+        self.skipTZ = QtGui.QCheckBox("Z")
+
+        self.constraintSkipLayout.addWidget(self.skipTX)
+        self.constraintSkipLayout.addWidget(self.skipTY)
+        self.constraintSkipLayout.addWidget(self.skipTZ)
+
+
+
+        label = QtGui.QLabel("Skip Rotate:")
+        label.setStyleSheet("background: transparent;")
+        constraintLayout.addWidget(label)
+
+
+        self.constraintSkipLayoutR = QtGui.QHBoxLayout()
+        constraintLayout.addLayout(self.constraintSkipLayoutR)
+
+        self.skipRX = QtGui.QCheckBox("X")
+        self.skipRY = QtGui.QCheckBox("Y")
+        self.skipRZ = QtGui.QCheckBox("Z")
+
+        self.constraintSkipLayoutR.addWidget(self.skipRX)
+        self.constraintSkipLayoutR.addWidget(self.skipRY)
+        self.constraintSkipLayoutR.addWidget(self.skipRZ)
+
+        maintainLayout = QtGui.QHBoxLayout()
+        constraintLayout.addLayout(maintainLayout)
+
+        label = QtGui.QLabel("Maintain Offsets:")
+        label.setStyleSheet("background: transparent;")
+
+        maintainLayout.addWidget(label)
+
+        self.maintainOffsets = QtGui.QCheckBox()
+        maintainLayout.addWidget(self.maintainOffsets)
+
+
+
+        self.createConstButton = QtGui.QPushButton("Create Constraint")
+        self.createConstButton.setObjectName("blueButton")
+        constraintLayout.addWidget(self.createConstButton)
+        self.createConstButton.clicked.connect(partial(self.createConstraint, False))
+
+        self.snapConstButton = QtGui.QPushButton("Snap")
+        self.snapConstButton.setToolTip("Create and delete constraint given settings.")
+        self.snapConstButton.setObjectName("blueButton")
+        constraintLayout.addWidget(self.snapConstButton)
+        self.snapConstButton.clicked.connect(partial(self.createConstraint, True))
+
+        #=======================================================================
+        # #section 2: common tools
+        #=======================================================================
+        self.commonFrame = QtGui.QGroupBox("Common Tools:")
+        self.commonFrame.setMinimumSize(180, 160)
+        self.commonFrame.setMaximumSize(180, 160)
+        self.mainLayout.addWidget(self.commonFrame)
+
+        commonToolsLayout = QtGui.QVBoxLayout(self.commonFrame)
+
+
+        #row1
+        row1layout = QtGui.QHBoxLayout()
+        commonToolsLayout.addLayout(row1layout)
+
+        self.deleteHistBtn = QtGui.QPushButton("HIST")
+        self.deleteHistBtn.setToolTip("delete history")
+        self.deleteHistBtn.setObjectName("blueButton")
+        row1layout.addWidget(self.deleteHistBtn)
+        self.deleteHistBtn.clicked.connect(self.deleteHistory)
+
+        self.deleteNonDHistBtn = QtGui.QPushButton("NON-D")
+        self.deleteNonDHistBtn.setToolTip("delete non-deformer history")
+        self.deleteNonDHistBtn.setObjectName("blueButton")
+        row1layout.addWidget(self.deleteNonDHistBtn)
+        self.deleteNonDHistBtn.clicked.connect(self.deleteNonDeformHistory)
+
+        #row2
+        row2layout = QtGui.QHBoxLayout()
+        commonToolsLayout.addLayout(row2layout)
+
+        self.centerPivotBtn = QtGui.QPushButton("CP")
+        self.centerPivotBtn.setToolTip("center pivot")
+        self.centerPivotBtn.setObjectName("blueButton")
+        row2layout.addWidget(self.centerPivotBtn)
+        self.centerPivotBtn.clicked.connect(self.centerPivot)
+
+        self.copyPivotBtn = QtGui.QPushButton("Copy Piv")
+        self.copyPivotBtn.setToolTip("copy pivot: source -> target")
+        self.copyPivotBtn.setObjectName("blueButton")
+        row2layout.addWidget(self.copyPivotBtn)
+        self.copyPivotBtn.clicked.connect(self.copyPivot)
+
+        #row3
+        row3layout = QtGui.QHBoxLayout()
+        commonToolsLayout.addLayout(row3layout)
+
+        self.invertSelBtn = QtGui.QPushButton("IS")
+        self.invertSelBtn.setToolTip("invert selection")
+        self.invertSelBtn.setObjectName("blueButton")
+        row3layout.addWidget(self.invertSelBtn)
+        self.invertSelBtn.clicked.connect(self.invertSelection)
+
+        self.selectHiBtn = QtGui.QPushButton("SH")
+        self.selectHiBtn.setToolTip("select hierarchy")
+        self.selectHiBtn.setObjectName("blueButton")
+        row3layout.addWidget(self.selectHiBtn)
+        self.selectHiBtn.clicked.connect(self.selectHierarchy)
+
+        #row4
+        row4layout = QtGui.QHBoxLayout()
+        commonToolsLayout.addLayout(row4layout)
+
+        label = QtGui.QLabel("Convert To:")
+        label.setStyleSheet("background: transparent;")
+        row4layout.addWidget(label)
+
+        self.convertSelection = QtGui.QComboBox()
+        row4layout.addWidget(self.convertSelection)
+
+        self.convertSelection.addItem("Object")
+        self.convertSelection.addItem("Vertices")
+        self.convertSelection.addItem("Faces")
+        self.convertSelection.addItem("Edges")
+
+
+        self.convertGoBtn = QtGui.QPushButton("GO")
+        self.convertGoBtn.setMinimumSize(30,30)
+        self.convertGoBtn.setMaximumSize(30,30)
+        self.convertGoBtn.setObjectName("blueButton")
+        row4layout.addWidget(self.convertGoBtn)
+        self.convertGoBtn.clicked.connect(self.convertSel)
+
+        self.convertSelection.currentIndexChanged.connect(self.convertSel)
+
+        #=======================================================================
+        # #Freeze Transforms
+        #=======================================================================
+        self.freezeXformsFrame = QtGui.QGroupBox("Freeze Transforms:")
+        self.freezeXformsFrame.setMinimumSize(180, 90)
+        self.freezeXformsFrame.setMaximumSize(180, 90)
+        self.mainLayout.addWidget(self.freezeXformsFrame)
+
+        freezeMainLayout = QtGui.QVBoxLayout(self.freezeXformsFrame)
+
+        freeze1Layout = QtGui.QHBoxLayout()
+        freezeMainLayout.addLayout(freeze1Layout)
+
+        self.translateFreeze = QtGui.QCheckBox("T")
+        self.translateFreeze.setChecked(True)
+        self.rotateFreeze = QtGui.QCheckBox("R")
+        self.rotateFreeze.setChecked(True)
+        self.scaleFreeze = QtGui.QCheckBox("S")
+        self.scaleFreeze.setChecked(True)
+
+        freeze1Layout.addWidget(self.translateFreeze)
+        freeze1Layout.addWidget(self.rotateFreeze)
+        freeze1Layout.addWidget(self.scaleFreeze)
+
+
+        self.freezeBtn = QtGui.QPushButton("Freeze Xforms")
+        self.freezeBtn.setObjectName("blueButton")
+        freezeMainLayout.addWidget(self.freezeBtn)
+        self.freezeBtn.clicked.connect(self.freezeTransformations)
+
+        #=======================================================================
+        # #lock/unlock
+        #=======================================================================
+        self.lockUnlockFrame = QtGui.QGroupBox("Lock/Unlock:")
+        self.lockUnlockFrame.setMinimumSize(180, 160)
+        self.lockUnlockFrame.setMaximumSize(180, 160)
+        self.mainLayout.addWidget(self.lockUnlockFrame)
+
+        lockMainLayout = QtGui.QVBoxLayout(self.lockUnlockFrame)
+
+        #row 1
+        lock1Layout = QtGui.QHBoxLayout()
+        lockMainLayout.addLayout(lock1Layout)
+
+        self.lockBtn = QtGui.QPushButton("Lock")
+        self.lockBtn.setToolTip("Lock Selected Channels")
+        self.lockBtn.setObjectName("blueButton")
+        lock1Layout.addWidget(self.lockBtn)
+        self.lockBtn.clicked.connect(partial(self.modifyChannels, True, False, False))
+
+        self.hideBtn = QtGui.QPushButton("Hide")
+        self.hideBtn.setToolTip("Hide Selected Channels")
+        self.hideBtn.setObjectName("blueButton")
+        lock1Layout.addWidget(self.hideBtn)
+        self.hideBtn.clicked.connect(partial(self.modifyChannels, False, False, True))
+
+
+
+        #row 2
+        lock2Layout = QtGui.QHBoxLayout()
+        lockMainLayout.addLayout(lock2Layout)
+
+        self.lockHideBtn = QtGui.QPushButton("Lock/Hide")
+        self.lockHideBtn.setToolTip("Lock & Hide Selected Channels")
+        self.lockHideBtn.setObjectName("blueButton")
+        lock2Layout.addWidget(self.lockHideBtn)
+        self.lockHideBtn.clicked.connect(partial(self.modifyChannels, True, False, True))
+
+        self.unlockHideBtn = QtGui.QPushButton("Unlock")
+        self.unlockHideBtn.setToolTip("Unlock Selected Channels")
+        self.unlockHideBtn.setObjectName("blueButton")
+        lock2Layout.addWidget(self.unlockHideBtn)
+        self.unlockHideBtn.clicked.connect(partial(self.modifyChannels, False, True, False))
+
+        #row 3
+        self.channelControlBtn = QtGui.QPushButton("Channel Control")
+        self.channelControlBtn.setObjectName("blueButton")
+        lockMainLayout.addWidget(self.channelControlBtn)
+        self.channelControlBtn.clicked.connect(self.openChannelControl)
+
+        #lock/ unlock node
+
+        lock3Layout = QtGui.QHBoxLayout()
+        lockMainLayout.addLayout(lock3Layout)
+
+        self.lockNodeBtn = QtGui.QPushButton("Lock Node")
+        self.lockNodeBtn.setObjectName("blueButton")
+        lock3Layout.addWidget(self.lockNodeBtn)
+        self.lockNodeBtn.clicked.connect(partial(self.lockNodes, True))
+
+        self.unlockNodeBtn = QtGui.QPushButton("Unlock Node")
+        self.unlockNodeBtn.setObjectName("blueButton")
+        lock3Layout.addWidget(self.unlockNodeBtn)
+        self.unlockNodeBtn.clicked.connect(partial(self.lockNodes, False))
+
+        #=======================================================================
+        # #joint tools/weight tools
+        #=======================================================================
+        self.jointToolsFrame = QtGui.QGroupBox("Rigging Tools:")
+        self.jointToolsFrame.setMinimumSize(180, 170)
+        self.jointToolsFrame.setMaximumSize(180, 170)
+        self.mainLayout.addWidget(self.jointToolsFrame)
+
+        rigMainLayout = QtGui.QVBoxLayout(self.jointToolsFrame)
+
+
+        #create joint from selection
+        self.createJoitnFromSelBtn = QtGui.QPushButton("Create Joint From Sel")
+        self.createJoitnFromSelBtn.setObjectName("blueButton")
+        rigMainLayout.addWidget(self.createJoitnFromSelBtn)
+        self.createJoitnFromSelBtn.clicked.connect(self.createJointFromSelection)
+
+        #create joint from geo
+        self.createJoitnFromGeoBtn = QtGui.QPushButton("Create Joint From Geo")
+        self.createJoitnFromGeoBtn.setObjectName("blueButton")
+        rigMainLayout.addWidget(self.createJoitnFromGeoBtn)
+        self.createJoitnFromGeoBtn.clicked.connect(self.createJointFromGeo)
+
+        row1Layout = QtGui.QHBoxLayout()
+        rigMainLayout.addLayout(row1Layout)
+
+
+        #bindPose
+        self.goToBindPoseBtn = QtGui.QPushButton("BindPose")
+        self.goToBindPoseBtn.setToolTip("Go To BindPose")
+        self.goToBindPoseBtn.setObjectName("blueButton")
+        row1Layout.addWidget(self.goToBindPoseBtn)
+        self.goToBindPoseBtn.clicked.connect(self.goToBindPose)
+
+
+        #reset bindPose
+        self.resetBpBtn = QtGui.QPushButton("Reset BP")
+        self.resetBpBtn.setToolTip("Reset BindPose")
+        self.resetBpBtn.setObjectName("blueButton")
+        row1Layout.addWidget(self.resetBpBtn)
+        self.resetBpBtn.clicked.connect(self.resetBindPose)
+
+
+        row2Layout = QtGui.QHBoxLayout()
+        rigMainLayout.addLayout(row2Layout)
+
+
+        #weight hammer
+        self.weightHammerBtn = QtGui.QPushButton("Hammer")
+        self.weightHammerBtn.setToolTip("weight hammer")
+        self.weightHammerBtn.setObjectName("blueButton")
+        row2Layout.addWidget(self.weightHammerBtn)
+        self.weightHammerBtn.clicked.connect(self.weightHammer)
+
+        #skin wrangler or artv2 skinTools
+        self.skinWranglerBtn = QtGui.QPushButton("Wrangle")
+        self.skinWranglerBtn.setToolTip("Chris Evans: Skin Wrangler")
+        self.skinWranglerBtn.setObjectName("blueButton")
+        row2Layout.addWidget(self.skinWranglerBtn)
+        self.skinWranglerBtn.clicked.connect(self.skinWrangle)
+
+        #paint
+        self.paintWeightsBtn = QtGui.QPushButton()
+        self.paintWeightsBtn.setMinimumSize(30,30)
+        self.paintWeightsBtn.setMaximumSize(30,30)
+        icon = QtGui.QIcon(os.path.join(iconPath, "paint.png"))
+        self.paintWeightsBtn.setIconSize(QtCore.QSize(25,25))
+        self.paintWeightsBtn.setIcon(icon)
+        self.paintWeightsBtn.clicked.connect(self.paint)
+        row2Layout.addWidget(self.paintWeightsBtn)
+        #=======================================================================
+        # #create rig controls
+        #=======================================================================
+        self.rigControlsFrame = QtGui.QGroupBox("Rig Controls:")
+        self.rigControlsFrame.setMinimumSize(180, 200)
+        self.rigControlsFrame.setMaximumSize(180, 200)
+        self.mainLayout.addWidget(self.rigControlsFrame)
+
+        rigControlsMainLayout = QtGui.QVBoxLayout(self.rigControlsFrame)
+
+
+        #comboBox of shapes
+        rowlayout = QtGui.QHBoxLayout()
+        rigControlsMainLayout.addLayout(rowlayout)
+
+        label = QtGui.QLabel("Control Shape:")
+        label.setStyleSheet("background: transparent;")
+        rowlayout.addWidget(label)
+
+        self.createControlShape = QtGui.QComboBox()
+        rowlayout.addWidget(self.createControlShape)
+
+        self.createControlShape.addItem("Circle")
+        self.createControlShape.addItem("Square")
+        self.createControlShape.addItem("Cross")
+        self.createControlShape.addItem("Arrow")
+        self.createControlShape.addItem("2 Arrow")
+        self.createControlShape.addItem("Star")
+        self.createControlShape.addItem("90 Arrow")
+        self.createControlShape.addItem("180 Arrow")
+        self.createControlShape.addItem("Sphere")
+        self.createControlShape.addItem("Cube")
+        self.createControlShape.addItem("Pyramid")
+        self.createControlShape.addItem("Ball")
+
+
+        #color selector
+        row2layout = QtGui.QHBoxLayout()
+        rigControlsMainLayout.addLayout(row2layout)
+
+        frame = QtGui.QFrame()
+        frame.setMinimumSize(24,24)
+        frame.setMaximumSize(24,24)
+        row2layout.addWidget(frame)
+        frame.setStyleSheet("background: black;")
+        frameLayout = QtGui.QVBoxLayout(frame)
+        frameLayout.setContentsMargins(2,2,2,2)
+
+        self.currentColor = QtGui.QLabel()
+
+        self.currentColor.setMinimumSize(20,20)
+        self.currentColor.setMaximumSize(20,20)
+        self.currentColor.setStyleSheet("background: rgb(120, 120, 120);")
+        frameLayout.addWidget(self.currentColor)
+
+        self.colorChanger = QtGui.QSlider()
+        self.colorChanger.setOrientation(QtCore.Qt.Horizontal)
+        self.colorChanger.setRange(0, 31)
+        row2layout.addWidget(self.colorChanger)
+        self.colorChanger.valueChanged.connect(self.changeControlColor)
+
+        #create button
+        self.createControlBtn = QtGui.QPushButton("Create Control")
+        self.createControlBtn.setObjectName("blueButton")
+        rigControlsMainLayout.addWidget(self.createControlBtn)
+        self.createControlBtn.clicked.connect(self.createControl)
+
+        #hookup button
+
+        self.hookupLayout = QtGui.QHBoxLayout()
+        self.hookItUpBtn = QtGui.QPushButton("Hook It Up")
+        self.hookItUpBtn.setObjectName("blueButton")
+        rigControlsMainLayout.addWidget(self.hookItUpBtn)
+        self.hookItUpBtn.clicked.connect(self.hookup)
+
+        self.hookupCB = QtGui.QCheckBox("maintain offsets")
+        rigControlsMainLayout.addWidget(self.hookupCB)
+
+        #spacer
+        self.mainLayout.addSpacerItem(QtGui.QSpacerItem(0,50, QtGui.QSizePolicy.Fixed, QtGui.QSizePolicy.Expanding))
+
+
+        #=======================================================================
+        # #signals/slots
+        #=======================================================================
+        self.section.currentIndexChanged.connect(self.jumpToSection)
+
+
+# # # # # # # # # # # # # # # # # # # # # # # # # # # # # # # # # # # # # # # # # # # # # # # # # # # # # # # #
+# # # # # # # # # # # # # # # # # # # # # # # # # # # # # # # # # # # # # # # # # # # # # # # # # # # # # # # #
+# # # # # # # # # # # # # # # # # # # # # # # # # # # # # # # # # # # # # # # # # # # # # # # # # # # # # # # #
+# # # # # # # # # # # # # # # # # # # # # # # # # # # # # # # # # # # # # # # # # # # # # # # # # # # # # # # #
+# # # # # # # # # # # # # # # # # # # # # # # # # # # # # # # # # # # # # # # # # # # # # # # # # # # # # # # #
+# # # # # # # # # # # # # # # # # # # # # # # # # # # # # # # # # # # # # # # # # # # # # # # # # # # # # # # #
+    def jumpToSection(self):
+
+        value = self.section.currentText()
+
+        if value == "Constraints":
+
+            self.scrollArea.ensureWidgetVisible(self.constraintFrame)
+
+        if value == "Common":
+
+            self.scrollArea.ensureWidgetVisible(self.commonFrame)
+
+        if value == "Freeze Xforms":
+
+            self.scrollArea.ensureWidgetVisible(self.freezeXformsFrame)
+
+        if value == "Lock/Unlock":
+
+            self.scrollArea.ensureWidgetVisible(self.lockUnlockFrame)
+
+        if value == "Rigging Tools":
+
+            self.scrollArea.ensureWidgetVisible(self.jointToolsFrame)
+
+        if value == "Rig Controls":
+
+            self.scrollArea.ensureWidgetVisible(self.rigControlsFrame)
+
+
+
+# # # # # # # # # # # # # # # # # # # # # # # # # # # # # # # # # # # # # # # # # # # # # # # # # # # # # # # #
+# # # # # # # # # # # # # # # # # # # # # # # # # # # # # # # # # # # # # # # # # # # # # # # # # # # # # # # #
+# # # # # # # # # # # # # # # # # # # # # # # # # # # # # # # # # # # # # # # # # # # # # # # # # # # # # # # #
+# # # # # # # # # # # # # # # # # # # # # # # # # # # # # # # # # # # # # # # # # # # # # # # # # # # # # # # #
+# # # # # # # # # # # # # # # # # # # # # # # # # # # # # # # # # # # # # # # # # # # # # # # # # # # # # # # #
+# # # # # # # # # # # # # # # # # # # # # # # # # # # # # # # # # # # # # # # # # # # # # # # # # # # # # # # #
+    def createConstraint(self, snap):
+
+        #get selection
+        selection = cmds.ls(sl = True)
+
+        if not len(selection) > 0:
+            cmds.warning("Need at least two objects selected..")
+            return
+
+
+        target = selection[-1]
+        selection.pop(-1)
+
+        maintain = self.maintainOffsets.isChecked()
+
+        #get all values for point
+        if self.pointConst.isChecked():
+            skipValues = []
+            if self.skipTX.isChecked():
+                skipValues.append("x")
+
+            if self.skipTY.isChecked():
+                skipValues.append("x")
+
+            if self.skipTZ.isChecked():
+                skipValues.append("x")
+
+            constraint = cmds.pointConstraint(selection, target, mo = maintain, skip = skipValues)
+            if snap:
+                cmds.delete(constraint)
+
+
+        #get all values for orient
+        if self.orientConst.isChecked():
+            skipValues = []
+            if self.skipRX.isChecked():
+                skipValues.append("x")
+
+            if self.skipRY.isChecked():
+                skipValues.append("x")
+
+            if self.skipRZ.isChecked():
+                skipValues.append("x")
+
+            constraint = cmds.orientConstraint(selection, target, mo = maintain, skip = skipValues)
+            if snap:
+                cmds.delete(constraint)
+
+
+# # # # # # # # # # # # # # # # # # # # # # # # # # # # # # # # # # # # # # # # # # # # # # # # # # # # # # # #
+# # # # # # # # # # # # # # # # # # # # # # # # # # # # # # # # # # # # # # # # # # # # # # # # # # # # # # # #
+# # # # # # # # # # # # # # # # # # # # # # # # # # # # # # # # # # # # # # # # # # # # # # # # # # # # # # # #
+# # # # # # # # # # # # # # # # # # # # # # # # # # # # # # # # # # # # # # # # # # # # # # # # # # # # # # # #
+# # # # # # # # # # # # # # # # # # # # # # # # # # # # # # # # # # # # # # # # # # # # # # # # # # # # # # # #
+# # # # # # # # # # # # # # # # # # # # # # # # # # # # # # # # # # # # # # # # # # # # # # # # # # # # # # # #
+    def deleteHistory(self):
+
+        selection = cmds.ls(sl = True)
+        if not len(selection) > 0:
+            cmds.warning("Need at least one object selected..")
+            return
+
+        cmds.delete(selection, ch = True)
+
+# # # # # # # # # # # # # # # # # # # # # # # # # # # # # # # # # # # # # # # # # # # # # # # # # # # # # # # #
+# # # # # # # # # # # # # # # # # # # # # # # # # # # # # # # # # # # # # # # # # # # # # # # # # # # # # # # #
+# # # # # # # # # # # # # # # # # # # # # # # # # # # # # # # # # # # # # # # # # # # # # # # # # # # # # # # #
+# # # # # # # # # # # # # # # # # # # # # # # # # # # # # # # # # # # # # # # # # # # # # # # # # # # # # # # #
+# # # # # # # # # # # # # # # # # # # # # # # # # # # # # # # # # # # # # # # # # # # # # # # # # # # # # # # #
+# # # # # # # # # # # # # # # # # # # # # # # # # # # # # # # # # # # # # # # # # # # # # # # # # # # # # # # #
+    def deleteNonDeformHistory(self):
+
+        selection = cmds.ls(sl = True)
+        if not len(selection) > 0:
+            cmds.warning("Need at least one object selected..")
+            return
+
+        cmds.bakePartialHistory(selection, prePostDeformers = True)
+
+
+# # # # # # # # # # # # # # # # # # # # # # # # # # # # # # # # # # # # # # # # # # # # # # # # # # # # # # # #
+# # # # # # # # # # # # # # # # # # # # # # # # # # # # # # # # # # # # # # # # # # # # # # # # # # # # # # # #
+# # # # # # # # # # # # # # # # # # # # # # # # # # # # # # # # # # # # # # # # # # # # # # # # # # # # # # # #
+# # # # # # # # # # # # # # # # # # # # # # # # # # # # # # # # # # # # # # # # # # # # # # # # # # # # # # # #
+# # # # # # # # # # # # # # # # # # # # # # # # # # # # # # # # # # # # # # # # # # # # # # # # # # # # # # # #
+# # # # # # # # # # # # # # # # # # # # # # # # # # # # # # # # # # # # # # # # # # # # # # # # # # # # # # # #
+    def centerPivot(self):
+
+        selection = cmds.ls(sl = True)
+        if not len(selection) > 0:
+            cmds.warning("Need at least one object selected..")
+            return
+
+        for each in selection:
+            cmds.xform(each, cp = True)
+
+
+
+# # # # # # # # # # # # # # # # # # # # # # # # # # # # # # # # # # # # # # # # # # # # # # # # # # # # # # # #
+# # # # # # # # # # # # # # # # # # # # # # # # # # # # # # # # # # # # # # # # # # # # # # # # # # # # # # # #
+# # # # # # # # # # # # # # # # # # # # # # # # # # # # # # # # # # # # # # # # # # # # # # # # # # # # # # # #
+# # # # # # # # # # # # # # # # # # # # # # # # # # # # # # # # # # # # # # # # # # # # # # # # # # # # # # # #
+# # # # # # # # # # # # # # # # # # # # # # # # # # # # # # # # # # # # # # # # # # # # # # # # # # # # # # # #
+# # # # # # # # # # # # # # # # # # # # # # # # # # # # # # # # # # # # # # # # # # # # # # # # # # # # # # # #
+    def copyPivot(self):
+
+        selection = cmds.ls(sl = True)
+        if not len(selection) > 0:
+            cmds.warning("Need at least two objects selected..")
+            return
+
+        piv = cmds.xform(selection[0], q = True, ws = True, rp = True)
+        cmds.xform(selection[1], ws = True, piv = piv)
+
+
+# # # # # # # # # # # # # # # # # # # # # # # # # # # # # # # # # # # # # # # # # # # # # # # # # # # # # # # #
+# # # # # # # # # # # # # # # # # # # # # # # # # # # # # # # # # # # # # # # # # # # # # # # # # # # # # # # #
+# # # # # # # # # # # # # # # # # # # # # # # # # # # # # # # # # # # # # # # # # # # # # # # # # # # # # # # #
+# # # # # # # # # # # # # # # # # # # # # # # # # # # # # # # # # # # # # # # # # # # # # # # # # # # # # # # #
+# # # # # # # # # # # # # # # # # # # # # # # # # # # # # # # # # # # # # # # # # # # # # # # # # # # # # # # #
+# # # # # # # # # # # # # # # # # # # # # # # # # # # # # # # # # # # # # # # # # # # # # # # # # # # # # # # #
+    def invertSelection(self):
+
+        cmds.select("*", tgl = True)
+
+# # # # # # # # # # # # # # # # # # # # # # # # # # # # # # # # # # # # # # # # # # # # # # # # # # # # # # # #
+# # # # # # # # # # # # # # # # # # # # # # # # # # # # # # # # # # # # # # # # # # # # # # # # # # # # # # # #
+# # # # # # # # # # # # # # # # # # # # # # # # # # # # # # # # # # # # # # # # # # # # # # # # # # # # # # # #
+# # # # # # # # # # # # # # # # # # # # # # # # # # # # # # # # # # # # # # # # # # # # # # # # # # # # # # # #
+# # # # # # # # # # # # # # # # # # # # # # # # # # # # # # # # # # # # # # # # # # # # # # # # # # # # # # # #
+# # # # # # # # # # # # # # # # # # # # # # # # # # # # # # # # # # # # # # # # # # # # # # # # # # # # # # # #
+    def selectHierarchy(self):
+
+        selection = cmds.ls(sl = True)
+        if not len(selection) > 0:
+            cmds.warning("Need at least one object selected..")
+            return
+
+        cmds.select(selection, hi = True)
+
+
+# # # # # # # # # # # # # # # # # # # # # # # # # # # # # # # # # # # # # # # # # # # # # # # # # # # # # # # #
+# # # # # # # # # # # # # # # # # # # # # # # # # # # # # # # # # # # # # # # # # # # # # # # # # # # # # # # #
+# # # # # # # # # # # # # # # # # # # # # # # # # # # # # # # # # # # # # # # # # # # # # # # # # # # # # # # #
+# # # # # # # # # # # # # # # # # # # # # # # # # # # # # # # # # # # # # # # # # # # # # # # # # # # # # # # #
+# # # # # # # # # # # # # # # # # # # # # # # # # # # # # # # # # # # # # # # # # # # # # # # # # # # # # # # #
+# # # # # # # # # # # # # # # # # # # # # # # # # # # # # # # # # # # # # # # # # # # # # # # # # # # # # # # #
+    def convertSel(self):
+
+        selection = cmds.ls(sl = True)
+        if not len(selection) > 0:
+            cmds.warning("Need at least one object selected..")
+            return
+
+        #get value
+        mode = self.convertSelection.currentText()
+
+        if mode == "Vertices":
+            for each in selection:
+                vertexSel =  cmds.polyListComponentConversion(each, tv = True)
+                cmds.selectMode(co = True)
+                cmds.select(vertexSel)
+
+        if mode == "Faces":
+            for each in selection:
+                faceSel =  cmds.polyListComponentConversion(each, tf = True)
+                cmds.selectMode(co = True)
+                cmds.select(faceSel)
+
+        if mode == "Edges":
+            for each in selection:
+                edgeSel =  cmds.polyListComponentConversion(each, te = True)
+                cmds.selectMode(co = True)
+                cmds.select(edgeSel)
+
+        if mode == "Object":
+            for each in selection:
+                cmds.selectMode(o = True)
+
+
+
+# # # # # # # # # # # # # # # # # # # # # # # # # # # # # # # # # # # # # # # # # # # # # # # # # # # # # # # #
+# # # # # # # # # # # # # # # # # # # # # # # # # # # # # # # # # # # # # # # # # # # # # # # # # # # # # # # #
+# # # # # # # # # # # # # # # # # # # # # # # # # # # # # # # # # # # # # # # # # # # # # # # # # # # # # # # #
+# # # # # # # # # # # # # # # # # # # # # # # # # # # # # # # # # # # # # # # # # # # # # # # # # # # # # # # #
+# # # # # # # # # # # # # # # # # # # # # # # # # # # # # # # # # # # # # # # # # # # # # # # # # # # # # # # #
+# # # # # # # # # # # # # # # # # # # # # # # # # # # # # # # # # # # # # # # # # # # # # # # # # # # # # # # #
+    def freezeTransformations(self):
+
+        selection = cmds.ls(sl = True)
+        if not len(selection) > 0:
+            cmds.warning("Need at least one object selected..")
+            return
+
+        for each in selection:
+            freezeTranslate = self.translateFreeze.isChecked()
+            freezeRotate = self.rotateFreeze.isChecked()
+            freezeScale = self.scaleFreeze.isChecked()
+
+            cmds.makeIdentity(each, t = freezeTranslate, r = freezeRotate, s = freezeScale, apply = True)
+
+
+# # # # # # # # # # # # # # # # # # # # # # # # # # # # # # # # # # # # # # # # # # # # # # # # # # # # # # # #
+# # # # # # # # # # # # # # # # # # # # # # # # # # # # # # # # # # # # # # # # # # # # # # # # # # # # # # # #
+# # # # # # # # # # # # # # # # # # # # # # # # # # # # # # # # # # # # # # # # # # # # # # # # # # # # # # # #
+# # # # # # # # # # # # # # # # # # # # # # # # # # # # # # # # # # # # # # # # # # # # # # # # # # # # # # # #
+# # # # # # # # # # # # # # # # # # # # # # # # # # # # # # # # # # # # # # # # # # # # # # # # # # # # # # # #
+# # # # # # # # # # # # # # # # # # # # # # # # # # # # # # # # # # # # # # # # # # # # # # # # # # # # # # # #
+    def modifyChannels(self, lock, unlock, hide):
+        selection = cmds.ls(sl = True)
+        if not len(selection) > 0:
+            cmds.warning("Need at least one object selected..")
+            return
+
+        channels = mel.eval("channelBox -q -sma mainChannelBox;")
+
+        for each in selection:
+            for channel in channels:
+                if lock:
+                    cmds.setAttr(each + "." + channel, lock = True)
+                if unlock:
+                    cmds.setAttr(each + "." + channel, lock = False)
+                if hide:
+                    cmds.setAttr(each + "." + channel, channelBox = False, keyable = False)
+
+
+# # # # # # # # # # # # # # # # # # # # # # # # # # # # # # # # # # # # # # # # # # # # # # # # # # # # # # # #
+# # # # # # # # # # # # # # # # # # # # # # # # # # # # # # # # # # # # # # # # # # # # # # # # # # # # # # # #
+# # # # # # # # # # # # # # # # # # # # # # # # # # # # # # # # # # # # # # # # # # # # # # # # # # # # # # # #
+# # # # # # # # # # # # # # # # # # # # # # # # # # # # # # # # # # # # # # # # # # # # # # # # # # # # # # # #
+# # # # # # # # # # # # # # # # # # # # # # # # # # # # # # # # # # # # # # # # # # # # # # # # # # # # # # # #
+# # # # # # # # # # # # # # # # # # # # # # # # # # # # # # # # # # # # # # # # # # # # # # # # # # # # # # # #
+    def openChannelControl(self):
+        mel.eval("ChannelControlEditor;")
+
+
+# # # # # # # # # # # # # # # # # # # # # # # # # # # # # # # # # # # # # # # # # # # # # # # # # # # # # # # #
+# # # # # # # # # # # # # # # # # # # # # # # # # # # # # # # # # # # # # # # # # # # # # # # # # # # # # # # #
+# # # # # # # # # # # # # # # # # # # # # # # # # # # # # # # # # # # # # # # # # # # # # # # # # # # # # # # #
+# # # # # # # # # # # # # # # # # # # # # # # # # # # # # # # # # # # # # # # # # # # # # # # # # # # # # # # #
+# # # # # # # # # # # # # # # # # # # # # # # # # # # # # # # # # # # # # # # # # # # # # # # # # # # # # # # #
+# # # # # # # # # # # # # # # # # # # # # # # # # # # # # # # # # # # # # # # # # # # # # # # # # # # # # # # #
+    def lockNodes(self, lock):
+        selection = cmds.ls(sl = True)
+        if not len(selection) > 0:
+            cmds.warning("Need at least one object selected..")
+            return
+
+        for each in selection:
+            cmds.lockNode(each, lock = lock)
+
+            if lock:
+                sys.stdout.write(each + " is now locked.\n")
+            else:
+                sys.stdout.write(each + " is now unlocked.\n")
+
+
+# # # # # # # # # # # # # # # # # # # # # # # # # # # # # # # # # # # # # # # # # # # # # # # # # # # # # # # #
+# # # # # # # # # # # # # # # # # # # # # # # # # # # # # # # # # # # # # # # # # # # # # # # # # # # # # # # #
+# # # # # # # # # # # # # # # # # # # # # # # # # # # # # # # # # # # # # # # # # # # # # # # # # # # # # # # #
+# # # # # # # # # # # # # # # # # # # # # # # # # # # # # # # # # # # # # # # # # # # # # # # # # # # # # # # #
+# # # # # # # # # # # # # # # # # # # # # # # # # # # # # # # # # # # # # # # # # # # # # # # # # # # # # # # #
+# # # # # # # # # # # # # # # # # # # # # # # # # # # # # # # # # # # # # # # # # # # # # # # # # # # # # # # #
+    def createJointFromSelection(self):
+        selection = cmds.ls(sl = True)
+        if not len(selection) > 0:
+            cmds.warning("Need at least one object selected..")
+            return
+
+        emitter = cmds.EmitFromObject()
+        newPos = cmds.xform(emitter, q=True, t=True, ws=True)
+        cmds.select("emitter1", "particle1", r=True)
+        cmds.delete()
+        cmds.joint(position=newPos)
+
+
+# # # # # # # # # # # # # # # # # # # # # # # # # # # # # # # # # # # # # # # # # # # # # # # # # # # # # # # #
+# # # # # # # # # # # # # # # # # # # # # # # # # # # # # # # # # # # # # # # # # # # # # # # # # # # # # # # #
+# # # # # # # # # # # # # # # # # # # # # # # # # # # # # # # # # # # # # # # # # # # # # # # # # # # # # # # #
+# # # # # # # # # # # # # # # # # # # # # # # # # # # # # # # # # # # # # # # # # # # # # # # # # # # # # # # #
+# # # # # # # # # # # # # # # # # # # # # # # # # # # # # # # # # # # # # # # # # # # # # # # # # # # # # # # #
+# # # # # # # # # # # # # # # # # # # # # # # # # # # # # # # # # # # # # # # # # # # # # # # # # # # # # # # #
+    def createJointFromGeo(self):
+
+        import createJointFromGeo as createJointFromGeo
+        reload(createJointFromGeo)
+        createJointFromGeo.run()
+
+# # # # # # # # # # # # # # # # # # # # # # # # # # # # # # # # # # # # # # # # # # # # # # # # # # # # # # # #
+# # # # # # # # # # # # # # # # # # # # # # # # # # # # # # # # # # # # # # # # # # # # # # # # # # # # # # # #
+# # # # # # # # # # # # # # # # # # # # # # # # # # # # # # # # # # # # # # # # # # # # # # # # # # # # # # # #
+# # # # # # # # # # # # # # # # # # # # # # # # # # # # # # # # # # # # # # # # # # # # # # # # # # # # # # # #
+# # # # # # # # # # # # # # # # # # # # # # # # # # # # # # # # # # # # # # # # # # # # # # # # # # # # # # # #
+# # # # # # # # # # # # # # # # # # # # # # # # # # # # # # # # # # # # # # # # # # # # # # # # # # # # # # # #
+    def goToBindPose(self):
+
+        mel.eval("GoToBindPose")
+
+# # # # # # # # # # # # # # # # # # # # # # # # # # # # # # # # # # # # # # # # # # # # # # # # # # # # # # # #
+# # # # # # # # # # # # # # # # # # # # # # # # # # # # # # # # # # # # # # # # # # # # # # # # # # # # # # # #
+# # # # # # # # # # # # # # # # # # # # # # # # # # # # # # # # # # # # # # # # # # # # # # # # # # # # # # # #
+# # # # # # # # # # # # # # # # # # # # # # # # # # # # # # # # # # # # # # # # # # # # # # # # # # # # # # # #
+# # # # # # # # # # # # # # # # # # # # # # # # # # # # # # # # # # # # # # # # # # # # # # # # # # # # # # # #
+# # # # # # # # # # # # # # # # # # # # # # # # # # # # # # # # # # # # # # # # # # # # # # # # # # # # # # # #
+    def resetBindPose(self):
+
+        selection = cmds.ls(sl = True)
+        if not len(selection) > 0:
+            cmds.warning("Need at least one object selected..")
+            return
+
+        bindPoses = cmds.ls(type = "dagPose")
+        deletePoses = []
+        for pose in bindPoses:
+            joints = list(set(cmds.listConnections(pose, type = "joint")))
+            if selection[0] in joints:
+                deletePoses.append(pose)
+
+
+        cmds.delete(deletePoses)
+        cmds.select(selection[0])
+        cmds.dagPose(s = True, bp = True)
+
+# # # # # # # # # # # # # # # # # # # # # # # # # # # # # # # # # # # # # # # # # # # # # # # # # # # # # # # #
+# # # # # # # # # # # # # # # # # # # # # # # # # # # # # # # # # # # # # # # # # # # # # # # # # # # # # # # #
+# # # # # # # # # # # # # # # # # # # # # # # # # # # # # # # # # # # # # # # # # # # # # # # # # # # # # # # #
+# # # # # # # # # # # # # # # # # # # # # # # # # # # # # # # # # # # # # # # # # # # # # # # # # # # # # # # #
+# # # # # # # # # # # # # # # # # # # # # # # # # # # # # # # # # # # # # # # # # # # # # # # # # # # # # # # #
+# # # # # # # # # # # # # # # # # # # # # # # # # # # # # # # # # # # # # # # # # # # # # # # # # # # # # # # #
+    def weightHammer(self):
+
+        selection = cmds.ls(sl = True)
+        if not len(selection) > 0:
+            cmds.warning("Need at least one object selected..")
+            return
+
+        if selection[0].find("vtx") == -1:
+            cmds.warning("Need to be in vertex component selection..")
+            return
+
+        mel.eval("weightHammerVerts;")
+
+
+
+# # # # # # # # # # # # # # # # # # # # # # # # # # # # # # # # # # # # # # # # # # # # # # # # # # # # # # # #
+# # # # # # # # # # # # # # # # # # # # # # # # # # # # # # # # # # # # # # # # # # # # # # # # # # # # # # # #
+# # # # # # # # # # # # # # # # # # # # # # # # # # # # # # # # # # # # # # # # # # # # # # # # # # # # # # # #
+# # # # # # # # # # # # # # # # # # # # # # # # # # # # # # # # # # # # # # # # # # # # # # # # # # # # # # # #
+# # # # # # # # # # # # # # # # # # # # # # # # # # # # # # # # # # # # # # # # # # # # # # # # # # # # # # # #
+# # # # # # # # # # # # # # # # # # # # # # # # # # # # # # # # # # # # # # # # # # # # # # # # # # # # # # # #
+    def skinWrangle(self):
+        try:
+            import skinWrangler as sw
+            global skinWranglerWindow
+            skinWranglerWindow = sw.show()
+        except Exception, e:
+            sys.stdout.write(str(e) + "\n")
+
+
+# # # # # # # # # # # # # # # # # # # # # # # # # # # # # # # # # # # # # # # # # # # # # # # # # # # # # # # #
+# # # # # # # # # # # # # # # # # # # # # # # # # # # # # # # # # # # # # # # # # # # # # # # # # # # # # # # #
+# # # # # # # # # # # # # # # # # # # # # # # # # # # # # # # # # # # # # # # # # # # # # # # # # # # # # # # #
+# # # # # # # # # # # # # # # # # # # # # # # # # # # # # # # # # # # # # # # # # # # # # # # # # # # # # # # #
+# # # # # # # # # # # # # # # # # # # # # # # # # # # # # # # # # # # # # # # # # # # # # # # # # # # # # # # #
+# # # # # # # # # # # # # # # # # # # # # # # # # # # # # # # # # # # # # # # # # # # # # # # # # # # # # # # #
+    def paint(self):
+        selection = cmds.ls(sl = True)
+        if not len(selection) > 0:
+            cmds.warning("Need at least one object selected..")
+            return
+
+        mel.eval("ArtPaintSkinWeightsToolOptions;")
+
+# # # # # # # # # # # # # # # # # # # # # # # # # # # # # # # # # # # # # # # # # # # # # # # # # # # # # # # #
+# # # # # # # # # # # # # # # # # # # # # # # # # # # # # # # # # # # # # # # # # # # # # # # # # # # # # # # #
+# # # # # # # # # # # # # # # # # # # # # # # # # # # # # # # # # # # # # # # # # # # # # # # # # # # # # # # #
+# # # # # # # # # # # # # # # # # # # # # # # # # # # # # # # # # # # # # # # # # # # # # # # # # # # # # # # #
+# # # # # # # # # # # # # # # # # # # # # # # # # # # # # # # # # # # # # # # # # # # # # # # # # # # # # # # #
+# # # # # # # # # # # # # # # # # # # # # # # # # # # # # # # # # # # # # # # # # # # # # # # # # # # # # # # #
+    def changeControlColor(self, *args):
+
+        #get slider value
+        color = self.colorChanger.value()
+
+        if color == 0:
+            self.currentColor.setStyleSheet("background: rgb(120,120,120);")
+
+        if color == 1:
+            self.currentColor.setStyleSheet("background: rgb(0,0,0);")
+
+        if color == 2:
+            self.currentColor.setStyleSheet("background: rgb(64,64,64);")
+
+        if color == 3:
+            self.currentColor.setStyleSheet("background: rgb(128,128,128);")
+
+        if color == 4:
+            self.currentColor.setStyleSheet("background: rgb(155,0,40);")
+
+        if color == 5:
+            self.currentColor.setStyleSheet("background: rgb(0,4,96);")
+
+        if color == 6:
+            self.currentColor.setStyleSheet("background: rgb(0,0,255);")
+
+        if color == 7:
+            self.currentColor.setStyleSheet("background: rgb(0,70,25);")
+
+        if color == 8:
+            self.currentColor.setStyleSheet("background: rgb(38,0,67);")
+
+        if color == 9:
+            self.currentColor.setStyleSheet("background: rgb(200,0,200);")
+
+        if color == 10:
+            self.currentColor.setStyleSheet("background: rgb(138,72,51);")
+
+        if color == 11:
+            self.currentColor.setStyleSheet("background: rgb(63,35,31);")
+
+        if color == 12:
+            self.currentColor.setStyleSheet("background: rgb(153,38,0);")
+
+        if color == 13:
+            self.currentColor.setStyleSheet("background: rgb(255,0,0);")
+
+        if color == 14:
+            self.currentColor.setStyleSheet("background: rgb(0,255,0);")
+
+        if color == 15:
+            self.currentColor.setStyleSheet("background: rgb(0,65,153;")
+
+        if color == 16:
+            self.currentColor.setStyleSheet("background: rgb(255,255,255);")
+
+        if color == 17:
+            self.currentColor.setStyleSheet("background: rgb(255,255,0);")
+
+        if color == 18:
+            self.currentColor.setStyleSheet("background: rgb(100,220,255);")
+
+        if color == 19:
+            self.currentColor.setStyleSheet("background: rgb(67,255,163);")
+
+        if color == 20:
+            self.currentColor.setStyleSheet("background: rgb(255,176,176);")
+
+        if color == 21:
+            self.currentColor.setStyleSheet("background: rgb(228,172,121);")
+
+        if color == 22:
+            self.currentColor.setStyleSheet("background: rgb(255,255,99);")
+
+        if color == 23:
+            self.currentColor.setStyleSheet("background: rgb(0,153,84);")
+
+        if color == 24:
+            self.currentColor.setStyleSheet("background: rgb(161,105,48);")
+
+        if color == 25:
+            self.currentColor.setStyleSheet("background: rgb(159,161,48);")
+
+        if color == 26:
+            self.currentColor.setStyleSheet("background: rgb(104,161,48);")
+
+        if color == 27:
+            self.currentColor.setStyleSheet("background: rgb(48,161,93);")
+
+        if color == 28:
+            self.currentColor.setStyleSheet("background: rgb(48,161,161);")
+
+        if color == 29:
+            self.currentColor.setStyleSheet("background: rgb(48,103,161);")
+
+        if color == 30:
+            self.currentColor.setStyleSheet("background: rgb(111,48,161);")
+
+        if color == 31:
+            self.currentColor.setStyleSheet("background: rgb(161,48,105);")
+
+# # # # # # # # # # # # # # # # # # # # # # # # # # # # # # # # # # # # # # # # # # # # # # # # # # # # # # # #
+# # # # # # # # # # # # # # # # # # # # # # # # # # # # # # # # # # # # # # # # # # # # # # # # # # # # # # # #
+# # # # # # # # # # # # # # # # # # # # # # # # # # # # # # # # # # # # # # # # # # # # # # # # # # # # # # # #
+# # # # # # # # # # # # # # # # # # # # # # # # # # # # # # # # # # # # # # # # # # # # # # # # # # # # # # # #
+# # # # # # # # # # # # # # # # # # # # # # # # # # # # # # # # # # # # # # # # # # # # # # # # # # # # # # # #
+# # # # # # # # # # # # # # # # # # # # # # # # # # # # # # # # # # # # # # # # # # # # # # # # # # # # # # # #
+    def createControl(self):
+
+        selection = cmds.ls(sl = True)
+
+
+        controlType = self.createControlShape.currentText()
+        group = cmds.group(empty = True, name = "control_grp")
+
+        color = self.colorChanger.value()
+
+        #=======================================================================
+        # #Circle
+        #=======================================================================
+        if controlType == "Circle":
+            control = cmds.circle(c = (0,0,0), sw = 360, r = 30, d = 3, name = "control")[0]
+
+        #=======================================================================
+        # #Square
+        #=======================================================================
+        if controlType == "Square":
+            control = cmds.circle(c = (0,0,0), s = 4, sw = 360, r = 30, d = 1, name = "control")[0]
+            cmds.setAttr(control + ".rz", 45)
+            cmds.makeIdentity(control, r = True, apply = True)
+
+        #=======================================================================
+        # #Sphere
+        #=======================================================================
+        if controlType == "Sphere":
+
+            points = [(0, 0, 1), (0, 0.5, 0.866), (0, 0.866025, 0.5), (0, 1, 0), (0, 0.866025, -0.5), (0, 0.5, -0.866025), (0, 0, -1), (0, -0.5, -0.866025), (0, -0.866025, -0.5), (0, -1, 0), (0, -0.866025, 0.5), (0, -0.5, 0.866025), (0, 0, 1), (0.707107, 0, 0.707107), (1, 0, 0), (0.707107, 0, -0.707107), (0, 0, -1), (-0.707107, 0, -0.707107), (-1, 0, 0), (-0.866025, 0.5, 0), (-0.5, 0.866025, 0), (0, 1, 0), (0.5, 0.866025, 0), (0.866025, 0.5, 0), (1, 0, 0), (0.866025, -0.5, 0), (0.5, -0.866025, 0), (0, -1, 0), (-0.5, -0.866025, 0), (-0.866025, -0.5, 0), (-1, 0, 0), (-0.707107, 0, 0.707107), (0, 0, 1)]
+            control = cmds.curve(name = "control", d = 1, p = points)
+            cmds.setAttr(control + ".scaleX", 10)
+            cmds.setAttr(control + ".scaleY", 10)
+            cmds.setAttr(control + ".scaleZ", 10)
+
+        #=======================================================================
+        # #Ball
+        #=======================================================================
+        if controlType == "Ball":
+            control = cmds.curve(name = "control", d = 1, p = [(0.80718, 0.830576, 8.022739), (0.80718, 4.219206, 7.146586 ), (0.80718, 6.317059, 5.70073), (2.830981, 6.317059, 5.70073), (0, 8.422749, 2.94335), (-2.830981, 6.317059, 5.70073), (-0.80718, 6.317059, 5.70073), (-0.80718, 4.219352, 7.146486), (-0.80718, 0.830576, 8.022739), (-4.187851, 0.830576, 7.158003), (-6.310271, 0.830576, 5.705409), (-6.317059, 2.830981, 5.7007), (-8.422749, 0, 2.94335), (-6.317059, -2.830981, 5.70073), (-6.317059, -0.830576, 5.70073), (-4.225134, -0.830576, 7.142501), (-0.827872, -0.830576, 8.017446), (-0.80718, -4.176512, 7.160965), (-0.80718, -6.317059, 5.70073), (-2.830981, -6.317059, 5.70073), (0, -8.422749, 2.94335), (2.830981, -6.317059, 5.70073), (0.80718, -6.317059, 5.70073), (0.80718, -4.21137, 7.151987), (0.80718, -0.830576, 8.022739), (4.183345, -0.830576, 7.159155), (6.317059, -0.830576, 5.70073), (6.317059, -2.830981, 5.70073), (8.422749, 0, 2.94335), (6.317059, 2.830981, 5.70073), (6.317059, 0.830576, 5.70073), (4.263245, 0.830576, 7.116234), (0.80718, 0.830576, 8.022739)])
+            cmds.setAttr(control + ".scaleX", 5)
+            cmds.setAttr(control + ".scaleY", 5)
+            cmds.setAttr(control + ".scaleZ", 5)
+
+        #=======================================================================
+        # #Cube
+        #=======================================================================
+        if controlType == "Cube":
+            points = [[6.02647, -6.02647, 12.04647], [6.02647, 6.02647, 12.04647], [-6.02647, 6.02647, 12.04647], [-6.02647, -6.02647, 12.04647], [6.02647, -6.02647, 12.04647], [6.02647, -6.02647, -0.00647034],
+                      [-6.02647, -6.02647, -0.00647034], [-6.02647, -6.02647, 12.04647], [-6.02647, 6.02647, 12.04647], [-6.02647, 6.02647, -0.00647034], [-6.02647, -6.02647, -0.00647034], [6.02647, -6.02647, -0.00647034],
+                      [6.02647, -6.02647, 12.04647], [6.02647, 6.02647, 12.04647], [6.02647, 6.02647, -0.00647034], [6.02647, -6.02647, -0.00647034], [6.02647, -6.02647, 12.04647], [6.02647, 6.02647, 12.04647],
+                      [6.02647, 6.02647, -0.00647034], [-6.02647, 6.02647, -0.00647034], [-6.02647, 6.02647, 12.04647], [6.02647, 6.02647, 12.04647]]
+
+            control = cmds.curve(name = "control", d = 1, p = points)
+            cmds.setAttr(control + ".scaleX", 2)
+            cmds.setAttr(control + ".scaleY", 2)
+            cmds.setAttr(control + ".scaleZ", 2)
+
+        #=======================================================================
+        # #Star
+        #=======================================================================
+        if controlType == "Star":
+
+            points = [[0, -50, 0], [-9.384594, -28.817275, 0], [-13.30763, -27.227791, 0], [-17.480905, -24.77712, 0], [-20.596755, -22.274624, 0], [-23.833962, -18.757781, 0],
+                      [-26.349207, -14.98229, 0], [-28.759575, -9.558683, 0], [-50, 0, 0], [-28.770723, 9.525313, 0], [-26.651022, 14.434428, 0],  [-24.193906, 18.286961, 0],
+                      [-21.755189, 21.146159, 0], [-18.927363, 23.700649, 0], [-15.486534, 26.058461, 0], [-12.377072, 27.6617, 0], [-10.022831, 28.600144, 0], [0, 50, 0],
+                      [10.518781, 28.42069, 0], [13.171175, 27.293763, 0], [16.092683, 25.691857, 0], [19.166776, 23.509024, 0], [21.810012, 21.089516, 0], [24.005562, 18.525015, 0],
+                      [25.905111, 15.743927, 0], [27.583159, 12.551569, 0], [28.60507, 10.008838, 0], [50, 0, 0], [28.657458, -9.858696, 0], [27.801009, -12.060318, 0],
+                      [26.358507, -14.965793, 0], [23.550257, -19.115686, 0], [20.51645, -22.348314, 0], [17.109316, -24.885162, 0], [14.229729, -26.760131, 0], [10.358914, -28.479572, 0], [0, -50, 0]]
+
+            control = cmds.curve(name = "control", d = 1, p = points)
+
+        #=======================================================================
+        # #Pyramid
+        #=======================================================================
+        if controlType == "Pyramid":
+
+            points = [[17.227867, 17.227867, 0.418059], [-17.227864, 17.227869, 0.418059], [-17.227868, -17.227865, 0.418059], [17.227866, -17.227867, 0.418059], [17.227867, 17.227867, 0.418059],
+                      [0, 0, 24.781941], [17.227866, -17.227867, 0.418059], [-17.227868, -17.227865, 0.418059], [0, 0, 24.781941], [-17.227864, 17.227869, 0.418059],
+                      [17.227867, 17.227867, 0.418059],[0, 0, 24.781941]]
+
+            control = cmds.curve(name = "control", d = 1, p = points)
+
+        #=======================================================================
+        # #Double Arrow
+        #=======================================================================
+        if controlType == "2 Arrow":
+
+            points = [[-4.352071, -4.747714, 0], [-8.704142, -4.747714, 0], [0, -17.803927, 0], [8.704142, -4.747714, 0], [4.352071, -4.747714, 0], [4.352071, 4.747714, 0],
+                      [8.704142, 4.747714, 0], [0, 17.803927, 0], [-8.704142, 4.747714, 0], [-4.352071, 4.747714, 0], [-4.352071, -4.747714, 0]]
+
+            control = cmds.curve(name = "control", d = 1, p = points)
+
+        #=======================================================================
+        # #Single Arrow
+        #=======================================================================
+        if controlType == "Arrow":
+
+            points = [[-4.658613, -9.317226, 0], [4.658613, -9.317226, 0], [4.658613, 0, 0], [9.317226, 0, 0], [0, 13.975838, 0], [-9.317226, 0, 0], [-4.658613, 0, 0], [-4.658613, -9.317226, 0]]
+            control = cmds.curve(name = "control", d = 1, p = points)
+
+        #=======================================================================
+        # #Cross
+        #=======================================================================
+        if controlType == "Cross":
+
+            points = [[-2.004224, -2.004224, 0], [-2.004224, -10.02112, 0], [2.004224, -10.02112, 0], [2.004224, -2.004224, 0], [10.02112, -2.004224, 0], [10.02112, 2.004224, 0],
+                      [2.004224, 2.004224, 0], [2.004224, 10.02112, 0], [-2.004224, 10.02112, 0], [-2.004224, 2.004224, 0], [-10.02112, 2.004224, 0], [-10.02112, -2.004224, 0], [-2.004224, -2.004224, 0]]
+            control = cmds.curve(name = "control", d = 1, p = points)
+
+        #=======================================================================
+        # #90 Arrow
+        #=======================================================================
+        if controlType == "90 Arrow":
+
+            points = [[1.307282, -0.0501891, 0], [2.437608, -5.599536, 0], [5.505557, -10.185733, 0], [8.861274, -12.361307, 0], [7.3419, -5.481389, 0], [13.817009, -17.186448, 0],
+                      [0.473643, -16.243087, 0], [6.978455, -14.831142, 0], [3.281864, -12.52095, 0], [-0.541275, -6.817733, 0], [-1.920613, 0, 0], [1.297319, 0, 0]]
+
+            control = cmds.curve(name = "control", d = 1, p = points)
+
+
+        #=======================================================================
+        # #180 Arrow
+        #=======================================================================
+        if controlType == "180 Arrow":
+
+            points = [[1.209242, -0.0218849, 0], [2.311397, -5.533185, 0], [5.399375, -10.149323, 0], [8.729472, -12.300559, 0], [7.217565, -5.454451, 0],
+                      [13.660854, -17.101988, 0], [0.383062, -16.163263, 0], [6.855906, -14.758256, 0], [3.157234, -12.446763, 0], [-0.658571, -6.736789, 0],
+                      [-1.993435, 0.0301111, 0], [-0.658571, 6.736789, 0], [3.157234, 12.446763, 0], [6.884714, 14.727421, 0], [0.383062, 16.163263, 0],
+                      [13.660854, 17.101988, 0], [7.217565, 5.454451, 0], [8.754146, 12.252552, 0], [5.423028, 10.184681, 0], [2.296727, 5.511255, 0], [1.209242, 0.0330112, 0 ]]
+
+            control = cmds.curve(name = "control", d = 1, p = points)
+
+
+        cmds.setAttr(control + ".overrideEnabled", 1)
+        cmds.setAttr(control + ".overrideColor", color)
+
+        if len(selection) > 0:
+
+            constraint = cmds.parentConstraint(selection[0], group)[0]
+            cmds.delete(constraint)
+
+            constraint = cmds.pointConstraint(selection[0], control)[0]
+            cmds.delete(constraint)
+
+            control = cmds.rename(control, selection[0] + "_anim")
+            group = cmds.rename(group, selection[0] + "_anim_grp")
+
+        cmds.parent(control, group)
+        cmds.select(group)
+
+
+# # # # # # # # # # # # # # # # # # # # # # # # # # # # # # # # # # # # # # # # # # # # # # # # # # # # # # # #
+# # # # # # # # # # # # # # # # # # # # # # # # # # # # # # # # # # # # # # # # # # # # # # # # # # # # # # # #
+# # # # # # # # # # # # # # # # # # # # # # # # # # # # # # # # # # # # # # # # # # # # # # # # # # # # # # # #
+# # # # # # # # # # # # # # # # # # # # # # # # # # # # # # # # # # # # # # # # # # # # # # # # # # # # # # # #
+# # # # # # # # # # # # # # # # # # # # # # # # # # # # # # # # # # # # # # # # # # # # # # # # # # # # # # # #
+# # # # # # # # # # # # # # # # # # # # # # # # # # # # # # # # # # # # # # # # # # # # # # # # # # # # # # # #
+    def hookup(self):
+        selection = cmds.ls(sl = True)
+        if not len(selection) > 0:
+            cmds.warning("Need at least two objects selected.. (Control, bone)")
+            return
+        #get value of checkbox
+        value = self.hookupCB.isChecked()
+
+
+        cmds.makeIdentity(selection[0], t = True, s = True, r = True, apply = True)
+        cmds.pointConstraint(selection[0], selection[1], mo = value)
+        cmds.orientConstraint(selection[0], selection[1], mo = value)
+        cmds.scaleConstraint(selection[0], selection[1], mo = value)
+
+# # # # # # # # # # # # # # # # # # # # # # # # # # # # # # # # # # # # # # # # # # # # # # # # # # # # # # # #
+def run():
+
+    if cmds.window("jeRigToolsWIN", exists = True):
+        cmds.deleteUI("jeRigToolsWIN", wnd = True)
+
+    inst = RigTools(getMainWindow())
+    inst.show()
+
diff --git a/Engine/Extras/Maya_AnimationRiggingTools/ARTv1/General/Scripts/je_shotReferenceTool.py b/Engine/Extras/Maya_AnimationRiggingTools/ARTv1/General/Scripts/je_shotReferenceTool.py
new file mode 100644
--- /dev/null
+++ b/Engine/Extras/Maya_AnimationRiggingTools/ARTv1/General/Scripts/je_shotReferenceTool.py
@@ -0,0 +1,442 @@
+import maya.cmds as cmds
+import os
+from functools import partial
+
+
+def UI():
+    
+    #check to see if the window exists, if so delete it
+    if cmds.window("je_shotReference_UI", exists = True):
+        cmds.deleteUI("je_shotReference_UI")
+        
+    #create the window    
+    window = cmds.window("je_shotReference_UI", title = "Shot Reference", w = 400, h = 200, sizeable = False, mxb= False, mnb = False)
+    
+    #create the mainLayout
+    mainLayout = cmds.columnLayout(w = 400, h = 200)
+    
+    #create a row column layout
+    rowColumnLayout = cmds.rowColumnLayout(nc =2, cw = [(1, 300), (2, 100)], columnOffset = [(1,"both", 5), (2, "both", 5)])
+    
+
+    #Create fields for the text file
+    
+    cmds.separator(h = 10, style = "none")
+    cmds.separator(h = 10, style = "none")
+    
+    cmds.text(label = "Reference Quicktime Movie:", align = "left")
+    cmds.text(label = "")
+    
+    cmds.textField("refQuickTimeMovieTextField", w = 300)
+    cmds.button(label = "Browse", w = 50, c = partial(shotReferenceTool_getTextFieldInfo, "*.mov", 1, "refQuickTimeMovieTextField"))
+
+    
+    cmds.separator(h = 10, style = "none")
+    cmds.separator(h = 10, style = "none")
+    
+    
+    #create a new row column layout for the frame offset values
+    frameOffsetRowColumnLayout = cmds.rowColumnLayout(nc =4, cw = [(1, 100), (2, 100), (3, 100), (4, 100)], columnOffset = [(1,"both", 5), (2, "both", 5), (3, "both", 5), (4, "both", 5)],parent = mainLayout)
+    
+    #create the frame offset fields
+    cmds.text(label = "Frame Offset: ")
+    cmds.textField("frameOffsetTextField", w = 100)
+    cmds.text(label = "")
+    cmds.text(label = "")
+    
+    
+
+    
+    #create the process button
+    cmds.separator(h = 15, style = "none", parent = mainLayout)
+
+    cmds.button(label = "Import Reference Video", parent = mainLayout, w = 400, h = 50, c = importReferenceVideo_Process)
+    
+    
+    
+    #show the window
+    cmds.showWindow(window)
+    
+    
+    
+
+def shotReferenceTool_getTextFieldInfo(fileFilter, fileMode, textField, *args):
+    
+    filePath =  cmds.fileDialog2(fileMode= fileMode, dialogStyle=2, fileFilter = fileFilter)[0]
+    filePath = str(filePath)
+    
+    cmds.textField(textField, edit = True, text = filePath)
+    return filePath
+
+def importReferenceVideo_Process(*args):
+    
+    videoFile = cmds.textField("refQuickTimeMovieTextField", q = True, text = True)
+    videoNiceName = videoFile.rpartition("/")[2].partition(".")[0]
+    
+    #create the polygon plane
+    moviePlane = cmds.polyPlane(w = 720, h = 480, sx = 1, sy = 1, name = ( videoNiceName + "_ShotReferencePlane"))[0]
+    cmds.select(moviePlane)
+    cmds.addAttr(longName='transparency', defaultValue=0, minValue=0, maxValue=1, keyable = True)
+
+    cmds.createDisplayLayer(name = ( videoNiceName + "_Layer"), nr = True)
+    
+    cmds.select(clear = True)
+
+    cmds.setAttr(moviePlane + ".tz", -10000)
+    cmds.setAttr(moviePlane + ".rx", -180)
+    cmds.makeIdentity(moviePlane, apply = True, r = 1, t = 1, s = 1)
+
+    
+    #create the camera for the plane
+    camera = cmds.camera()[0]
+    cam = cmds.rename(camera, (videoNiceName + "_Reference_Cam"))
+    
+    #position the camera
+    constraint = cmds.parentConstraint(moviePlane, cam)[0]
+    cmds.delete(constraint)
+    
+    cmds.setAttr(cam + ".rotateX", 90)
+    cmds.setAttr(cam + ".translateY", -800)
+    cmds.setAttr(cam + "Shape.focalLength", 38)
+    
+    #constrain the camera to the plane
+    cmds.parentConstraint(moviePlane, cam, mo = True)
+    
+    #scale up the plane a tad, and lock all attrs
+    cmds.setAttr(moviePlane + ".sx", 1.055, lock = True)
+    cmds.setAttr(moviePlane + ".sy", 1.055, lock = True)
+    cmds.setAttr(moviePlane + ".sz", 1.055, lock = True)
+    cmds.setAttr(moviePlane + ".tx", lock = True)
+    cmds.setAttr(moviePlane + ".ty", lock = True)
+    cmds.setAttr(moviePlane + ".tz", lock = True)
+    cmds.setAttr(moviePlane + ".rx", lock = True)
+    cmds.setAttr(moviePlane + ".ry", lock = True)
+    cmds.setAttr(moviePlane + ".rz", lock = True)
+    
+    #lock down the camera attrs
+    cmds.setAttr(cam + ".tx", lock = True)
+    cmds.setAttr(cam + ".ty", lock = True)
+    cmds.setAttr(cam + ".tz", lock = True)
+    cmds.setAttr(cam + ".rx", lock = True)
+    cmds.setAttr(cam + ".ry", lock = True)
+    cmds.setAttr(cam + ".rz", lock = True)
+    cmds.setAttr(cam + ".sx", lock = True)
+    cmds.setAttr(cam + ".sy", lock = True)
+    cmds.setAttr(cam + ".sz", lock = True)
+    
+    
+    #create the movie shader/texture and attach to plane
+    movieShader = cmds.shadingNode("lambert", asShader = True, name = moviePlane + "_M")
+    moviePlaceNode = cmds.shadingNode("place2dTexture", asUtility = True)
+    movieNode = cmds.shadingNode("movie", asTexture = True)
+    
+    cmds.connectAttr(moviePlaceNode + ".outUV", movieNode + ".uvCoord")
+    cmds.connectAttr(moviePlaceNode + ".outUvFilterSize", movieNode + ".uvFilterSize")
+    cmds.connectAttr(movieNode + ".outColor", movieShader + ".color")
+    
+    #connect the transparency of the shader to the plane attr
+    cmds.connectAttr(moviePlane + ".transparency", movieShader + ".transparencyR")
+    cmds.connectAttr(moviePlane + ".transparency", movieShader + ".transparencyG")
+    cmds.connectAttr(moviePlane + ".transparency", movieShader + ".transparencyB")
+    
+    
+    cmds.setAttr(movieNode + ".fileTextureName", videoFile, type = "string")
+    
+    #add the material to the movie plane
+    cmds.select(moviePlane)
+    cmds.hyperShade(assign = movieShader)
+    cmds.select(clear = True)
+    
+    cmds.setAttr(movieNode + ".useFrameExtension", 1)
+    
+    #make sure that the playback options are updating all viewports
+    cmds.playbackOptions(v = "all")
+    
+    
+    #take the frame offset value, and create an expression
+    frameOffsetValue = cmds.textField("frameOffsetTextField", q = True, text = True)
+    if frameOffsetValue != "":
+        expressionString = (movieNode + ".frameExtension = frame + " + str(frameOffsetValue))
+        expression = cmds.expression(name = (movieNode + "_expression"), string = expressionString)
+        
+    else:
+        expressionString = (movieNode + ".frameExtension = frame")
+        expression = cmds.expression(name = (movieNode + "_expression"), string = expressionString)
+        
+        
+def listAllReferenceMovies_UI():
+    
+    #check to see if the window exists, if so delete it
+    if cmds.window("je_listAllReferenceMovies_UI", exists = True):
+        cmds.deleteUI("je_listAllReferenceMovies_UI")
+        
+    #create the window    
+    window = cmds.window("je_listAllReferenceMovies_UI", title = "Shot Reference Movies", w = 400, h = 200, sizeable = False, mxb= False, mnb = False)
+    
+    #create the mainLayout
+    mainLayout = cmds.columnLayout(w = 400, h = 200)
+    
+    #create a dropdown menu of all the movies in the scene
+    cmds.text(label = "Movies In Scene: ", align = 'left')
+    
+    cmds.separator(h = 10, style = "none", parent = mainLayout)
+    optionMenu = cmds.optionMenu("referenceMoviesOptionMenu", w = 400)
+    cmds.separator(h = 20, style = "none", parent = mainLayout)
+    
+    
+    
+    #find all the movies in the scene
+    movieFiles = []
+    
+    try:
+        cmds.select("*ShotReferencePlane_M*")
+        
+    except ValueError:
+        print "No Reference Videos in Scene"
+        return
+    
+    movieShaders = cmds.ls(sl = True)
+    
+    for each in movieShaders:
+        connections = cmds.listConnections(each, connections = True)
+        
+        for connection in connections:
+            if connection.find("movie") == 0:
+                movieNode = connection
+                
+                file = cmds.getAttr(movieNode + ".fileTextureName")
+                movieFiles.append(file)
+                
+    
+    if len(movieFiles) > 0:
+        for file in movieFiles:
+            cmds.menuItem( label= file, parent =  "referenceMoviesOptionMenu")
+            
+    else:
+        cmds.warning( "No Movies in Scene")
+        return
+        
+
+    
+    
+    #create a new row column layout for the frame offset values
+    frameOffsetRowColumnLayout = cmds.rowColumnLayout(nc =4, cw = [(1, 100), (2, 100), (3, 100), (4, 100)], columnOffset = [(1,"both", 5), (2, "both", 5), (3, "both", 5), (4, "both", 5)],parent = mainLayout)
+    
+    
+    #create the frame offset fields
+    cmds.text(label = "Frame Offset: ", align = 'left')
+    cmds.textField("movieFilesFrameOffsetTextField", w = 100)
+    cmds.text(label = "")
+    cmds.checkBox("hideVideoCheckBox", label = "Toggle Video", v = 0, cc = hideVideo)
+    
+    
+    cmds.separator(h = 20, style = "none")
+    cmds.separator(h = 20, style = "none")
+    cmds.separator(h = 20, style = "none")
+    cmds.separator(h = 20, style = "none")
+
+    
+    
+    #create an optionMenu for attaching reference plane to another camera
+    cmds.text(label = "Attach to Camera: ", align = 'left')
+    cameraOptionMenu = cmds.optionMenu("cameraOptionMenu", w = 100)
+    cmds.menuItem(label = "None")
+    
+    #get all the cameras and populate the option menu
+    cameras = cmds.ls(type = 'camera')
+    
+    for camera in cameras:
+        cam = cmds.listRelatives(camera, parent = True)[0]
+        cmds.menuItem(label = str(cam), parent = cameraOptionMenu)
+        
+    cmds.text(label = "")
+    cmds.text(label = "")
+    
+    cmds.separator(h = 20, style = "none")
+    cmds.separator(h = 20, style = "none")
+    cmds.separator(h = 20, style = "none")
+    cmds.separator(h = 20, style = "none")
+    
+    
+    
+    #create an optionMenu for PiP location
+    cmds.text(label = "PiP Location: ", align = 'left')
+    pipOptionMenu = cmds.optionMenu("pipOptionMenu", w = 100)
+    cmds.menuItem(label = "Bottom Left Corner")
+    cmds.menuItem(label = "Bottom Right Corner")
+    cmds.menuItem(label = "Top Left Corner")
+    cmds.menuItem(label = "Top Right Corner")
+    
+
+
+    cmds.text(label = "")
+    cmds.text(label = "")
+    
+
+    #create the process button
+    cmds.separator(h = 15, style = "none", parent = mainLayout)
+
+    cmds.button(label = "Apply Settings", parent = mainLayout, w = 400, h = 50, c = applyFrameOffsetToSelected)
+    cmds.button(label = "Delete Selected Video", parent = mainLayout, w = 400, h = 50, c = partial(deleteSelectedMovie, file))
+
+    
+    
+    #show the window
+    cmds.showWindow(window)
+    
+    
+def deleteSelectedMovie(file, *args):
+    
+    result = cmds.confirmDialog( title='Confirm', message='Are you sure?', button=['Yes','No'], defaultButton='Yes', cancelButton='No', dismissString='No' )
+    
+    if result == "Yes":
+        name = file.rpartition("/")[2].partition(".")[0]
+        shader = (name + "_ShotReferencePlane_M")
+        plane = (name + "_ShotReferencePlane")
+        layer = (name + "_Layer")
+        camera = (name + "_Reference_Cam")
+        cmds.delete([shader, plane, layer, camera])
+        
+    else:
+        return
+    
+    
+def applyFrameOffsetToSelected(*args):
+    selectVideo = cmds.optionMenu("referenceMoviesOptionMenu", q = True, v = True)
+    PiPLoc = cmds.optionMenu("pipOptionMenu", q = True, v = True)
+    videoNiceName = selectVideo.rpartition("/")[2].partition(".")[0]
+    
+    if selectVideo == "":
+        return
+    
+    #find the movie node that has the matching video file
+    cmds.select("*ShotReferencePlane_M*")
+    movieShaders = cmds.ls(sl = True)
+    
+    for each in movieShaders:
+        connections = cmds.listConnections(each, connections = True)
+        
+        for connection in connections:
+            if connection.find("movie") == 0:
+                movieNode = connection
+                
+                file = cmds.getAttr(movieNode + ".fileTextureName")
+                if file == selectVideo:
+
+                    
+                    #find offset value
+                    frameOffsetValue = cmds.textField("movieFilesFrameOffsetTextField", q = True, text = True)
+                    
+                    if frameOffsetValue != "":
+                    
+                        #delete the existing expression
+                        cmds.delete(movieNode + "_expression")
+                        
+                    
+                        #create the new expression
+                        expressionString = (movieNode + ".frameExtension = frame + " + str(frameOffsetValue))
+                        expression = cmds.expression(name = (movieNode + "_expression"), string = expressionString)
+                    
+                    
+    
+    #If the attach to camera is not None, then let's set that up now
+    attachToCamera = cmds.optionMenu("cameraOptionMenu", q = True, v = True)
+    
+    if attachToCamera != "None":
+        
+        #duplicate our plane, and position it to the selected camera
+        
+        if cmds.objExists(("ShotReferencePlane_" + attachToCamera + "_Grp")):
+            cmds.delete(("ShotReferencePlane_" + attachToCamera + "_Grp"))
+            
+            
+        dupePlane = cmds.duplicate((videoNiceName + "_ShotReferencePlane"), name = ("ShotReferencePlane_" + attachToCamera), rr = True)[0]
+        
+        cmds.setAttr(dupePlane + ".tx", lock = False)
+        cmds.setAttr(dupePlane + ".ty", lock = False)
+        cmds.setAttr(dupePlane + ".tz", lock = False)
+        
+        cmds.setAttr(dupePlane + ".rx", lock = False)
+        cmds.setAttr(dupePlane + ".ry", lock = False)
+        cmds.setAttr(dupePlane + ".rz", lock = False)
+        
+        cmds.setAttr(dupePlane + ".sx", lock = False)
+        cmds.setAttr(dupePlane + ".sy", lock = False)
+        cmds.setAttr(dupePlane + ".sz", lock = False)
+        
+        
+        dupePlaneGroup = cmds.group(empty = True, name = ("ShotReferencePlane_" + attachToCamera + "_Grp"))
+        
+        constraint = cmds.parentConstraint(dupePlane, dupePlaneGroup)
+        cmds.delete(constraint)
+        cmds.parent(dupePlane, dupePlaneGroup)
+        cmds.makeIdentity(dupePlane, apply = True, t = 1, r = 1, s = 0)
+        
+        
+        #constrain the dupePlane to the camera
+        pointConstraint = cmds.pointConstraint(attachToCamera, dupePlaneGroup)[0]
+        orientConstraint = cmds.orientConstraint(attachToCamera, dupePlaneGroup)[0]
+        
+        #offset the constraint values based on the PiP location set
+        cmds.setAttr((orientConstraint + ".offsetX"), -90)
+        
+        
+        #scale down the plane and offset it back
+        cmds.setAttr(dupePlane + ".sx", .008)
+        cmds.setAttr(dupePlane + ".sy", .008)
+        cmds.setAttr(dupePlane + ".sz", .008)
+        
+        if attachToCamera.find("SHOT_CAM") != -1:
+            cmds.setAttr(dupePlane + ".ty", 30)
+            
+        if attachToCamera.find("Face_Cam") != -1:
+            cmds.setAttr(dupePlane + ".ty", 10)
+        
+        
+        if PiPLoc == "Bottom Right Corner":
+            cmds.setAttr(dupePlane + ".tx", 3.25)
+            cmds.setAttr(dupePlane + ".tz", -2)
+            
+        if PiPLoc == "Bottom Left Corner":
+            cmds.setAttr(dupePlane + ".tx", -3.25)
+            cmds.setAttr(dupePlane + ".tz", -2)
+            
+        if PiPLoc == "Top Right Corner":
+            cmds.setAttr(dupePlane + ".tx", 3.25)
+            cmds.setAttr(dupePlane + ".tz", 2)
+            
+        if PiPLoc == "Top Left Corner":
+            cmds.setAttr(dupePlane + ".tx", -3.25)
+            cmds.setAttr(dupePlane + ".tz", 2)
+            
+            
+        #delete the constraints
+        cmds.delete([pointConstraint, orientConstraint])
+        
+        #create a parent constraint on the group to the cam
+        cmds.parentConstraint(attachToCamera, dupePlaneGroup, mo = True)
+        
+        
+        #clear selection
+        cmds.select(clear = True)
+        
+        
+        
+def hideVideo(*args):
+    
+    selectVideo = cmds.optionMenu("referenceMoviesOptionMenu", q = True, v = True)
+    videoNiceName = selectVideo.rpartition("/")[2].partition(".")[0]
+    
+    videoLayer = videoNiceName + "_Layer"
+    
+    vis = cmds.checkBox("hideVideoCheckBox", q = True, v = True)
+    
+    if vis == True:
+        cmds.setAttr(videoLayer + ".visibility", 0)
+        
+    if vis == False:
+        cmds.setAttr(videoLayer + ".visibility", 1)
+        
+    
+    
+
diff --git a/Engine/Extras/Maya_AnimationRiggingTools/ARTv1/General/Scripts/mayaSetup.py b/Engine/Extras/Maya_AnimationRiggingTools/ARTv1/General/Scripts/mayaSetup.py
new file mode 100644
--- /dev/null
+++ b/Engine/Extras/Maya_AnimationRiggingTools/ARTv1/General/Scripts/mayaSetup.py
@@ -0,0 +1,37 @@
+import maya.cmds as cmds
+import os
+import sys
+from functools import partial
+import maya.mel as mel
+
+
+    
+def setupTools():
+    
+    path = cmds.internalVar(usd = True) + "mayaTools.txt"
+    
+
+    f = open(path, 'r')
+    
+    mayaToolsDir = f.readline()
+    
+    path = mayaToolsDir + "/General/Scripts"
+    pluginPath = mayaToolsDir + "/General/Plugins"
+    
+    #look in sys.path to see if path is in sys.path. if not, add it
+    if not path in sys.path:
+        sys.path.append(path)
+        
+    #make sure MAYA_PLUG_IN_PATH has our plugin path
+    pluginPaths = os.environ["MAYA_PLUG_IN_PATH"]
+    pluginPaths = pluginPaths + ";" + pluginPath
+    os.environ["MAYA_PLUG_IN_PATH"] = pluginPaths
+        
+
+    
+    
+    #setup menu item in main window
+    import customMayaMenu as cmm
+    cmm.customMayaMenu()
+
+setupTools()
