Branching

diff --git a/Engine/Extras/Maya_AnimationRiggingTools/ARTv1/General/Scripts/perforceUtils.py b/Engine/Extras/Maya_AnimationRiggingTools/ARTv1/General/Scripts/perforceUtils.py
new file mode 100644
--- /dev/null
+++ b/Engine/Extras/Maya_AnimationRiggingTools/ARTv1/General/Scripts/perforceUtils.py
@@ -0,0 +1,1015 @@
+import maya.cmds as cmds
+from P4 import P4,P4Exception
+import os, cPickle
+from functools import partial
+
+# # # # # # # # # # # # # # # # # # # # # # # # # # # # # # # # # # # # # 
+def p4_getLatestRevision(fileName, *args):
+    fileArg = fileName
+    #try to connect
+    p4 = P4()
+    
+    try:
+        p4.connect()
+        
+    except:
+        cmds.confirmDialog(title = "Perforce", icon = "critical", message = "Unable to connect to perforce server.")
+        return
+    
+    #find currently opened file name
+    if fileName == None:
+	fileName = cmds.file(q = True, sceneName = True)
+	
+    syncFiles = []
+    
+    
+    try:
+	#client info
+	spec = p4.run( "client", "-o" )[0]
+	client = spec.get("Client")
+	owner = spec.get("Owner")	
+	p4.user = owner
+	p4.client = client
+        
+    except:
+        cmds.confirmDialog(title = "Perforce", icon = "critical", message = "Unable to obtain client spec information.")
+    
+    
+    
+    #try to get the current file revision on local, and compare to depot
+    try:
+        #Find out the revision of the local version of the file
+        myFile = p4.run_have(fileName)[0]
+        
+        #This will find the revision number of your local file.
+        localRevision = int(myFile['haveRev'])
+        
+        #find out the revision number of the depot version of the file
+        depotVersion = p4.run_files(myFile['depotFile'])[0]
+        
+        #find the depot file path
+        depotFile = depotVersion['depotFile']
+        
+        #find the depot revision number of the file
+        depotRevision = int(depotVersion['rev'])        
+    
+        #check for latest
+        if localRevision != depotRevision:
+            syncFiles.append(depotFile)
+                
+                
+        
+        #Check for scene references in the file
+        allRefs = []
+        references = cmds.file(q = True, reference = True)
+        for reference in references:
+            nestedRef = cmds.file(reference, q = True, reference = True)
+            allRefs.append(reference)
+            allRefs.append(nestedRef)
+            
+        #loop through all found references and check for latest
+        for ref in allRefs:
+            #get revision of local file
+            myFile = p4.run_have(ref)[0]
+            
+            #get revision number
+            localRefRevision = int(myFile['haveRev']) 
+            
+            #grab depot file info
+            depotRefVersion = p4.run_files(myFile['depotFile'])[0]
+            
+            #depot file path
+            depotFile = depotRefVersion['depotFile']
+            
+            #get depot's revision #
+            depotRefRevision = int(depotRefVersion['rev'])
+            
+            
+            #compare
+            if localRefRevision != depotRefRevision:
+                syncFiles.append(depotFile)
+            
+    
+        
+        
+        #if there are files to sync, do it now
+        if len(syncFiles) > 0:
+            message = "The following files are not at latest revision:\n\n"
+            for file in syncFiles:
+                message += file + "\n"
+                
+            result = cmds.confirmDialog(title = "Perforce", icon = "warning", message =  message, button = ["Sync", "Cancel"])
+            
+            if result == "Sync":
+                #sync files
+                for f in syncFiles:
+                    p4.run_sync(f)
+                    
+                #ask if user would like to reopen
+		if fileArg == None:
+		    result = cmds.confirmDialog(title = "Perforce", icon = "question", message =  "Sync Complete. Reopen file to get changes?", button = ["Yes", "Cancel"])
+		    
+		    if result == "Yes":
+			cmds.file(fileName, open = True, force = True)
+		    
+	else:
+	    cmds.confirmDialog(title = "Perforce", icon = "information", message = "This file is already at head revision.", button = "Close")
+                    
+        #disconnect from server
+        p4.disconnect()
+                
+        
+        
+    #Handle any p4 errors that come back from trying to run the above code
+    except P4Exception:
+        errorString = "The following errors were encountered:\n\n"
+        
+        for e in p4.errors:  
+            errorString += e + "\n"
+        
+        cmds.confirmDialog(title = "Perforce", icon = "critical", message = errorString)
+	p4.disconnect()
+        return 
+    
+    
+# # # # # # # # # # # # # # # # # # # # # # # # # # # # # # # # # # # # # 
+def p4_checkOutCurrentFile(fileName, *args):
+    fileArg = fileName
+    #try to connect
+    p4 = P4()
+    
+    try:
+        p4.connect()
+        
+    except:
+        cmds.confirmDialog(title = "Perforce", icon = "critical", message = "Unable to connect to perforce server.")
+        return False
+    
+    #find currently opened file name
+    if fileName == None:
+	fileName = cmds.file(q = True, sceneName = True)
+	
+    reopen = False
+    syncFiles = []
+    
+    
+    try:
+	#client info
+	spec = p4.run( "client", "-o" )[0]
+	client = spec.get("Client")
+	owner = spec.get("Owner")	
+	p4.user = owner
+	p4.client = client
+        
+    except:
+        cmds.confirmDialog(title = "Perforce", icon = "critical", message = "Unable to obtain client spec information.")
+        
+    
+    try:
+        #check to see if file is at head revision
+        myFile = p4.run_have(fileName)[0]
+        
+        #This will find the revision number of your local file.
+        localRevision = int(myFile['haveRev'])
+        
+        #find out the revision number of the depot version of the file
+        depotVersion = p4.run_files(myFile['depotFile'])[0]
+        
+        #find the depot file path
+        depotFile = depotVersion['depotFile']
+        
+        #find the depot revision number of the file
+        depotRevision = int(depotVersion['rev'])        
+    
+        #check for latest
+        if localRevision != depotRevision:
+            result = cmds.confirmDialog(title = "Perforce", icon = "warning", message = "This file is not at head revision. Please get latest and try again.", button = ["Get Latest", "Cancel"])
+            
+            if result == "Get Latest":
+                p4_getLatestRevision(fileArg)
+                p4.disconnect()
+		
+            else:
+                return False
+            
+        else:
+	    try:
+		#check to see if file is checked out
+		opened = p4.run_opened(depotFile)
+		
+		if len(opened) > 0:
+		    user = opened[0]['user']
+		    cmds.confirmDialog(title = "Perforce", icon = "warning", message = "This file is already checked out by: " + user, button = "Close")
+		    p4.disconnect()
+	    
+		else:
+		    #check out the file
+		    p4.run_edit(depotFile)
+		    cmds.confirmDialog(title = "Perfoce", icon = "information", message = "This file is now checked out.", button = "Close")
+		    p4.disconnect()
+		    #tools path
+		    toolsPath = cmds.internalVar(usd = True) + "mayaTools.txt"
+		    if os.path.exists(toolsPath):
+			
+			f = open(toolsPath, 'r')
+			mayaToolsDir = f.readline()
+			f.close()
+			
+		    return True
+		    
+            #Handle any p4 errors that come back from trying to run the above code
+            except P4Exception:
+                errorString = "The following errors were encountered:\n\n"
+                
+                for e in p4.errors:  
+                    errorString += e + "\n"
+                
+                cmds.confirmDialog(title = "Perforce", icon = "critical", message = errorString)
+                return False
+            
+
+
+    #Handle any p4 errors that come back from trying to run the above code
+    except P4Exception:
+        errorString = "The following errors were encountered:\n\n"
+        
+        for e in p4.errors:  
+            errorString += e + "\n"
+        
+        cmds.confirmDialog(title = "Perforce", icon = "critical", message = errorString)
+	p4.disconnect()
+        return False
+    
+    
+    
+# # # # # # # # # # # # # # # # # # # # # # # # # # # # # # # # # # # # # 
+def p4_getRevisionHistory(*args):
+    
+    #try to connect
+	p4 = P4()
+	
+	try:
+	    p4.connect()
+	    
+	except:
+	    cmds.confirmDialog(title = "Perforce", icon = "critical", message = "Unable to connect to perforce server.")
+	    return
+	
+	#find currently opened file name
+	clientFile = cmds.file(q = True, sceneName = True)
+	reopen = False
+	syncFiles = []
+	
+	
+	try:
+	    #client info
+	    spec = p4.run( "client", "-o" )[0]
+	    client = spec.get("Client")
+	    owner = spec.get("Owner")	
+	    p4.user = owner
+	    p4.client = client    
+
+	except:
+	    cmds.confirmDialog(title = "Perforce", icon = "critical", message = "Unable to obtain client spec information.")
+	    
+	    
+	#get revision history of current file
+	try:
+	    #check to see if file is at head revision
+	    myFile = p4.run_have(clientFile)[0]
+	    depotVersion = p4.run_files(myFile['depotFile'])[0]
+	    depotFile = depotVersion['depotFile']	    
+	    history = p4.run_changes(depotFile)
+	    
+	    info = ""
+	    
+	    for h in history:
+		user = h.get("user")
+		change = h.get("change")
+		desc = h.get("desc")
+		
+		if desc.find("\n") == -1:
+		    desc = desc + "...\n"
+		    
+		else:
+		    desc = desc.partition("\n")[0] + "...\n"
+		    
+		info += change + " by " + user + ": " + desc
+		
+	    
+	    #print report into a confirm dialog
+	    cmds.confirmDialog(title = "History", icon = "information", ma = "left",  message = info, button = "Close")
+	    p4.disconnect()
+		
+	    
+	    
+	
+	except P4Exception:
+	    errorString = "The following errors were encountered:\n\n"
+	    
+	    for e in p4.errors:  
+		errorString += e + "\n"
+	    
+	    cmds.confirmDialog(title = "Perforce", icon = "critical", message = errorString)
+	    p4.disconnect()
+	    return
+	
+	
+# # # # # # # # # # # # # # # # # # # # # # # # # # # # # # # # # # # # # 
+def p4_submitCurrentFile(fileName, desc, *args):
+    
+        fileArg = fileName
+        #try to connect
+	p4 = P4()
+	
+	
+	try:
+	    p4.connect()
+	    
+	except:
+	    cmds.confirmDialog(title = "Perforce", icon = "critical", message = "Unable to connect to perforce server.")
+	    return
+	
+	#find currently opened file name
+	if fileName == None:
+	    fileName = cmds.file(q = True, sceneName = True)
+	    
+	reopen = False
+	syncFiles = []
+	
+	
+	try:
+	    #client info
+	    spec = p4.run( "client", "-o" )[0]
+	    client = spec.get("Client")
+	    owner = spec.get("Owner")	
+	    p4.user = owner
+	    p4.client = client    
+	    
+	except:
+	    cmds.confirmDialog(title = "Perforce", icon = "critical", message = "Unable to obtain client spec information.")
+	    
+	    
+	
+	#SUBMIT
+	try:
+	    	    
+	    if desc == None:
+		result = cmds.promptDialog(title = "Perforce", message = "Please Enter a Description..", button = ["Accept", "Cancel"], defaultButton = "Accept", dismissString = "Cancel", cancelButton = "Cancel")
+		
+	    else:
+		result = "Accept"
+		
+	    #process
+	    if result == "Accept":
+		
+		#description = "test"
+		myFile = p4.run_have(fileName)[0]
+		depotVersion = p4.run_files(myFile['depotFile'])[0]
+		depotFile = depotVersion['depotFile']
+		
+		#check to see if file is checked out
+		opened = p4.run_opened(depotFile)
+		
+		if len(opened) > 0:
+		    opendBy = opened[0]['user']
+
+		    
+		    if opendBy.lower() != owner.lower():
+			cmds.confirmDialog(title = "Perforce", icon = "warning", message = "This file is already checked out by: " + opendBy, button = "Close")
+			p4.disconnect()
+			return
+		
+		    else:
+		    
+			#fetch the description
+			if desc == None:
+			    desc = cmds.promptDialog(q = True, text = True)	    
+			
+			#save the file locally (so depot and HD are in sync)
+			openedFile = cmds.file(q = True, sceneName = True)
+			saveFileName = openedFile.rpartition("/")[2]
+			if fileArg == None:
+			    cmds.file(f = True, save = True, options = "v = 0", type = "mayaBinary")
+			
+			
+			#grab the name of the file
+			fileNameWithExt = openedFile.rpartition("/")[2]
+			fileName = fileNameWithExt.rpartition(".")[0]
+			
+			description = (desc + "\n Affected Files: " + openedFile)
+			
+			#create new changelist
+			newChange = p4.fetch_change()
+			newChange._description = description
+			
+			#make sure we don't add existing default changelist files.
+			newChange._files = []
+			
+			#determine the new number so we can refetch it. 
+			newChangeNum = int(p4.save_change(newChange)[0].split()[1])
+	    
+	
+			#change changelist number
+			p4.run_reopen('-c', newChangeNum, depotFile)
+			
+			#submit the changelist
+			p4.run_submit('-c', newChangeNum)
+			
+			#tell the user submit was successful
+			if fileArg == None:
+			    cmds.confirmDialog(title = "Perforce", icon = "information", message = "Submit Operation was successful!", button = "Close")
+			    
+			else:
+			    return True
+			
+		else:
+		    #if the file is not checked out by the user, let them know
+		    result = cmds.confirmDialog(title = "Perforce", icon = "warning", message = "File is not checked out. Unable to continue submit operation on this file:\n\n" + fileName)
+		
+
+		
+	except P4Exception:
+	    if fileArg == None:
+		errorString = "The following errors were encountered:\n\n"
+		
+		for e in p4.errors:  
+		    errorString += e + "\n"
+		
+		cmds.confirmDialog(title = "Perforce", icon = "critical", message = errorString)
+		
+	    p4.disconnect()
+	    return False
+		
+# # # # # # # # # # # # # # # # # # # # # # # # # # # # # # # # # # # # # 			
+def p4_addAndSubmitCurrentFile(fileName, description, *args):
+    fileArg = fileName
+    #try to connect
+    p4 = P4()
+    
+    
+    try:
+	p4.connect()
+	
+    except:
+	cmds.confirmDialog(title = "Perforce", icon = "critical", message = "Unable to connect to perforce server.")
+	return
+    
+
+    #find currently opened file name
+    if fileName == None:
+	fileName = cmds.file(q = True, sceneName = True)
+    reopen = False
+    syncFiles = []
+    
+    
+    
+    try:
+	#client info
+	spec = p4.run( "client", "-o" )[0]
+	client = spec.get("Client")
+	owner = spec.get("Owner")	
+	p4.user = owner
+	p4.client = client
+	
+    except:
+	cmds.confirmDialog(title = "Perforce", icon = "critical", message = "Unable to obtain client spec information.")
+	
+    
+    #find currently opened file name
+    proceed = False
+    if fileArg == None:
+	fileName = cmds.file(q = True, sceneName = True)
+	
+	if fileName == "":
+	    cmds.confirmDialog(title = "Perforce", icon = "warning", message = "Cannot Add file to perforce as file has no name.", button = "Close")
+	    p4.disconnect()
+	    return
+	
+	else:
+	    proceed = True
+	    
+    else:
+	proceed = True
+	
+    
+    #if the file has a filename, 
+    if proceed:
+	try:
+	    
+	    
+	    clientRoot = p4.fetch_client(p4.client)._Root
+	    
+	    #check to make sure client root is in the client file path
+	    if os.path.normpath(fileName).find(os.path.normpath(clientRoot))  == 0:
+		#if it was, then get a description for the changelist
+		if description == None:
+		    result = cmds.promptDialog(title = "Perforce", message = "Please Enter a Description..", button = ["Accept", "Cancel"], defaultButton = "Accept", dismissString = "Cancel", cancelButton = "Cancel")
+		
+		else:
+		    result = "Accept"
+		    
+		
+		if result == "Accept":
+		    
+		    #get changelist description
+		    if description == None:
+			description = cmds.promptDialog(q = True, text = True)
+			
+		    
+		    #create changelist
+		    newChange = p4.fetch_change()
+		    newChange._description = description
+		    
+		    #make sure we don't add existing default changelist files.
+		    newChange._files = []
+			    
+		    #determine the new number so we can refetch it.
+		    newChangeNum = int(p4.save_change(newChange)[0].split()[1])		    
+		    
+		    
+		    #description = "test"
+		    p4.run_add('-c', newChangeNum, fileName)
+	
+		    #submit the changelist
+		    p4.run_submit('-c', newChangeNum)
+		    
+		    #tell user operation was successful
+		    if fileArg == None:
+			result = cmds.confirmDialog(title = "Perforce", icon = "information", message = "File has been successfully added to perforce and submitted!", button = ["Close", "Check Out File"])
+			
+			if result == "Close":
+			    p4.disconnect()
+			    return
+			
+			if result == "Check Out File":
+			    p4_checkOutCurrentFile(fileName)
+			    
+		    #return operation succuessful
+		    return True
+			
+			
+
+
+		else:
+		    p4.disconnect()
+		    return
+		
+		
+	    else:
+		cmds.confirmDialog(title = "Perforce", icon = "warning", message = "Cannot proceed. File is not under client's root, " + clientRoot, button = "Close")
+		p4.disconnect()
+		return False
+	    
+	    
+	    
+	    
+	except P4Exception:
+	    errorString = "The following errors were encountered:\n\n"
+	    
+	    for e in p4.errors:  
+		errorString += e + "\n"
+	    
+	    cmds.confirmDialog(title = "Perforce", icon = "critical", message = errorString)
+	    p4.disconnect()
+	    return False
+	
+	
+# # # # # # # # # # # # # # # # # # # # # # # # # # # # # # # # # # # # # 
+def p4_checkForUpdates(*args):
+    #try to connect
+    p4 = P4()
+    
+    #get maya tools path
+    toolsPath = cmds.internalVar(usd = True) + "mayaTools.txt"
+    if os.path.exists(toolsPath):
+	
+	f = open(toolsPath, 'r')
+	mayaToolsDir = f.readline()
+	f.close()
+	
+    #connect to p4
+    try:
+	p4.connect()
+	
+    except:
+	cmds.confirmDialog(title = "Perforce", icon = "critical", message = "Unable to connect to perforce server.")
+	return
+    
+    #find currently opened file name
+    clientFile = cmds.file(q = True, sceneName = True)
+    reopen = False
+    syncFiles = []
+    
+    
+    try:
+	#client info
+	spec = p4.run( "client", "-o" )[0]
+	client = spec.get("Client")
+	owner = spec.get("Owner")	
+	p4.user = owner
+	p4.client = client
+
+	
+    except:
+	cmds.confirmDialog(title = "Perforce", icon = "critical", message = "Unable to obtain client spec information.")
+	
+	
+    #this will check the maya tools directory in p4 for any updates
+    try:
+	syncFiles = []
+	clientRoot = p4.fetch_client(p4.client)._Root
+	depotDirectories = []
+	
+	#get current project
+	if os.path.exists(mayaToolsDir + "/General/Scripts/projectSettings.txt"):
+	    #read to find current project
+	    f = open(mayaToolsDir + "/General/Scripts/projectSettings.txt", 'r')
+	    settings = cPickle.load(f)
+	    f.close()
+	    
+	    #write out new settings
+	    project = settings.get("CurrentProject")
+	    
+	    if os.path.exists(mayaToolsDir + "/General/Scripts/" + project + "_Project_Settings.txt"):
+		#read the depot paths to sync from the project settings
+		f = open(mayaToolsDir + "/General/Scripts/" + project + "_Project_Settings.txt", 'r')
+		settings = cPickle.load(f)
+		f.close() 
+		depotDirectories = settings.get("depotPaths")		
+		print depotDirectories
+
+	
+	#look at each directory inside MayaTools
+	for dir in depotDirectories:
+	    depotFiles = p4.run_files(dir + "...")
+	    
+	    for each in depotFiles:
+		
+		#try to compare depot to local. It is possible that there are local files not in depot, and vise versa
+		try:
+		    fileInfo = p4.run_files(each['depotFile'])[0]
+		    depotFilePath = fileInfo['depotFile']
+		    fileName = depotFilePath.rpartition("/")[2]
+		    
+		    #compare local files
+		    localFile = p4.run_have(depotFilePath)[0]
+		    localRevision = int(localFile['haveRev'])
+		    depotRevision = int(fileInfo['rev'])
+		    
+		    if localRevision < depotRevision:
+			syncFiles.append(depotFilePath)
+			
+		    
+		except:
+		    
+		    try:
+			#check to see if it errors out because we don't have a local version of the file
+			fileInfo = p4.run_files(each['depotFile'])[0]
+			depotFilePath = fileInfo['depotFile']
+			fileName = depotFilePath.rpartition("/")[2]			
+			localFile = p4.run_have(depotFilePath)[0]
+			
+		    except:
+			action = each.get("action")
+			if action != "delete":
+			    syncFiles.append(depotFilePath)
+			    
+			pass
+		    
+		    
+		
+	
+	#check size of syncFiles and ask user if they want to sync
+	if len(syncFiles) > 0:
+	    result = cmds.confirmDialog(title = "MayaTools", icon = "warning", message = "There are new updates available to the depot directories specified by your project settings.", button = ["Update", "Not Now"])
+	    
+	    if result == "Update":
+		for file in syncFiles:
+		    p4.run_sync(file)
+		    
+		
+		cmds.confirmDialog(title = "MayaTools", icon = "information", message = "Tools are now up to date!", button = "Close")
+		p4.disconnect()
+		
+	    else:
+		p4.disconnect()
+		return
+		    
+
+    #handle any errors
+    except P4Exception:
+	errorString = "The following errors were encountered:\n\n"
+	
+	for e in p4.errors:  
+	    errorString += e + "\n"
+	
+	cmds.confirmDialog(title = "Perforce", icon = "critical", message = errorString)
+	p4.disconnect()
+	return
+    
+    
+
+   
+# # # # # # # # # # # # # # # # # # # # # # # # # # # # # # # # # # # # # 
+def createNewProject(*args):
+    
+    if cmds.window("createNewARTProject_Window", exists = True):
+	cmds.deleteUI("createNewARTProject_Window")
+	
+    #create window
+    window = cmds.window("createNewARTProject_Window", w = 400, h = 600, mnb = False, mxb = False, title = "Create New Project")
+    
+    #frameLayouts for settings: Perforce/Auto-Sync
+    mainLayout = cmds.columnLayout(w = 400, h = 600)
+    
+    #project name field
+    projectName = cmds.textFieldGrp("newARTProjectNameField", label = "Project Name: ", w = 400, h = 30, parent = mainLayout, cal = [1, "left"])
+    
+    
+    scrollLayout = cmds.scrollLayout(w = 400, h = 520, parent = mainLayout)
+    columnLayout = cmds.columnLayout(w = 380, parent = scrollLayout)
+    
+
+    
+    #perforce/auto-sync layout
+    p4Frame = cmds.frameLayout(parent = columnLayout, w = 370, cll = True, label='Perforce/Auto-Sync', borderStyle='in')
+    p4Layout = cmds.columnLayout(w = 360, parent = p4Frame, co = ["both", 5], rs = 5)
+    
+    #create the scrollField with the information
+    cmds.scrollField(parent = p4Layout, w = 350, h = 100, editable = False, wordWrap = True, text = "Add depot paths you would like the tools to check for updates on. If updates are found, you will be notified, and asked if you would like to sync. Valid depot paths look like:\n\n//depot/usr/jeremy_ernst/MayaTools")
+    
+    #crete the add button
+    cmds.button(w = 350, label = "Add Perforce Depot Path", parent = p4Layout, c = partial(addPerforceDepotPath, p4Layout))
+    
+    
+    
+    
+    #save settings button
+    cmds.button(parent = mainLayout, w = 400, h = 50, label = "Save Settings and Close", c = partial(saveProjectSettings, p4Layout, False))
+    
+    #show window
+    cmds.showWindow(window)
+    
+    
+    
+    
+# # # # # # # # # # # # # # # # # # # # # # # # # # # # # # # # # # # # # 
+def editProject(project, *args):
+    
+    if cmds.window("createNewARTProject_Window", exists = True):
+	cmds.deleteUI("createNewARTProject_Window")
+	
+    #create window
+    window = cmds.window("createNewARTProject_Window", w = 400, h = 600, mnb = False, mxb = False, title = "Edit Project")
+    
+    #frameLayouts for settings: Perforce/Auto-Sync
+    mainLayout = cmds.columnLayout(w = 400, h = 600)
+    
+    #project name field
+    projectName = cmds.textFieldGrp("newARTProjectNameField", label = "Project Name: ", text = project, w = 400, h = 30, parent = mainLayout, cal = [1, "left"])
+    
+    
+    scrollLayout = cmds.scrollLayout(w = 400, h = 520, parent = mainLayout)
+    columnLayout = cmds.columnLayout(w = 380, parent = scrollLayout)
+    
+
+    
+    #perforce/auto-sync layout
+    p4Frame = cmds.frameLayout(parent = columnLayout, w = 370, cll = True, label='Perforce/Auto-Sync', borderStyle='in')
+    p4Layout = cmds.columnLayout(w = 360, parent = p4Frame, co = ["both", 5], rs = 5)
+    
+    #create the scrollField with the information
+    cmds.scrollField(parent = p4Layout, w = 350, h = 100, editable = False, wordWrap = True, text = "Add depot paths you would like the tools to check for updates on. If updates are found, you will be notified, and asked if you would like to sync. Valid depot paths look like:\n\n//depot/usr/jeremy_ernst/MayaTools")
+    
+    #crete the add button
+    cmds.button(w = 350, label = "Add Perforce Depot Path", parent = p4Layout, c = partial(addPerforceDepotPath, p4Layout))
+    
+    #get maya tools path
+    toolsPath = cmds.internalVar(usd = True) + "mayaTools.txt"
+    if os.path.exists(toolsPath):
+	
+	f = open(toolsPath, 'r')
+	mayaToolsDir = f.readline()
+	f.close()
+    
+    
+    #open the project settings and auto-fill in the info
+    if os.path.exists(mayaToolsDir + "/General/Scripts/" + project + "_Project_Settings.txt"):
+	f = open(mayaToolsDir + "/General/Scripts/" + project + "_Project_Settings.txt", 'r')
+	settings = cPickle.load(f)
+	f.close() 
+	
+	paths = settings.get("depotPaths")
+	
+    
+	if len(paths) > 0:
+	    for path in paths:
+		#add the path
+		field = addPerforceDepotPath(p4Layout)
+		cmds.textField(field, edit = True, text = path)
+	
+	
+	
+	
+	   
+    
+    #save settings button
+    cmds.button(parent = mainLayout, w = 400, h = 50, label = "Save Settings and Close", c = partial(saveProjectSettings, p4Layout, True))
+    
+    #show window
+    cmds.showWindow(window)
+    
+# # # # # # # # # # # # # # # # # # # # # # # # # # # # # # # # # # # # # 
+def addPerforceDepotPath(layout, *args):
+    
+    field = cmds.textField(docTag = "P4DepotPath", w = 350, parent = layout)
+    
+    #add a RMB menu to remove the field
+    menu = cmds.popupMenu(parent = field, b = 3)
+    cmds.menuItem(parent = menu, label = "Remove Path", c = partial(removePerforceDepotPath, field))
+    
+    return field
+   
+
+
+
+
+
+# # # # # # # # # # # # # # # # # # # # # # # # # # # # # # # # # # # # # 
+def removePerforceDepotPath(field, *args):
+    cmds.textField(field, edit = True, visible = False, h = 1)
+    #cmds.deleteUI(field) This crashes maya instantly. Come ON AUTODESK
+    
+
+
+
+
+# # # # # # # # # # # # # # # # # # # # # # # # # # # # # # # # # # # # # 
+def saveProjectSettings(perforceLayout, edit, *args):
+    
+    #find p4 depot path textfields
+    children = cmds.columnLayout(perforceLayout, q = True, childArray = True)
+    textFields = []
+    for child in children:
+	if child.find("textField") == 0:
+	    data = cmds.textField(child, q = True, docTag = True)
+	    if data == "P4DepotPath":
+		textFields.append(child)
+	
+    
+    #make sure paths are valid
+    savePaths = []
+    for field in textFields:
+	path = cmds.textField(field, q = True, text = True)
+	if path != "":
+	    
+	    try:
+		p4 = P4()
+		p4.connect()
+    
+		    
+	    except:
+		cmds.confirmDialog(title = "Perforce", icon = "critical", message = "Unable to connect to perforce server.")
+		return
+	    
+	    
+	    try:
+		#client info
+		spec = p4.run( "client", "-o" )[0]
+		client = spec.get("Client")
+		owner = spec.get("Owner")	
+		p4.user = owner
+		p4.client = client
+	
+		
+	    except:
+		cmds.confirmDialog(title = "Perforce", icon = "critical", message = "Unable to obtain client spec information.")	
+    
+	    
+	    #now check
+	    try:
+		depotFiles = p4.run_files(path + "...")
+		if len(depotFiles) > 0:
+		    savePaths.append(path)
+	    
+	    #handle any errors
+	    except P4Exception:
+		errorString = "The following errors were encountered:\n\n"
+		
+		for e in p4.errors:  
+		    errorString += e + "\n"
+		
+		cmds.confirmDialog(title = "Perforce", icon = "critical", message = errorString)
+		p4.disconnect()
+		return
+	    
+	else:
+	    #see if the text field is just hidden or if it is actually blank
+	    vis = cmds.textField(field, q = True, visible = True)
+	    if vis == True:
+		cmds.confirmDialog(title = "Error", icon = "critical", message = "Empty string not allowed as a path name. Either remove that field, or enter a correct depot path.")
+		return
+	    
+	    
+    #write out to disk
+    projectName = cmds.textFieldGrp("newARTProjectNameField", q = True, text = True)
+    if projectName == "":
+	cmds.confirmDialog(title = "Error", icon = "critical", message = "Empty string not allowed as a project name.")
+	return
+    
+    #save the new project file under MayaTools/General/Scripts as projName + "_Project_Settings.txt"
+    toolsPath = cmds.internalVar(usd = True) + "mayaTools.txt"
+    if os.path.exists(toolsPath):
+	
+	f = open(toolsPath, 'r')
+	mayaToolsDir = f.readline()
+	f.close()
+	
+    if edit == False:
+	if os.path.exists(mayaToolsDir + "/General/Scripts/" + projectName + "_Project_Settings.txt"):
+	    cmds.confirmDialog(title = "Error", icon = "critical", message = "Project already exists with that name")
+	    return
+    
+    #save out
+    f = open(mayaToolsDir + "/General/Scripts/" + projectName + "_Project_Settings.txt", 'w')
+
+
+    #create a dictionary with  values
+    settings = {}
+    settings["depotPaths"] = savePaths
+    
+    
+    #write our dictionary to file
+    cPickle.dump(settings, f)
+    f.close()
+    
+    
+    #delete the UI
+    cmds.deleteUI("createNewARTProject_Window")
+    
+    
+    #add the project to the menu
+    create = True
+    items = cmds.lsUI(menuItems = True)
+    for i in items:
+	data = cmds.menuItem(i, q = True, docTag = True)
+	if data == "P4Proj":
+	    label = cmds.menuItem(i, q = True, label = True)
+	    print label
+	    
+	    if label == projectName:
+		create = False
+		
+		
+	    
+    if create:
+	menuItem = cmds.menuItem(label = projectName, parent = "perforceProjectList", cl = "perforceProjectRadioMenuCollection", rb = True, docTag = "P4Proj", c = partial(setCurrentProject, projectName))
+	cmds.menuItem(parent = "perforceProjectList", optionBox = True, c = partial(editProject, projectName))
+    
+    
+    
+    #open up the projectSettings.txt file and add an entry for current project
+    if os.path.exists(mayaToolsDir + "/General/Scripts/projectSettings.txt"):
+	f = open(mayaToolsDir + "/General/Scripts/projectSettings.txt", 'r')
+	oldSettings = cPickle.load(f)
+	useSourceControl = oldSettings.get("UseSourceControl")
+	f.close()
+	
+	#write out new settings
+	settings = {}
+	settings["UseSourceControl"] = useSourceControl
+	settings["CurrentProject"] = projectName
+	
+	f = open(mayaToolsDir + "/General/Scripts/projectSettings.txt", 'w')
+	cPickle.dump(settings, f)
+	f.close()	
+	
+	
+# # # # # # # # # # # # # # # # # # # # # # # # # # # # # # # # # # # # # 
+def setCurrentProject(projectName, *args):
+    #get access to maya tools path
+    toolsPath = cmds.internalVar(usd = True) + "mayaTools.txt"
+    if os.path.exists(toolsPath):
+	
+	f = open(toolsPath, 'r')
+	mayaToolsDir = f.readline()
+	f.close()
+	
+	
+    #re-write settings
+    if os.path.exists(mayaToolsDir + "/General/Scripts/projectSettings.txt"):
+	f = open(mayaToolsDir + "/General/Scripts/projectSettings.txt", 'r')
+	oldSettings = cPickle.load(f)
+	useSourceControl = oldSettings.get("UseSourceControl")
+	f.close()
+	
+	#write out new settings
+	settings = {}
+	settings["UseSourceControl"] = useSourceControl
+	settings["CurrentProject"] = projectName
+	
+	f = open(mayaToolsDir + "/General/Scripts/projectSettings.txt", 'w')
+	cPickle.dump(settings, f)
+	f.close()	    
+    
diff --git a/Engine/Extras/Maya_AnimationRiggingTools/ARTv1/General/Scripts/setupScene.py b/Engine/Extras/Maya_AnimationRiggingTools/ARTv1/General/Scripts/setupScene.py
new file mode 100644
--- /dev/null
+++ b/Engine/Extras/Maya_AnimationRiggingTools/ARTv1/General/Scripts/setupScene.py
@@ -0,0 +1,11 @@
+import maya.cmds as cmds
+
+
+def setupScene():
+    #Setup scene to set timeline to 0-100 with current frame being 0 and the fps to be 30
+    cmds.currentUnit(time = 'ntsc')
+    cmds.playbackOptions(min = 0, max = 100, animationStartTime = 0, animationEndTime = 100)
+    cmds.currentTime(0)
+    
+    
+setupScene()
\ No newline at end of file
diff --git a/Engine/Extras/Maya_AnimationRiggingTools/ARTv1/General/Scripts/shotSplitter.py b/Engine/Extras/Maya_AnimationRiggingTools/ARTv1/General/Scripts/shotSplitter.py
new file mode 100644
--- /dev/null
+++ b/Engine/Extras/Maya_AnimationRiggingTools/ARTv1/General/Scripts/shotSplitter.py
@@ -0,0 +1,384 @@
+import maya.cmds as cmds
+from functools import partial
+import re
+import time
+import maya.mel as mel
+
+def UI():
+    
+    #check to see if the window exists, if so delete it
+    if cmds.window("je_shotSplitter_UI", exists = True):
+        cmds.deleteUI("je_shotSplitter_UI")
+        
+    #create the window    
+    window = cmds.window("je_shotSplitter_UI", title = "Shot Splitter", w = 400, h = 200, sizeable = False, mxb= False, mnb = False)
+    
+    #create the mainLayout
+    mainLayout = cmds.columnLayout(w = 400, h = 200)
+    
+    #create a row column layout
+    rowColumnLayout = cmds.rowColumnLayout(nc =2, cw = [(1, 300), (2, 100)], columnOffset = [(1,"both", 5), (2, "both", 5)])
+    
+
+    #Create fields for the text file
+    
+    cmds.separator(h = 10, style = "none")
+    cmds.separator(h = 10, style = "none")
+    
+    cmds.text(label = "Matinee Text File:", align = "left")
+    cmds.text(label = "")
+    
+    cmds.textField("shotSplitterTextFile", w = 300)
+    cmds.button(label = "Browse", w = 50, c = partial(shotSplitter_getTextFieldInfo, "*.txt", 1, "shotSplitterTextFile"))
+    
+    
+
+    
+    #Create fields for the audio file
+    
+    cmds.separator(h = 15, style = "none")
+    cmds.separator(h = 15, style = "none")
+    
+    cmds.text(label = "Matinee Audio File:", align = "left")
+    cmds.text(label = "")
+    
+    cmds.textField("shotSplitterAudioFile", w = 300)
+    cmds.button(label = "Browse", w = 50, c = partial(shotSplitter_getTextFieldInfo, "*.wav", 1, "shotSplitterAudioFile"))
+    
+    
+    #Create fields for the fbx file
+    
+    cmds.separator(h = 15, style = "none")
+    cmds.separator(h = 15, style = "none")
+    
+    cmds.text(label = "Matinee FBX File:", align = "left")
+    cmds.text(label = "")
+    
+    cmds.textField("shotSplitterFBXFile", w = 300, enable = False, text = "Uses currently made file at the moment")
+    cmds.button(label = "Browse", w = 50, enable = False, c = partial(shotSplitter_getTextFieldInfo, "*.fbx", 1, "shotSplitterFBXFile"))
+    
+    
+    #Create fields for the output directory
+    
+    cmds.separator(h = 15, style = "none")
+    cmds.separator(h = 15, style = "none")
+    
+    cmds.text(label = "Output Directory:", align = "left")
+    cmds.text(label = "")
+    
+    cmds.textField("shotSplitterOutput", w = 300)
+    cmds.button(label = "Browse", w = 50, c = partial(shotSplitter_getTextFieldInfo, None, 3, "shotSplitterOutput"))
+    
+    
+    
+    #create the process button
+    cmds.separator(h = 15, style = "none", parent = mainLayout)
+
+    cmds.button(label = "Process", parent = mainLayout, w = 400, h = 50, c = shotSplitter_process)
+    
+    
+    
+    #show the window
+    cmds.showWindow(window)
+    
+
+def shotSplitter_getTextFieldInfo(fileFilter, fileMode, textField, *args):
+    
+    filePath =  cmds.fileDialog2(fileMode= fileMode, dialogStyle=2, fileFilter = fileFilter)[0]
+    filePath = str(filePath)
+    
+    cmds.textField(textField, edit = True, text = filePath)
+    return filePath
+
+
+def shotSplitter_process(*args):
+    start_time = time.time()
+    
+    #get info
+    textFile = cmds.textField("shotSplitterTextFile", q = True, text = True)
+    audioFile = cmds.textField("shotSplitterAudioFile", q = True, text = True)
+    #fbxFile = cmds.textField("shotSplitterFBXFile", q = True, text = True)
+    output = cmds.textField("shotSplitterOutput", q = True, text = True)
+    
+    mel.eval('trace "Getting Data from text file..." ;')
+    #get shot data from the text file
+    shotInfo = shotSplitter_getShotData(textFile)
+    
+    
+    #import in the audio file
+    #if audioFile != "":
+        #mel.eval('trace "Importing audio..." ;')
+        #cmds.file(audioFile, i = True, type = "audio")
+        
+        ##offset audio 300 frames
+        #firstShot = shotInfo[0]
+        #offsetValue = firstShot[1]
+        #audioFileName = audioFile.rpartition("/")[2].partition(".")[0]
+        #cmds.setAttr(audioFileName + ".offset", offsetValue)
+    
+    
+    #import in the fbx file and set it up
+    #mel.eval('trace "Setting up matinee FBX file..." ;')
+    cmds.unloadPlugin(  "fbxmaya.mll" )
+    cmds.loadPlugin(  "fbxmaya.mll" )
+    shotSplitter_setupFBXFile()
+    
+
+    #snap the characters to their spots and orient everything correctly
+    #mel.eval('trace "Snapping Characters..." ;')
+    #shotSplitter_setupCharacters()
+    
+
+    #save out individual shot files
+    mel.eval('trace "Saving out master file..." ;')
+    masterFile = shotSplitter_SaveMaster()
+    
+    i = 0
+    sceneFrameEnd = shotSplitter_findSceneFrameRange()
+    
+    for shot in shotInfo:
+        shotNum = shot[0]
+        mel.eval('trace "Creating shot %s" ;' %str(shotNum))
+        shotFrameStart = shot[1]
+        try:
+            shotFrameEnd = shotInfo[i + 1][1]
+            
+        except IndexError:
+            print "Last shot. Using Scene Frame Range"
+            shotFrameEnd = sceneFrameEnd
+            
+        shotCamera = shot[2]
+        shotName = shot[3]
+        
+        
+        #set frame range, set camera, save out shot
+        cmds.playbackOptions(min = shotFrameStart, max = shotFrameEnd, ast = shotFrameStart, aet = shotFrameEnd)
+        
+        if cmds.objExists(shotCamera):
+            cameraName = cmds.listRelatives(shotCamera, children = True)[0]
+            
+            #debug print
+            print(cameraName)
+            
+            #if shake is in the name
+            #select the non shake camera
+            if 'Shake' in cameraName:
+                print('changing name')
+                cameraName = cameraName.partition('_Shake')[0]
+                
+            #debug print            
+            print('changing the camera name to ' + cameraName)
+            
+                
+            newCamName = cmds.rename(cameraName, "SHOT_CAM")
+            
+            #lock the camera
+            cmds.setAttr(newCamName + ".tx", lock = True)
+            cmds.setAttr(newCamName + ".ty", lock = True)
+            cmds.setAttr(newCamName + ".tz", lock = True)
+            cmds.setAttr(newCamName + ".rx", lock = True)
+            cmds.setAttr(newCamName + ".ry", lock = True)
+            cmds.setAttr(newCamName + ".rz", lock = True)
+            cmds.setAttr(newCamName + ".sx", lock = True)
+            cmds.setAttr(newCamName + ".sy", lock = True)
+            cmds.setAttr(newCamName + ".sz", lock = True)
+            
+        
+        #This will get replaced when the new text file info is in
+        
+        f = open(textFile, "r")
+        lines=f.readlines()
+        f.close()
+        
+        #find the cine name
+        cineFullName = lines[1]
+        cineName = cineFullName.partition(": ")[2]
+        cineName = cineName.rstrip("\r\n")
+    
+        #chop off excess keys
+        cmds.select(all = True)
+
+        
+        try:
+            cmds.cutKey(option="keys", time = (0, shotFrameStart - 30))
+            cmds.cutKey(option="keys", time = (shotFrameEnd + 30, 120000))
+        except:
+            pass
+        
+
+        #save out the shot
+        if shotName.find("<Unknown>") == -1:
+            newFileName = (output + "/" + cineName + "_" + shotName + ".mb")
+            cmds.file(rename = newFileName)
+            shotFile = cmds.file(save=True, type='mayaBinary', force = True)
+        
+            #rename the shot cam back
+            if cmds.objExists(shotCamera):
+                cmds.rename(newCamName, cameraName)
+        
+        #iterate i
+        i = i + 1
+        
+        #open up the master file
+        cmds.file(masterFile, open = True, force = True, prompt = False)
+        
+    
+    
+    elapsed_time = time.time() - start_time
+    timeInMins = str((elapsed_time/60))[0:4]
+    minutes = str(timeInMins).partition(".")[0]
+    seconds = str(timeInMins).rpartition(".")[2]
+    seconds = ("." + seconds)
+    seconds = float(seconds) * 60
+    seconds = str(seconds)[0:2]
+    finalTime = "Total Time: " + minutes + " minutes and " + seconds + " seconds"
+    
+    result = cmds.confirmDialog(title = "Shot Splits Complete", message = "Yay! Your shots have been generated! Total time: " + minutes + " minutes and " + seconds + " seconds", button = ["Close"], cancelButton = "Close", dismissString = "Close")
+
+        
+    
+
+def shotSplitter_findSceneFrameRange():
+    endFrame = cmds.playbackOptions(q = True, max = True)
+    endFrame = int(endFrame)
+    return endFrame
+    
+def shotSplitter_SaveMaster():
+    
+    #first, save the current file in the output directory as the master file.
+    
+    #open the text file, grab all the lines, and close it
+    textFile = cmds.textField("shotSplitterTextFile", q = True, text = True)
+    output = cmds.textField("shotSplitterOutput", q = True, text = True)
+    
+    
+    f = open(textFile, "r")
+    lines=f.readlines()
+    f.close()
+    
+    #find the cine name
+    cineFullName = lines[1]
+    cineName = cineFullName.partition(": ")[2]
+    cineName = cineName.rstrip("\r\n")
+    
+    fileName = (output + "/" + str(cineName) + "_Master.mb")
+    cmds.file( rename= fileName  )
+    cmds.file( save=True, type='mayaBinary', force = True)
+    
+    return fileName
+    
+    
+def shotSplitter_getShotData(textFile):
+    
+    #open the text file, grab all the lines, and close it
+    f = open(textFile, "r")
+    lines=f.readlines()
+    f.close()
+    
+    #find the cine name
+    cineFullName = lines[1]
+    cineName = cineFullName.partition(": ")[2]
+    
+    #compile a list of the shots
+    shotInfo = []
+    for line in lines:
+        if re.search("CameraGroup:", line):
+            shot = line.partition(": ")[2].partition(",")[0]
+            
+            timeStart = line.partition("Time: ")[2].partition(",")[0]
+            timeStart = float(timeStart) * 30
+            timeStart = int(timeStart)
+            
+            camera = line.partition("CameraGroup: ")[2].partition(",")[0]
+            camera = camera.rstrip("\r\n")
+            
+            shotName = line.partition("ShotName: ")[2]
+            shotName = shotName.rstrip("\r\n")
+            
+            shotInfo.append([shot, timeStart, camera, shotName])
+            
+    
+    return shotInfo
+
+
+def shotSplitter_setupFBXFile():
+    
+    #Now let's import the fbx file from matinee
+    #cmds.file(fbxFile, i = True, ignoreVersion = True, force = True)
+    
+    #Now let's clean up the imported fbx assets
+    mel.eval('trace "Cleaning and organizing scene..." ;')
+    cmds.select(all = True)
+    selection = cmds.ls(sl = True, transforms = True)
+    cmds.select(clear = True)
+    
+    
+    #create some groups to put things in
+    matineeGroup = cmds.group(empty = True, name = "Matinee_Group")
+    interpActorGroup = cmds.group(empty = True, name = "InterpActors_Group")
+    worldGeoGroup = cmds.group(empty = True, name = "WorldGeo_Group")
+    fxGroup = cmds.group(empty = True, name = "FX_And_Lighting_Group")
+    camGroup = cmds.group(empty = True, name = "Cam_Group")
+    miscGroup = cmds.group(empty = True, name = "Misc_Group")
+    
+    cmds.parent(interpActorGroup, matineeGroup)
+    cmds.parent(worldGeoGroup, matineeGroup)
+    cmds.parent(fxGroup, matineeGroup)
+    cmds.parent(camGroup, matineeGroup)
+    cmds.parent(miscGroup, matineeGroup)
+    
+    cmds.select(clear = True)
+    
+    for each in selection:
+        
+        if re.search("(.*?)InterpActor", each):
+            cmds.parent(each, interpActorGroup)
+        
+        if re.search("(.*?)Light", each):
+            cmds.parent(each, fxGroup) 
+            
+        if re.search("(.*?)Emitter", each):
+            cmds.parent(each, fxGroup) 
+            
+        if re.search("(.*?)Mesh", each):
+            cmds.parent(each, worldGeoGroup) 
+            
+        if re.search("(.*?)Camera", each):
+            cmds.parent(each, camGroup)
+            
+        
+    #Now select all again, toggle off matinee group, and put anything left in misc group
+    cmds.select(all = True)
+    cmds.select("Matinee_Group", tgl = True)
+    selection = cmds.ls(sl = True, transforms = True)
+    cmds.select(clear = True)
+    
+    for each in selection:
+        if each.find(":") == -1:
+            cmds.parent(each, miscGroup)
+        
+
+def shotSplitter_setupCharacters():
+    
+    cmds.select("*:Master_Anim")
+    selection = cmds.ls(sl = True)
+    
+    namespaces = []
+    
+    for each in selection:
+        namespace = each.partition(":")[0]
+        namespaces.append(namespace)
+        
+    cmds.select(clear = True)
+    
+    for name in namespaces:
+        #find fbx file
+        actorNode = cmds.textField((name + "_mocapTextField"), q = True, text = True).rpartition("/")[2].partition(".")[0]
+        characterNamespace = cmds.frameLayout((name + "_frameLayout"), q = True, label = True)
+        
+        if cmds.objExists(actorNode):
+            cmds.pointConstraint(actorNode, (characterNamespace + ":Master_Anim"))
+            orientConstraint = cmds.orientConstraint(actorNode, (characterNamespace + ":Master_Anim"))[0]
+            
+            cmds.setAttr((orientConstraint + ".offsetZ"), 90)
+            
diff --git a/Engine/Extras/Maya_AnimationRiggingTools/ARTv1/General/Scripts/sidebar.py b/Engine/Extras/Maya_AnimationRiggingTools/ARTv1/General/Scripts/sidebar.py
new file mode 100644
--- /dev/null
+++ b/Engine/Extras/Maya_AnimationRiggingTools/ARTv1/General/Scripts/sidebar.py
@@ -0,0 +1,8 @@
+import maya.cmds as cmds
+
+def UI():
+    
+    print "We are showing the sidebar now!"
+    
+
+UI()
\ No newline at end of file
diff --git a/Engine/Extras/Maya_AnimationRiggingTools/ARTv1/userSetup.py b/Engine/Extras/Maya_AnimationRiggingTools/ARTv1/userSetup.py
new file mode 100644
--- /dev/null
+++ b/Engine/Extras/Maya_AnimationRiggingTools/ARTv1/userSetup.py
@@ -0,0 +1,139 @@
+import maya.cmds as cmds
+import os
+import sys
+from functools import partial
+import maya.mel as mel
+
+MAYA_SEP = '/'
+
+def correctSysPath():
+    for i, path in enumerate(sys.path):
+        if path[-1] in (MAYA_SEP, os.sep, os.altsep):
+            sys.path[i] = sys.path[i][:-1]
+
+correctSysPath()
+
+def deleteOldAutoSaves():
+    import maya.cmds as cmds
+    import time
+    import os.path, time
+    
+    
+    save = cmds.autoSave(q = True, enable = True)
+    destDir = cmds.autoSave(q = True, destinationFolder = True)
+    deleteDays = 3
+    
+    if save:
+        now = time.time()
+        fileCap = now - 60*60*24*deleteDays
+    
+        files = os.listdir(destDir)
+        
+        for file in files:
+            fileCreation = os.path.getctime(os.path.join(destDir, file))
+            if fileCreation < fileCap:
+                os.remove(os.path.join(destDir, file))
+                
+                
+def mayaToolsInstall_UI():
+    
+    if cmds.window("mayaToolsInstall_UI", exists = True):
+        cmds.deleteUI("mayaToolsInstall_UI")
+        
+    window = cmds.window("mayaToolsInstall_UI", w = 300, h = 100, title = "Maya Tools Install", mnb = False, mxb = False)
+    
+    mainLayout = cmds.columnLayout(w = 300, h = 100)
+    formLayout = cmds.formLayout(w = 300, h = 100)
+    
+    text = cmds.text(label = "ERROR: Could not find Maya Tools directory.\n Please locate folder using the \'Browse\' button.", w = 300)
+    cancelButton = cmds.button(label = "Cancel", w = 140, h = 50, c = mayaToolsInstall_Cancel)
+    browseButton = cmds.button(label = "Browse", w = 140, h = 50, c = mayaToolsInstall_Browse)
+    
+    cmds.formLayout(formLayout, edit = True, af = [(text, 'left', 10), (text, 'top', 10)])
+    cmds.formLayout(formLayout, edit = True, af = [(cancelButton, 'left', 5), (cancelButton, 'top', 50)])
+    cmds.formLayout(formLayout, edit = True, af = [(browseButton, 'right', 5), (browseButton, 'top', 50)])
+    
+    cmds.showWindow(window)
+    cmds.window(window, edit = True, w = 300, h = 100)
+    
+    
+    
+def mayaToolsInstall_Cancel(*args):
+    
+    cmds.deleteUI("mayaToolsInstall_UI")
+    cmds.warning("Maya Tools will not be setup")
+    
+
+
+def mayaToolsInstall_Browse(*args):
+    
+    mayaToolsDir = cmds.fileDialog2(dialogStyle = 2, fileMode = 3)[0]
+
+    #confirm that this is actually the maya tools directory
+    if mayaToolsDir.rpartition("/")[2] != "MayaTools":
+        cmds.warning("Selected directory is not valid. Please locate the MayaTools directory.")
+        
+        
+    else:
+        cmds.deleteUI("mayaToolsInstall_UI")
+        
+        #create file that contains this path
+        path = cmds.internalVar(usd = True) + "mayaTools.txt"
+        
+        f = open(path, 'w')
+        f.write(mayaToolsDir)
+        f.close()
+        
+        
+        #run setup
+        cmds.file(new = True, force = True)
+        
+        
+        
+def mayaTools():
+    
+    #import custom maya setup script
+    print cmds.internalVar(usd = True)
+    print cmds.internalVar(upd = True)
+    path = cmds.internalVar(usd = True)
+    path = path + "mayaTools.txt"
+    print path
+    
+    if os.path.exists(path):
+        f = open(path, 'r')
+        
+        mayaToolsDir = f.readline()
+        if not os.path.exists(mayaToolsDir):
+            mayaToolsInstall_UI()
+        
+        path = mayaToolsDir + "/General/Scripts"
+        pluginPath = mayaToolsDir + "/General/Plugins"
+        
+        #look in sys.path to see if path is in sys.path. if not, add it
+        if not path in sys.path:
+            sys.path.append(path)
+            
+            
+        #run setup
+        import customMayaMenu as cmm
+        cmm.customMayaMenu()
+        cmds.file(new = True, force = True)
+    
+    
+    else:
+        mayaToolsInstall_UI()
+        
+        
+
+def launchCustomFileLister():
+    import customFileLister
+    reload(customFileLister)
+    customFileLister.UI()
+    
+    
+
+#setup script jobs
+scriptJobNum = cmds.scriptJob(event = ["NewSceneOpened", launchCustomFileLister ] )
+scriptJobNum2 = cmds.scriptJob(event = ["NewSceneOpened", mayaTools])
+
+
