cleaning up

diff --git a/Engine/Extras/Maya_AnimationRiggingTools/ARTv1/MayaTools/General/Scripts/Modules/facial/face.py b/Engine/Extras/Maya_AnimationRiggingTools/ARTv1/MayaTools/General/Scripts/Modules/facial/face.py
new file mode 100644
--- /dev/null
+++ b/Engine/Extras/Maya_AnimationRiggingTools/ARTv1/MayaTools/General/Scripts/Modules/facial/face.py
@@ -0,0 +1,695 @@
+'''
+FACIAL MODULE
+
+'''
+
+import os, sys
+import traceback
+import json
+
+from PySide import QtGui, QtCore
+reload(QtGui)
+reload(QtCore)
+
+import maya.cmds as cmds
+
+# internal libs
+from Modules.facial import utils
+
+##################
+# BASE
+##################
+
+
+class BaseFace(object):
+    '''
+    This class is the base to be inherited and does not serialize to anything in Maya
+    '''
+    def __init__(self, version=None, create=True, uType='uMeta'):
+        pass
+
+    # FUNCTIONS
+    ##############
+
+    # PROPERTIES
+    ###############
+
+    #get and set controllers
+    @property
+    def controls(self):
+        return cmds.listConnections(self.node + '.controls')
+    @controls.setter
+    def controls(self, controls):
+        for c in controls:
+            utils.msgConnect(self.node + '.controls', c + '.facialNode')
+
+    #get and set version
+    @property
+    def version(self):
+        return cmds.getAttr(self.node+'.version')
+    @version.setter
+    def version(self, ver):
+        cmds.setAttr(self.node+'.version', ver, type='string')
+
+    #get and set uType identifier
+    @property
+    def uType(self):
+        return cmds.getAttr(self.node+'.uType')
+    @uType.setter
+    def uType(self, type):
+        cmds.setAttr(self.node+'.uType', type, type='string')
+
+    #Build attributes on the node if it's being built
+    def buildBase(self):
+        cmds.addAttr(self.node, longName='version', dataType='string')
+        cmds.addAttr(self.node, longName='controls', attributeType='message')
+        cmds.addAttr(self.node, longName='uType', dataType='string')
+
+
+
+##################
+## MASK
+##################
+class FaceMask(BaseFace):
+    '''
+    This class serializes to a network node
+    '''
+    def __init__(self, maskGeoNode=None, faceModuleNode=None, maskNode=None, **kwargs):
+        super(FaceMask, self).__init__(**kwargs)
+        self.maskGeoNode = maskGeoNode
+        self.faceModuleNode = faceModuleNode
+
+        if maskNode:
+            self.node = maskNode
+            self.initialize()
+
+        elif 'create' in kwargs:
+            self.buildNode()
+
+    ## FUNCTIONS
+    ##############
+
+    def bakeMovers(self):
+        pass
+
+    def zeroMovers(self):
+        for mover in self.faceModule.jointMovers:
+            cmds.xform()
+
+    def snapMovers(self, axis, mesh, debug=1):
+        '''
+        for this system, in the tool UI itself that snaps the movers, I always bidirectionally cast against Y
+        '''
+        try:
+            cmds.undoInfo(openChunk=True)
+            #iterate through the *active* facial joint movers and snap them to the surface
+            if debug: print 'activeJointMovers:', self.faceModule.activeJointMovers
+            for mover in self.faceModule.activeJointMovers:
+                if debug: print 'snapping', mover
+                #do not snap manually placed movers
+                if not cmds.getAttr(mover + '.alignmentType'):
+
+                    #set the axis vector
+                    axisVector = None
+                    if axis == 'x': axisVector = utils.getAxisWorldVectors(mover)[0]
+                    elif axis == 'y': axisVector = utils.getAxisWorldVectors(mover)[1]
+                    elif axis == 'z': axisVector = utils.getAxisWorldVectors(mover)[2]
+
+                    pos = cmds.xform(mover, q=1, rp=1, ws=1)
+
+                    #raycast in and out and find the closest hit of the facial surface
+                    pos1 = utils.rayMeshIntersect(mesh, pos, axisVector)
+                    pos2 = utils.rayMeshIntersect(mesh, pos, -axisVector)
+                    dist1 = utils.square_distance(pos, pos1)
+                    dist2 = utils.square_distance(pos, pos2)
+                    newPos = None
+                    if dist1 <= dist2:
+                        newPos = pos1
+                    else:
+                        newPos = pos2
+
+                    cmds.xform(mover, t=newPos, ws=1)
+                    if debug: print mover, 'snapped.'
+        except Exception as e:
+            print e
+        finally:
+            cmds.undoInfo(closeChunk=True)
+
+    def deleteSdkParentConstraints(self):
+        try:
+            cmds.undoInfo(openChunk=True)
+            #update facial SDK locations to current joint mover placement
+            print '>>>deleteSdkParentConstraints'
+            for mover in self.faceModule.jointMovers:
+                sdk = cmds.listConnections(mover + '.sdk')[0]
+                sdkParent = cmds.listRelatives(sdk, parent = True)[0]
+                utils.deleteConstraints(sdkParent)
+        except Exception as e:
+            print(traceback.format_exc())
+        finally:
+            cmds.undoInfo(closeChunk=True)
+
+    def buildSdkParentConstraints(self):
+        try:
+            cmds.undoInfo(openChunk=True)
+            for mover in self.faceModule.jointMovers:
+                sdk = cmds.listConnections(mover + '.sdk')[0]
+                lra = cmds.listConnections(mover + '.lra')[0]
+                sdkParent = cmds.listRelatives(sdk, parent = True)[0]
+                #TODO: Maybe don't maintain offset later (?)
+                # there are some relationships that have offsets (eyes) that should be maintained
+                cmds.parentConstraint(lra, sdkParent, mo=1)
+        except Exception as e:
+            print(traceback.format_exc())
+        finally:
+            cmds.undoInfo(closeChunk=True)
+
+    def toggleMaskActive(self, mode):
+        try:
+            cmds.undoInfo(openChunk=True)
+            if mode:
+                cmds.showHidden(self.controls)
+                cmds.showHidden(self.maskGeo)
+
+                #hide all movers
+                cmds.hide(self.faceModule.jointMovers)
+                #TODO: Need to check if the constraints were deleted and re-add them
+                '''for mover in self.faceModule.jointMovers:
+                    sdk = cmds.listConnections(mover + '.sdk')[0]
+                    sdkParent = cmds.listRelatives(sdk, parent = True)[0]
+                    lra = cmds.listConnections(mover + '.lra')[0]
+                    cmds.parentConstraint(lra, sdkParent)'''
+                print 'MASK: ACTIVATED'
+
+            else:
+                #show joint movers
+                cmds.showHidden(self.faceModule.jointMovers)
+
+                #hide the mask controls and mesh
+                cmds.hide(self.controls)
+                cmds.hide(self.maskGeo)
+                print 'MASK: DEACTIVATED'
+        except Exception as e:
+            print(traceback.format_exc())
+        finally:
+            cmds.undoInfo(closeChunk=True)
+
+
+
+    ## BUILD
+    def buildNode(self):
+        self.node = cmds.createNode('network', name='maskNode')
+        self.buildBase()
+
+        #maskGeo attr and connection if it was passed in
+        cmds.addAttr(self.node, longName='maskGeo', attributeType='message')
+        if self.maskGeoNode:
+            self.maskGeo = self.maskGeoNode
+
+        #FaceModule attr and connection if it was passed in
+        cmds.addAttr(self.node, longName='faceModule', attributeType='message')
+        if self.faceModuleNode:
+            self.faceModule = self.faceModuleNode
+
+        #Set other metadata
+        self.version = '0.00'
+        self.uType = 'faceMask'
+
+    ##SERIALIZATION
+    def initialize(self):
+        pass
+
+
+    ## PROPERTIES
+    ###############
+    @property
+    def maskGeo(self):
+        return cmds.listConnections(self.node + '.maskGeo')[0]
+    @maskGeo.setter
+    def maskGeo(self, maskG):
+        utils.msgConnect(self.node + '.maskGeo', maskG + '.maskNode')
+
+    @property
+    def faceModule(self):
+        return FaceModule(faceNode=cmds.listConnections(self.node + '.faceModule')[0])
+    @faceModule.setter
+    def faceModule(self, faceM):
+        utils.msgConnect(faceM + '.maskNode', self.node + '.faceModule')
+
+    @property
+    def attachLocations(self):
+        return cmds.listConnections(self.node + '.attachLocations')
+    @attachLocations.setter
+    def attachLocations(self, locs):
+        for loc in locs:
+            utils.msgConnect(self.node + '.attachLocations', loc + '.maskNode')
+
+    @property
+    def active(self):
+        return cmds.listConnections(self.node + '.active')
+    @active.setter
+    def active(self, mode):
+        if not utils.attrExists(self.node + '.active'):
+            cmds.addAttr(self.node, longName='active', at='bool')
+        cmds.setAttr(self.node+'.active', mode)
+
+        self.toggleMaskActive(mode)
+
+
+####################
+### FACE
+####################
+class FaceModule(BaseFace):
+    '''
+    This class serializes to a network node
+    '''
+    def __init__(self, renderMeshNodes=None, faceNode=None, **kwargs):
+        super(FaceModule, self).__init__(**kwargs)
+
+        if faceNode:
+            self.node = faceNode
+            self.initialize()
+
+        elif 'create' in kwargs:
+            self.buildNode()
+        else:
+            cmds.error('No facial node set through faceNode, and no create flag used!\nLOCALS: ' + str(locals()))
+
+
+    ## FUNCTIONS
+    ##############
+
+    def initialize(self):
+        self.mask = self.faceMask
+
+    def readDrivenDict(self, node):
+        if utils.attrExists(node + '.drivenDict'):
+            dict = cmds.getAttr(node + '.drivenDict')
+            if dict:
+                return json.loads(dict)
+        else:
+            cmds.warning('No drivenDict found on object: ' + node)
+            return {}
+
+    def updateConstraints(self, node):
+        '''
+        This is basically made for the method 'snapEyeSdkPivots' below.
+        It is a general function, except the part where is pops the sdkNode connection.
+        '''
+        parent = cmds.listRelatives(node, parent=1)[0]
+        listConnectionsSucks = cmds.listConnections(parent, type='parentConstraint')
+        if listConnectionsSucks:
+            for const in set(listConnectionsSucks):
+                conns = list(set(cmds.listConnections(const, type='transform')))
+                conns.pop(conns.index(const))
+                conns.pop(conns.index(parent))
+                if utils.attrExists(const + '.sdkNode'):
+                    conns.pop(conns.index(cmds.listConnections(const + '.sdkNode')[0]))
+                if len(conns) == 1:
+                    cmds.parentConstraint(conns[0], parent, e=1, mo=1)
+
+    def snapEyeSdkPivots(self):
+        L_eyeSDKs = [u'l_eye_lid_upper_mid_SDK', u'l_eye_lid_upper_inner_SDK', u'l_eye_lid_upper_outer_SDK', u'l_eye_lid_lower_outer_SDK', u'l_eye_lid_lower_mid_SDK', u'l_eye_lid_lower_inner_SDK']
+        R_eyeSDKs = [u'r_eye_lid_upper_outer_SDK', u'r_eye_lid_upper_mid_SDK', u'r_eye_lid_upper_inner_SDK', u'r_eye_lid_lower_inner_SDK', u'r_eye_lid_lower_mid_SDK', u'r_eye_lid_lower_outer_SDK']
+        for node in L_eyeSDKs:
+            utils.snapPivotWithZeroParent(node, 'l_eye_SDK')
+            self.updateConstraints(node)
+        for node in R_eyeSDKs:
+            utils.snapPivotWithZeroParent(node, 'r_eye_SDK')
+            self.updateConstraints(node)
+
+    ## BUILD
+    def buildNode(self):
+        #create and pass the 'node' to the base
+        self.node = cmds.createNode('network', name='faceNode')
+        #call the base build on that node
+        self.buildBase()
+
+        #Set other useful metadata
+        self.version = '0.00'
+        self.uType = 'faceModule'
+
+        cmds.addAttr(self.node, longName='jointMovers', attributeType='message')
+        cmds.addAttr(self.node, longName='sdks', attributeType='message')
+
+    def stampDeltas(self):
+        try:
+            cmds.undoInfo(openChunk=True)
+            bb = self.blendBoard
+            if bb:
+                #build attr
+                if not utils.attrExists(bb + '.initialPoseDict'):
+                    cmds.addAttr(bb, longName='initialPoseDict', dt='string')
+                    cmds.setAttr(bb + '.initialPoseDict', '{}', type='string')
+
+                initialPoseDict={}
+
+                drivenDict = self.readDrivenDict(bb)
+                for drivingAttr in drivenDict:
+                    attrDict = {}
+                    for driven in drivenDict[drivingAttr]:
+                        cmds.setAttr(drivingAttr, 1.0)
+                        attrDict[driven] = cmds.getAttr(driven)
+                    initialPoseDict[drivingAttr] = attrDict
+                    cmds.setAttr(drivingAttr, 0.0)
+
+                cmds.setAttr(bb + '.initialPoseDict', json.dumps(initialPoseDict), type='string')
+
+            else:
+                cmds.warning('No blendBoard/poses found!')
+        except Exception as e:
+            print(traceback.format_exc())
+            cmds.undoInfo(closeChunk=True)
+
+
+    def buildBlendBoard(self, fpath):
+        try:
+            cmds.undoInfo(openChunk=True)
+
+            #TODO: rebuild and keep all connections and attrs
+            if self.blendBoard:
+                cmds.delete(self.blendBoard)
+
+            lines = []
+            if os.path.exists(fpath):
+                f = open(fpath, 'r')
+                lines = f.readlines()
+                f.close()
+            self.blendBoard = cmds.spaceLocator(name='blendBoard')[0]
+            cmds.addAttr(self.blendBoard, longName='uType', dataType='string')
+            cmds.setAttr(self.blendBoard + '.uType', 'blendBoard', type='string')
+            utils.msgConnect(self.node + '.blendBoard', self.blendBoard + '.faceModule')
+
+            #store input file path
+            cmds.addAttr(self.blendBoard, longName='inputFile', dataType='string')
+            cmds.setAttr(self.blendBoard + '.inputFile', fpath, type='string')
+
+            #make purdy
+            cmds.setAttr("blendBoard.tx", l=1, k=0, channelBox=0)
+            cmds.setAttr("blendBoard.ty", l=1, k=0, channelBox=0)
+            cmds.setAttr("blendBoard.tz", l=1, k=0, channelBox=0)
+            cmds.setAttr("blendBoard.rx", l=1, k=0, channelBox=0)
+            cmds.setAttr("blendBoard.ry", l=1, k=0, channelBox=0)
+            cmds.setAttr("blendBoard.rz", l=1, k=0, channelBox=0)
+            cmds.setAttr("blendBoard.sx", l=1, k=0, channelBox=0)
+            cmds.setAttr("blendBoard.sy", l=1, k=0, channelBox=0)
+            cmds.setAttr("blendBoard.sz", l=1, k=0, channelBox=0)
+            cmds.setAttr("blendBoard.v", l=1, k=0, channelBox=0)
+
+            #make pose attrs
+            for line in lines:
+                line = line.strip()
+                if line != '':
+                    print line
+                    if '##' in line:
+                        cmds.addAttr(self.blendBoard,enumName='#####', longName=line.replace('##',''), at='enum', k=1)
+                        cmds.setAttr(self.blendBoard + '.' + line.replace('##',''), lock=1)
+                    else:
+                        cmds.addAttr(self.blendBoard, longName=line, at='float', min=0, max=1, k=1)
+
+            cmds.undoInfo(closeChunk=True)
+        except Exception as e:
+            exc_type, exc_value, exc_tb = sys.exc_info()
+            for line in traceback.format_exception(exc_type, exc_value, exc_tb):
+                print line
+            cmds.undoInfo(closeChunk=True)
+
+    def createJointMover(self, name=None, pos=(0,0,0), matrix=None, scale=1, toroidRatio=0.02, ss=0, facialNode=None):
+        try:
+            cmds.undoInfo(openChunk=True)
+            if not name:
+                name = 'test'
+            radius = 1
+
+            #create mask attachment node
+            maskAttach = cmds.spaceLocator(name=name+'_mover_maskAttach')[0]
+
+            #create jeremy-joint-mover-stack
+            topNode = cmds.group(name=name+'_mover_grp', em=1)
+            mover = utils.createCircleControl(name=name+'_mover', radius=2, type='toroid', heightRatio=toroidRatio)
+            #offset = utils.createCircleControl(name=name+'_mover_offset', radius=1.5, color='xyz', type='toroid', heightRatio=toroidRatio)
+            offset = cmds.group(name=name+'_mover_offset', em=1)
+            lra_grp = cmds.group(name=name+'_lra_grp', em=1)
+            lra = utils.createLRA(name=name + '_LRA', length=2, thickness=0.05, arrowRadiusMult=3)
+            cmds.parent(mover, topNode)
+            cmds.parent(offset, mover)
+            cmds.parent(lra_grp, offset)
+            cmds.parent(lra, lra_grp)
+            cmds.displaySmoothness(mover, f=1)
+            if scale is not 1:
+                cmds.scale(scale, scale, scale, topNode)
+                cmds.makeIdentity(topNode, apply=1, t=0, r=0, s=1, n=0, pn=1)
+            if not ss:
+                cmds.select(topNode)
+            cmds.parentConstraint(maskAttach, topNode)
+
+            #create SDK pose stack
+            sdkOffset = cmds.group(name=name+'_SDK_grp', em=1)
+            sdk = utils.createLRA(name=name + '_SDK', length=2, color=False, thickness=0.05, arrowRadiusMult=3)
+            tweak = cmds.group(name=name+'_TWEAK', em=1)
+            if scale is not 1:
+                cmds.scale(scale, scale, scale, sdk)
+                cmds.makeIdentity(sdk, apply=1, t=0, r=0, s=1, n=0, pn=1)
+            cmds.parent(tweak, sdk)
+            cmds.parent(sdk, sdkOffset)
+            cmds.parentConstraint(lra, sdkOffset)
+            utils.toggleShapeSelection(sdk, False)
+
+            #cleanup outliner
+            if cmds.objExists('maskAttachNodes'):
+                cmds.parent(maskAttach, 'maskAttachNodes')
+            if cmds.objExists('sdkNodes'):
+                cmds.parent(sdkOffset, 'sdkNodes')
+            if cmds.objExists('jointMovers'):
+                cmds.parent(topNode, 'jointMovers')
+
+            #make internal attrs/connections
+            utils.msgConnect(mover + '.lra', lra + '.jointMover')
+            utils.msgConnect(mover + '.attachNode', maskAttach + '.jointMover')
+            utils.msgConnect(mover + '.sdk', sdk + '.jointMover')
+
+            #make external attrs/connections
+            if facialNode:
+                locs = None
+
+                #mark attachLocations
+                try:
+                    locs = self.mask.attachLocations
+                    locs.extend(maskAttach)
+                    self.mask.attachLocations = locs
+                except Exception as e:
+                    print(traceback.format_exc())
+
+                #hang the jm off the facial node
+                utils.msgConnect(self.node + '.jointMovers', mover + '.facialNode')
+
+                #mark jointMovers
+                jm = self.jointMovers
+                jm.append(mover)
+                self.jointMovers = jm
+
+                #mark sdks
+                s = self.sdks
+                s.append(sdk)
+                self.sdks = s
+
+            #stamp with alignment type
+            cmds.addAttr(sdk, ln='alignmentType', enumName='mask surface:manual',  k=1, at='enum')
+
+            cmds.undoInfo(closeChunk=True)
+            cmds.select(maskAttach)
+            return [topNode, lra, offset, maskAttach]
+        except Exception as e:
+            cmds.undoInfo(closeChunk=True)
+            import sys
+            exc_type, exc_obj, exc_tb = sys.exc_info()
+            fname = os.path.split(exc_tb.tb_frame.f_code.co_filename)[1]
+            print(e, fname, exc_tb.tb_lineno)
+
+    def findMirrorJointMoverByName(self, jm):
+        mirrorNode = None
+        if not jm.lower().startswith('c_'):
+            if jm.lower().startswith('r_'):
+                mirrorNode = 'l' + jm[1:]
+            else:
+                mirrorNode = 'r' + jm[1:]
+            if mirrorNode:
+                if cmds.objExists(mirrorNode):
+                    return mirrorNode
+            else:
+                cmds.warning('Cannot find l_ or r_ prefix in: ' + jm)
+                return False
+
+    def mirrorFacialXform(self, jm, axis='x', debug=0):
+        mirror = self.findMirrorJointMoverByName(jm)
+        if debug:
+            print mirror
+        if mirror:
+            pos = cmds.xform(jm, q=1, ws=1, rp=1)
+            rot = cmds.xform(jm, q=1, ws=1, ro=1)
+            #scale = cmds.xform(jm, q=1, r=1, s=1)
+
+            cmds.select(mirror)
+            if axis == 'x':
+                cmds.move(-pos[0], pos[1], pos[2], mirror,  ws=1, rpr=1 )
+                cmds.rotate(rot[0], -rot[1], -rot[2], mirror,  ws=1, a=1 )
+            else:
+                cmds.warning('Currently works for faces created looking down y axis.')
+
+    def connectControlBoard(self, controlBoard):
+        bb = self.blendBoard
+        utils.msgConnect(self.node + '.controlBoard', controlBoard + '.faceModule')
+        if bb:
+            for attr in cmds.listAttr(bb, k=1):
+                bbAttr = bb + '.' + attr
+                if cmds.listConnections(bbAttr):
+                    cbAttr = controlBoard + '.' + attr
+                    if utils.attrExists(cbAttr):
+                        cmds.connectAttr(cbAttr, bbAttr, f=1)
+
+    def connectSdksToFinalSkeleton(self):
+        jmLocked = cmds.lockNode('JointMover', lock=1, q=1)[0]
+        if jmLocked:
+            cmds.lockNode('JointMover', lock=0)
+        for mover in self.activeJointMovers:
+            #TODO: this won't work with multiple heads
+            #hate to use name string here, but there's only one skeleton
+            sdk = cmds.listConnections(mover + '.sdk')[0]
+            joint = mover.replace('_mover','')
+            tweak = cmds.listRelatives(sdk, children=1, typ='transform')[0]
+            cmds.parentConstraint(tweak, joint, mo=1)
+            cmds.scaleConstraint(tweak, joint)
+
+            #to get away from sting names later
+            #TODO: detect locks and only unlock/relock if the item was locked
+            moverLocked = cmds.lockNode(mover, lock=1, q=1)[0]
+            if moverLocked:
+                cmds.lockNode(mover, lock=0)
+            utils.msgConnect(mover + '.deformingJoint', joint + '.mover')
+            if moverLocked:
+                cmds.lockNode(mover, lock=1)
+        if jmLocked:
+            cmds.lockNode('JointMover', lock=1)
+
+
+    ## PROPERTIES
+    ###############
+    @property
+    def renderMeshes(self):
+        return cmds.listConnections(self.node + '.renderMeshes')
+    @renderMeshes.setter
+    def renderMeshes(self, meshes):
+        for mesh in meshes:
+            utils.msgConnect(self.node + '.renderMeshes', mesh + '.faceNode')
+
+    @property
+    def faceMask(self):
+        return FaceMask(maskNode=cmds.listConnections(self.node + '.maskNode')[0])
+    @faceMask.setter
+    def faceMask(self, maskNode):
+        utils.msgConnect(self.node + '.faceMask', maskNode + '.faceNode')
+
+    @property
+    def blendBoard(self):
+        bb = cmds.listConnections(self.node + '.blendBoard')
+        if bb:
+            return cmds.listConnections(self.node + '.blendBoard')[0]
+        else:
+            return []
+    @blendBoard.setter
+    def blendBoard(self, bBoard):
+        utils.msgConnect(self.node + '.blendBoard', bBoard + '.faceNode')
+
+    @property
+    def jointMovers(self):
+        conns = cmds.listConnections(self.node + '.jointMovers')
+        if conns:
+            return conns
+        else: return []
+    @jointMovers.setter
+    def jointMovers(self, movers):
+        for m in movers:
+            utils.msgConnect(self.node + '.jointMovers', m + '.facialNode')
+
+    @property
+    def sdks(self):
+        conns = cmds.listConnections(self.node + '.sdks')
+        if conns:
+            return conns
+        else: return []
+    @sdks.setter
+    def sdks(self, sdks):
+        for sdk in sdks:
+            utils.msgConnect(self.node + '.sdks', sdk + '.facialNode')
+
+    @property
+    def blendBoardConnections(self):
+        if self.blendBoard:
+            poses = []
+            bb = self.blendBoard
+            for att in cmds.listAttr(bb, k=1):
+                conns = cmds.listConnections(bb + '.' + att)
+                if conns:
+                    poses.append(att)
+            return poses
+        else:
+            cmds.warning('No blendBoard connected!')
+            return None
+    @blendBoardConnections.setter
+    def blendBoardConnections(self, *args):
+        print 'Attribute cannot be set: connections are derived from the blendBoard'
+
+    @property
+    def poses(self):
+        bb = self.blendBoard
+        return self.readDrivenDict(bb)
+    @poses.setter
+    def poses(self, poses):
+        cmds.warning('Attribute cannot be set: poses are derived from the blendBoard drivenDict attr')
+
+    @property
+    def activeJointMovers(self, debug=0):
+        conns = cmds.listConnections(self.node + '.jointMovers')
+        activeCons = []
+        if conns:
+            for con in conns:
+                if utils.shapesHidden(con) == False:
+                    if debug: print con, 'is not hidden and active'
+                    activeCons.append(con)
+                else:
+                    if debug: print con, 'shapes hidden! not added.'
+            return activeCons
+        else:
+            return []
+    @activeJointMovers.setter
+    def activeJointMovers(self, *args):
+        print 'Attribute cannot be set: connections are derived from the facial joint mover visibility'
+        return False
+
+
+####################
+### TESTING
+####################
+def autoTest(build=True):
+    #build
+    if build:
+        face = FaceModule(create=1)
+        mask = FaceMask(create=1, maskGeoNode='maskGeo', faceModuleNode=face.node)
+        mask.controls = [u'jaw', u'lips_CTRL_MASK', u'R_nose', u'nose_CTRL_MASK', u'L_nose', u'l_lip_corner_CTRL_MASK', u'r_lip_corner_CTRL_MASK', u'R_mask_CTRL_MASK', u'brow_ridge_CTRL_MASK', u'R_eye_LRA', u'L_eye_LRA', u'L_mask_CTRL_MASK', u'top_mask_CTRL_MASK']
+        face.renderMeshes = ['steele']
+        '''
+        print face.faceMask
+        print face.version
+        print mask.maskGeo
+        print mask.faceModule
+        '''
+    else:
+        face = FaceModule(faceNode='faceNode')
+        mask = FaceMask(faceModuleNode='faceNode')
+
+
+#autoTest()
+
+
+if __name__ == '__main__':
+    show()
diff --git a/Engine/Extras/Maya_AnimationRiggingTools/ARTv1/MayaTools/General/Scripts/Modules/facial/facialMaskTool.py b/Engine/Extras/Maya_AnimationRiggingTools/ARTv1/MayaTools/General/Scripts/Modules/facial/facialMaskTool.py
new file mode 100644
--- /dev/null
+++ b/Engine/Extras/Maya_AnimationRiggingTools/ARTv1/MayaTools/General/Scripts/Modules/facial/facialMaskTool.py
@@ -0,0 +1,249 @@
+'''
+FACIAL TOOL
+'''
+
+import os
+
+#uic related
+import shiboken
+from PySide import QtGui, QtCore
+reload(QtGui)
+reload(QtCore)
+from cStringIO import StringIO
+import xml.etree.ElementTree as xml
+import pysideuic
+
+import maya.cmds as cmds
+import maya.OpenMayaUI as openMayaUI
+import maya.mel as mel
+
+#internal libs
+from Modules.facial import utils
+reload(utils)
+from Modules.facial import jointMover_UI as jm_ui
+reload(jm_ui)
+from Modules.facial import poseEditor_UI as pose_ui
+reload(pose_ui)
+from Modules.facial import face
+reload(face)
+from Modules.facial import mask_align_UI as ma_ui
+reload(jm_ui)
+
+########
+## UIC
+########
+def show():
+    global facialMaskToolWindow
+    try:
+        facialMaskToolWindow.close()
+    except:
+        pass
+    facialMaskToolWindow = facialMaskTool()
+    facialMaskToolWindow.show()
+    return facialMaskToolWindow
+
+def loadUiType(uiFile):
+    """
+    Pyside lacks the "loadUiType" command, so we have to convert the ui file to py code in-memory first
+    and then execute it in a special frame to retrieve the form_class.
+    http://tech-artists.org/forum/showthread.php?3035-PySide-in-Maya-2013
+    """
+    parsed = xml.parse(uiFile)
+    widget_class = parsed.find('widget').get('class')
+    form_class = parsed.find('class').text
+
+    with open(uiFile, 'r') as f:
+        o = StringIO()
+        frame = {}
+
+        pysideuic.compileUi(f, o, indent=0)
+        pyc = compile(o.getvalue(), '<string>', 'exec')
+        exec pyc in frame
+
+        #Fetch the base_class and form class based on their type in the xml from designer
+        form_class = frame['Ui_%s'%form_class]
+        base_class = eval('QtGui.%s'%widget_class)
+    return form_class, base_class
+
+def getMayaWindow():
+    ptr = openMayaUI.MQtUtil.mainWindow()
+    if ptr is not None:
+        return shiboken.wrapInstance(long(ptr), QtGui.QWidget)
+
+uiFile = None	
+try:
+    selfDirectory = os.path.dirname(__file__)
+    uiFile = selfDirectory + '/facialMaskTool.ui'
+except:
+    uiFile = 'D:\\Build\\usr\\jeremy_ernst\\MayaTools\\General\\Scripts\\Modules\\facial\\facialMaskTool.ui'
+if os.path.isfile(uiFile):
+    form_class, base_class = loadUiType(uiFile)
+else:
+    cmds.error('Cannot find UI file: ' + uiFile)
+
+
+
+##############
+## MASK TOOL
+##############
+class facialMaskTool(base_class, form_class):	  
+    def __init__(self, parent=getMayaWindow()):
+        self.closeExistingWindow()
+        super(facialMaskTool, self).__init__(parent)
+        
+        self.setupUi(self)
+        
+        wName = openMayaUI.MQtUtil.fullName(long(shiboken.getCppPointer(self)[0]))
+        
+        self.currentFace = None
+        
+        ## Connect UI
+        ########################################################################
+        self.createJointMoverBTN.clicked.connect(self.createJointMoverFn)
+        self.showHideAttachBTN.clicked.connect(self.showHideAttachFn)
+        
+        self.openJointMoverDlgBTN.clicked.connect(self.openJointMoverDlgFn)
+        self.openPoseDlgBTN.clicked.connect(self.openPoseDlgFn)
+        self.openMaskAlignmentDlgBTN.clicked.connect(self.openMaskAlignmentDlgFn)
+        
+        self.selectFaceBTN.clicked.connect(self.selectFaceFn)
+        self.selectMaskBTN.clicked.connect(self.selectMaskFn)
+        self.maskModeBTN.clicked.connect(self.maskModeFn)
+        self.constrainToMaskBTN.pressed.connect(self.constrainToMaskFn)
+        self.buildBlendBoardBTN.pressed.connect(self.buildBlendBoardFn)
+        self.stampDeltasBTN.pressed.connect(self.stampDeltasFn)
+        self.mirrorPlacementBTN.pressed.connect(self.mirrorPlacementFn)
+        self.deleteSdkParentConstraintsBTN.pressed.connect(self.deleteSdkParentConstraintsFn)
+
+        self.refreshUI()
+    
+    ##UI
+    ######
+    
+    def stampDeltasFn(self):
+        self.currentFace.stampDeltas()
+    
+    def createJointMoverFn(self):
+        inputName, ok = QtGui.QInputDialog.getText(self, 'Create Facial Joint Mover', 'Base Name:')
+        if ok:
+            self.currentFace.createJointMover(name=inputName, scale=0.3,toroidRatio=0.05, facialNode=self.currentFace.node)
+        self.refreshUI()
+    
+    def showHideAttachFn(self):
+        if self.showHideAttachBTN.isChecked():
+            print 'hiding'
+            self.showHideAttachBTN.setText('Show Mask Attach Nodes')
+        else:
+            print 'showing'
+            self.showHideAttachBTN.setText('Hide Mask Attach Nodes')
+            
+    def closeExistingWindow(self):
+        for qt in QtGui.qApp.topLevelWidgets():
+            try:
+                if qt.__class__.__name__ == self.__class__.__name__:
+                    qt.deleteLater()
+                    print 'uExport: Closed ' + str(qt.__class__.__name__)
+            except:
+                pass
+    
+    def buildBlendBoardFn(self):
+        fName, ok = QtGui.QFileDialog.getOpenFileName(self, "Select Poses file")
+        if ok:
+            self.currentFace.buildBlendBoard(fName)
+    
+    def constrainToMaskFn(self):
+        ma = cmds.ls(sl=1)[0]
+        if '_maskAttach' in ma:
+            utils.constrainToMesh(self.currentFace.faceMask.maskGeo, ma)
+        else:
+            cmds.warning('constrainToMask: Node does not appear to be a maskAttach nmode: ' + ma)
+    
+    def openJointMoverDlgFn(self):
+        global jointMoverWin
+        try:
+            jointMoverWin.close()
+        except:
+            pass
+        
+        jointMoverWin = jm_ui.jointMoverWin()
+        jointMoverWin.show()
+    
+    def openPoseDlgFn(self):
+        global poseWin
+        try:
+            poseWin.close()
+        except:
+            pass
+        
+        poseWin = pose_ui.poseEditorWin()
+        poseWin.show()
+
+    def openMaskAlignmentDlgFn(selfself):
+        reload(jm_ui)
+        global maskAlignWin
+        try:
+            maskAlignWin.close()
+        except:
+            pass
+
+        maskAlignWin = ma_ui.maskAlignWin()
+        maskAlignWin.show()
+
+    def selectFaceFn(self):
+        cmds.select(self.currentFace.node)
+        
+    def selectMaskFn(self):
+        cmds.select(self.currentFace.faceMask)
+        
+    def maskModeFn(self):
+        if self.maskModeBTN.isChecked():
+            self.currentFace.faceMask.active = True
+            self.maskModeBTN.setText('Mask Mode: Active')
+            print 'MASK MODE: ACTIVE'
+        else:
+            self.currentFace.faceMask.active = False
+            self.maskModeBTN.setText('Mask Mode: Disabled')
+            print 'MASK MODE: DISABLED'
+    
+    def mirrorPlacementFn(self):
+        try:
+            cmds.undoInfo(openChunk=True)
+            for node in cmds.ls(sl=1):
+                prefix = node.lower()[:1]
+                if prefix == 'l':
+                    mirrorNode = 'r_' + node[2:]
+                    if not cmds.objExists(mirrorNode):
+                        print 'Creating mirrorNode:', mirrorNode
+                        inputName = mirrorNode.split('_mover_')[0]
+                        nodes = self.currentFace.createJointMover(name=inputName, scale=0.3,toroidRatio=0.05, facialNode=self.currentFace.node)
+                        
+                    self.currentFace.mirrorFacialXform(node, axis='x', debug=0)
+            cmds.undoInfo(closeChunk=True)
+        except Exception as e:
+            print e
+            cmds.undoInfo(closeChunk=True)
+
+    def deleteSdkParentConstraintsFn(self):
+        if self.deleteSdkParentConstraintsBTN.isChecked():
+            print '>>>buildSdkParentConstraints'
+            self.currentFace.mask.buildSdkParentConstraints()
+            self.deleteSdkParentConstraintsBTN.setText('deleteSdkConstraints')
+        else:
+            print '>>>deleteSdkParentConstraints'
+            self.currentFace.mask.deleteSdkParentConstraints()
+            self.deleteSdkParentConstraintsBTN.setText('buildSdkConstraints')
+            
+    def refreshUI(self):
+        for f in utils.getUType('faceModule'):
+            self.currentFacialNodeCMB.clear()
+            self.currentFacialNodeCMB.addItem(f)
+            self.currentFace = face.FaceModule(faceNode=f)
+            self.maskControlsLBL.setText('Mask Controls: ' + str(len(self.currentFace.faceMask.controls)))
+            if self.currentFace.jointMovers:
+                self.facialMoversLBL.setText('Facial Joint Movers: ' + str(len(self.currentFace.jointMovers)))
+            self.maskGeoLBL.setText('Mask Geo: ' + self.currentFace.faceMask.maskGeo)
+            if self.currentFace.sdks:
+                self.drivenNodesLBL.setText('Driven Nodes: ' + str(len(self.currentFace.sdks)))
+
+if __name__ == '__main__':
+    show()
diff --git a/Engine/Extras/Maya_AnimationRiggingTools/ARTv1/MayaTools/General/Scripts/Modules/facial/jointMover_UI.py b/Engine/Extras/Maya_AnimationRiggingTools/ARTv1/MayaTools/General/Scripts/Modules/facial/jointMover_UI.py
new file mode 100644
--- /dev/null
+++ b/Engine/Extras/Maya_AnimationRiggingTools/ARTv1/MayaTools/General/Scripts/Modules/facial/jointMover_UI.py
@@ -0,0 +1,182 @@
+'''
+from Modules.facial import jointMover_UI as jm_ui
+reload(jm_ui)
+
+global jointMoverWin
+try:
+    jointMoverWin.close()
+except:
+    pass
+jointMoverWin = jm_ui.jointMoverWin()
+jointMoverWin.show()
+'''
+
+from PySide import QtCore, QtGui
+
+import maya.cmds as cmds
+import maya.OpenMayaUI as mui
+import shiboken
+
+import os
+
+from Modules.facial import face
+from Modules.facial import utils
+
+from Modules.facial import mask_align_UI as ma_ui
+reload(ma_ui)
+
+
+def getMayaWindow():
+    ptr = mui.MQtUtil.mainWindow()
+    if ptr is not None:
+        return shiboken.wrapInstance(long(ptr), QtGui.QWidget)
+
+class jointMoverWin(QtGui.QMainWindow):
+    def __init__(self, parent=getMayaWindow(), debug=0):
+        QtGui.QMainWindow.__init__(self, parent)
+        
+        self.currentFace = None
+        
+        #facial folder
+        facialDir = os.path.dirname(__file__)
+
+        for f in utils.getUType('faceModule'):
+            print 'jointMover_UI: Face initialized from:', f
+            self.currentFace = face.FaceModule(faceNode=f)
+        
+        ### BUILD USER INTERFACE
+        ##############################
+        
+        #create/set a central widget
+        wid = QtGui.QWidget()
+        self.setCentralWidget(wid)
+        
+        #setArt v2 stylesheet
+        utils.setArtV2StyleSheet(wid, imageDirectory='/art2images')
+        
+        self.resize(315, 514)
+        
+        #top face module picker
+        self.label = QtGui.QLabel(self)
+        self.label.setText("FacialNode:")
+        self.horizontalLayout = QtGui.QHBoxLayout()
+        self.horizontalLayout.setSpacing(2)
+        self.label.setMaximumSize(QtCore.QSize(78, 16777215))
+        self.horizontalLayout.addWidget(self.label)
+        self.currentFacialNodeCMB = QtGui.QComboBox(self)
+        self.horizontalLayout.addWidget(self.currentFacialNodeCMB)
+        
+        #attachment joint
+        self.h2 = QtGui.QHBoxLayout()
+        self.attachJointLBL = QtGui.QLabel(self)
+        self.attachJointLBL.setText("Face Attachment Joint:")
+        self.h2.setSpacing(2)
+        self.h2.addWidget(self.attachJointLBL)
+        self.attachJointLINE = QtGui.QLineEdit(self)
+        self.h2.addWidget(self.attachJointLINE)
+        
+        #main vertical layout
+        self.verticalLayout = QtGui.QVBoxLayout(wid)
+        self.verticalLayout.addLayout(self.horizontalLayout)
+        self.verticalLayout.addLayout(self.h2)
+        
+        #joint mover tree
+        self.moverTree = QtGui.QTreeWidget(self)
+        self.moverTree.setObjectName("treeWidget")
+        self.moverTree.headerItem().setText(0, "1")
+        self.verticalLayout.addWidget(self.moverTree)
+        self.moverTree.setHeaderHidden(True)
+        
+        QtCore.QMetaObject.connectSlotsByName(self)
+        
+        self.setWindowTitle("Select Joint Movers")
+
+        
+        #snap button
+        self.continueBTN = QtGui.QPushButton(self)
+        self.continueBTN.setText('>> Continue >>')
+        self.continueBTN.setObjectName("blueButton")
+        font = self.continueBTN.font()
+        font.setBold(True)
+        font.setPointSize(12)
+        self.continueBTN.setFont(font)
+        
+        #load button
+        self.loadBTN = QtGui.QPushButton(self)
+        self.loadBTN.setText('Template...')
+        font = self.loadBTN.font()
+        font.setBold(True)
+        font.setPointSize(12)
+        self.loadBTN.setFont(font)
+        self.loadBTN.setMaximumSize(QtCore.QSize(75, 16777215))
+        
+        #button layout on bottom
+        self.h3 = QtGui.QHBoxLayout()
+        self.verticalLayout.addLayout(self.h3)
+        self.h3.setSpacing(2)
+        self.h3.addWidget(self.loadBTN)
+        self.h3.addWidget(self.continueBTN)
+        
+        #connections
+        self.moverTree.itemClicked.connect(self.check_status)
+        self.continueBTN.pressed.connect(self.continueFn)
+        
+        #set the mask
+        self.mask = self.currentFace.faceMask
+        self.mask.active = True
+        cmds.showHidden(self.currentFace.sdks)
+        
+        ### FILL UI
+        ##############################
+        
+        self.attachJointLINE.setText('head')
+        self.refreshUI()
+
+    def check_status(self, item, column):
+        for i in range(0, self.moverTree.topLevelItemCount()):
+            if self.moverTree.topLevelItem(i).checkState(0) == QtCore.Qt.Unchecked:
+                utils.hideShapes(self.moverTree.topLevelItem(i).mover)
+            else:
+                utils.unhideShapes(self.moverTree.topLevelItem(i).mover)
+        cmds.select(item.text(0))
+
+
+    def continueFn(self):
+        global maskAlignWin
+        try:
+            maskAlignWin.close()
+        except:
+            pass
+
+        maskAlignWin = ma_ui.maskAlignWin()
+        maskAlignWin.show()
+        self.close()
+    
+    def refreshUI(self):
+        for f in utils.getUType('faceModule'):
+            self.currentFacialNodeCMB.addItem(f)
+            self.currentFace = face.FaceModule(faceNode=f)
+            self.buildExportTree(self.currentFace.jointMovers)
+            
+    def buildExportTree(self, jointMovers):
+        for jm in jointMovers:
+
+            #top level
+            wid1 = QtGui.QTreeWidgetItem()
+            font = wid1.font(0)
+            font.setPointSize(12)
+            wid1.setFont(0, font)
+            
+            wid1.mover = jm.replace('_SDK', '_SDK_offset')
+
+            wid1.setText(0,jm)
+
+            self.moverTree.addTopLevelItem(wid1)
+            
+            wid1.setFlags(QtCore.Qt.ItemIsEnabled | QtCore.Qt.ItemIsEditable | QtCore.Qt.ItemIsUserCheckable)
+
+            #check if enabled
+            if not utils.shapesHidden(jm):
+                wid1.setCheckState(0, QtCore.Qt.Checked)
+            else:
+                wid1.setCheckState(0, QtCore.Qt.Unchecked)
diff --git a/Engine/Extras/Maya_AnimationRiggingTools/ARTv1/MayaTools/General/Scripts/Modules/facial/mask_align_UI.py b/Engine/Extras/Maya_AnimationRiggingTools/ARTv1/MayaTools/General/Scripts/Modules/facial/mask_align_UI.py
new file mode 100644
--- /dev/null
+++ b/Engine/Extras/Maya_AnimationRiggingTools/ARTv1/MayaTools/General/Scripts/Modules/facial/mask_align_UI.py
@@ -0,0 +1,249 @@
+
+from PySide import QtCore, QtGui
+
+import maya.cmds as cmds
+import maya.OpenMayaUI as mui
+import shiboken
+
+from Modules.facial import utils
+reload(utils)
+from Modules.facial import face
+reload(face)
+
+
+import os
+
+def getMayaWindow():
+    ptr = mui.MQtUtil.mainWindow()
+    if ptr is not None:
+        return shiboken.wrapInstance(long(ptr), QtGui.QWidget)
+
+class maskAlignWin(QtGui.QMainWindow):
+    def __init__(self, parent=getMayaWindow(), debug=0):
+        QtGui.QMainWindow.__init__(self, parent)
+
+        #facial folder
+        #facialDir = os.path.dirname(__file__)
+
+        ### BUILD USER INTERFACE
+        ##############################
+
+        #create/set a central widget
+        wid = QtGui.QWidget()
+        self.setCentralWidget(wid)
+
+        #setArt v2 stylesheet
+        utils.setArtV2StyleSheet(wid, imageDirectory='/art2images')
+        
+        self.verticalLayout = QtGui.QVBoxLayout(wid)
+        self.verticalLayout.setSpacing(2)
+        
+        #top face module picker
+        self.label = QtGui.QLabel(self)
+        self.label.setText("FacialNode:")
+        self.horizontalLayout = QtGui.QHBoxLayout()
+        self.horizontalLayout.addWidget(self.label)
+        self.currentFacialNodeCMB = QtGui.QComboBox(self)
+        self.horizontalLayout.addWidget(self.currentFacialNodeCMB)
+        
+        self.resize(530, 390)
+
+        self.verticalLayout.addLayout(self.horizontalLayout)
+        
+        self.page1 = QtGui.QFrame()
+        self.page1.setMinimumSize(QtCore.QSize( 530, 380 ))
+        self.page1.setMaximumSize(QtCore.QSize( 530, 380 ))
+        
+        self.page2 = QtGui.QFrame()
+        self.page2.setMinimumSize(QtCore.QSize( 530, 380 ))
+        self.page2.setMaximumSize(QtCore.QSize( 530, 380 ))
+
+        self.page3 = QtGui.QFrame()
+        self.page3.setMinimumSize(QtCore.QSize( 530, 380 ))
+        self.page3.setMaximumSize(QtCore.QSize( 530, 380 ))
+
+        self.page4 = QtGui.QFrame()
+        self.page4.setMinimumSize(QtCore.QSize( 530, 380 ))
+        self.page4.setMaximumSize(QtCore.QSize( 530, 380 ))
+        
+        self.stackedWidget = QtGui.QStackedWidget()
+        self.stackedWidget.addWidget(self.page1)
+        self.stackedWidget.addWidget(self.page2)
+        self.stackedWidget.addWidget(self.page3)
+        self.stackedWidget.addWidget(self.page4)
+        self.verticalLayout.addWidget(self.stackedWidget)
+        
+        #set the image backdrops
+        image = (os.path.dirname(__file__) + '\\process_splash_align_mask.png').replace('\\','/')
+        self.page1.setStyleSheet("background-image: url(" + image + ");")
+        
+        image = (os.path.dirname(__file__) + '\\process_splash_align_eyes.png').replace('\\','/')
+        self.page2.setStyleSheet("background-image: url(" + image + ");")
+
+        image = (os.path.dirname(__file__) + '\\process_splash_snap.png').replace('\\','/')
+        self.page3.setStyleSheet("background-image: url(" + image + ");")
+
+        image = (os.path.dirname(__file__) + '\\process_splash_goodbye.png').replace('\\','/')
+        self.page4.setStyleSheet("background-image: url(" + image + ");")
+
+        #Add next button
+        self.nextBTN = QtGui.QPushButton(self)
+        self.nextBTN.setText('NEXT')
+        font = QtGui.QFont()
+        font.setPointSize(12)
+        font.setBold(True)
+        self.nextBTN.setFont(font)
+        self.nextBTN.setGeometry(420, 300, 90, 25)
+        
+        #Add back button
+        self.backBTN = QtGui.QPushButton(self)
+        self.backBTN.setText('BACK')
+        font = QtGui.QFont()
+        font.setPointSize(12)
+        font.setBold(True)
+        self.backBTN.setFont(font)
+        self.backBTN.setGeometry(420, 330, 90, 25)
+        self.backBTN.setHidden(True)
+        
+        self.stackedWidget.setCurrentIndex(0)
+
+        #build and hide the 'mesh chooser'
+        self.meshBTN = QtGui.QPushButton(self)
+        self.meshBTN.setText('<<')
+        font = QtGui.QFont()
+        font.setPointSize(12)
+        font.setBold(True)
+        self.meshBTN.setFont(font)
+        self.meshBTN.setGeometry(0, 0, 25, 25)
+
+        #build and hide the 'mesh chooser'
+        self.snapBTN = QtGui.QPushButton(self)
+        self.snapBTN.setText('SNAP')
+        font = QtGui.QFont()
+        font.setPointSize(12)
+        font.setBold(True)
+        self.snapBTN.setFont(font)
+        self.snapBTN.setGeometry(300, 330, 90, 25)
+
+        self.meshLBL = QtGui.QLabel(self)
+        self.meshLBL.setText("Facial Mesh:")
+        self.meshEdit = QtGui.QLineEdit(self)
+        self.hMeshLayout = QtGui.QHBoxLayout()
+        self.hMeshLayout.addWidget(self.meshLBL)
+        self.hMeshLayout.addWidget(self.meshEdit)
+        self.hMeshLayout.addWidget(self.meshBTN)
+        self.hMeshLayout.setGeometry(QtCore.QRect(100, 300, 300, 25))
+
+        #close button
+        self.closeBTN = QtGui.QPushButton(self)
+        self.closeBTN.setText('Thank you Masky,\nI will alsways remember you..')
+        font = QtGui.QFont()
+        font.setPointSize(16)
+        font.setBold(True)
+        self.closeBTN.setFont(font)
+        self.closeBTN.setGeometry(25, 250, 500, 70)
+        self.closeBTN.setHidden(True)
+
+        #button connects
+        self.nextBTN.clicked.connect(self.nextFn)
+        self.backBTN.clicked.connect(self.backFn)
+        self.meshBTN.clicked.connect(self.meshFn)
+        self.snapBTN.clicked.connect(self.snapFn)
+        self.closeBTN.clicked.connect(self.close)
+
+        self.snapVis(hide=True)
+
+        self.refreshUI()
+
+        #show joint movers and sdks
+        cmds.showHidden(self.currentFace.jointMovers)
+        cmds.showHidden(self.currentFace.sdks)
+
+    def nextFn(self):
+        index = self.stackedWidget.currentIndex()+1
+        self.stackedWidget.setCurrentIndex(index)
+        if index > 0:
+            #disable / hide mask
+            if self.currentMask.active:
+                self.currentMask.active = False
+
+            self.backBTN.setHidden(False)
+            if index == 2:
+                self.snapVis(hide=False)
+            if index == 3:
+                self.snapVis(hide=True)
+                self.closeBTN.setHidden(False)
+                self.nextBTN.setHidden(True)
+                self.backBTN.setHidden(True)
+
+        #show sdks
+        if index != 1:
+            cmds.showHidden(self.currentFace.sdks)
+            cmds.hide(self.currentMask.maskGeo)
+            cmds.hide(self.currentMask.controls)
+        else:
+            cmds.hide(self.currentFace.sdks)
+
+    def backFn(self):
+        index = self.stackedWidget.currentIndex()-1
+        self.stackedWidget.setCurrentIndex(index)
+        if index == 0:
+            self.backBTN.setHidden(True)
+            self.currentMask.active = True
+            #show joint movers and sdks
+            cmds.showHidden(self.currentFace.jointMovers)
+            cmds.showHidden(self.currentFace.sdks)
+        if index != 2:
+            self.snapVis(hide=True)
+        else:
+            self.snapVis(hide=False)
+            cmds.hide(self.currentMask.maskGeo)
+            cmds.hide(self.currentMask.controls)
+
+        #show sdks
+        if index != 1:
+            cmds.showHidden(self.currentFace.sdks)
+        else:
+            cmds.hide(self.currentFace.sdks)
+
+        self.nextBTN.setHidden(False)
+        self.closeBTN.setHidden(True)
+
+    def snapVis(self, hide=True):
+        if hide:
+            self.snapBTN.setHidden(True)
+            self.meshLBL.setHidden(True)
+            self.meshBTN.setHidden(True)
+            self.meshEdit.setHidden(True)
+
+        else:
+            self.snapBTN.setHidden(False)
+            self.meshLBL.setHidden(False)
+            self.meshBTN.setHidden(False)
+            self.meshEdit.setHidden(False)
+
+    def meshFn(self):
+        mesh = cmds.ls(sl=1)[0]
+        self.meshEdit.setText(mesh)
+        self.currentFace.renderMeshes = [mesh]
+
+    def snapFn(self):
+        self.currentFace.faceMask.active = False
+        mesh = self.meshEdit.text()
+        self.currentFace.faceMask.snapMovers('y', mesh)
+        #self.currentFace.snapEyeSdkPivots()
+
+    def updateFaceAndMask(self):
+        currentModule = self.currentFacialNodeCMB.currentText()
+        self.currentFace = face.FaceModule(faceNode=currentModule)
+        self.currentMask = self.currentFace.faceMask
+
+    def refreshUI(self):
+        oldIndex = self.currentFacialNodeCMB.currentIndex()
+        for f in utils.getUType('faceModule'):
+            self.currentFacialNodeCMB.addItem(f)
+        if oldIndex != -1:
+            self.currentFacialNodeCMB.setCurrentIndex(oldIndex)
+        self.updateFaceAndMask()
+        if self.currentFace.renderMeshes:
+            self.meshEdit.setText(self.currentFace.renderMeshes[0])
\ No newline at end of file
diff --git a/Engine/Extras/Maya_AnimationRiggingTools/ARTv1/MayaTools/General/Scripts/Modules/facial/poseEditor_UI.py b/Engine/Extras/Maya_AnimationRiggingTools/ARTv1/MayaTools/General/Scripts/Modules/facial/poseEditor_UI.py
new file mode 100644
--- /dev/null
+++ b/Engine/Extras/Maya_AnimationRiggingTools/ARTv1/MayaTools/General/Scripts/Modules/facial/poseEditor_UI.py
@@ -0,0 +1,413 @@
+from PySide import QtCore, QtGui
+
+import maya.cmds as cmds
+import maya.OpenMayaUI as mui
+import shiboken
+
+from Modules.facial import face
+reload(face)
+from Modules.facial import utils
+reload(utils)
+
+import functools
+import json
+
+def getMayaWindow():
+    ptr = mui.MQtUtil.mainWindow()
+    if ptr is not None:
+        return shiboken.wrapInstance(long(ptr), QtGui.QWidget)
+
+class poseWidget(QtGui.QWidget):
+    def __init__(self, edit=0):
+        super(poseWidget, self).__init__()
+
+        self.edit = edit
+        self.pose = None
+        self.driven = None
+        self.verticalLayout = QtGui.QVBoxLayout(self)
+
+        #horiz layout for top half
+        self.topHorizLayout = QtGui.QHBoxLayout()
+
+
+        #pose label
+        self.poseLBL = QtGui.QLabel(self)
+        font = QtGui.QFont()
+        font.setPointSize(12)
+        font.setBold(True)
+        self.poseLBL.setFont(font)
+        self.topHorizLayout.addWidget(self.poseLBL)
+
+        #edit label
+        self.editLBL = QtGui.QLabel(self)
+        self.editLBL.setText('<font color=\'#CC5200\'>POSE EDITED</font>')
+        self.editLBL.setHidden(True)
+        font = QtGui.QFont()
+        font.setPointSize(8)
+        font.setBold(True)
+        self.editLBL.setFont(font)
+        self.topHorizLayout.insertStretch(1)
+        self.topHorizLayout.addWidget(self.editLBL)
+
+        self.verticalLayout.addLayout(self.topHorizLayout)
+
+        #horiz layout for bottom half
+        self.horizontalLayout = QtGui.QHBoxLayout()
+
+        #slider
+        self.poseSlider = QtGui.QSlider(self)
+        self.poseSlider.setMaximum(100)
+        self.poseSlider.setOrientation(QtCore.Qt.Horizontal)
+        self.horizontalLayout.addWidget(self.poseSlider)
+        self.sliderValueLBL = QtGui.QLabel(self)
+        self.horizontalLayout.addWidget(self.sliderValueLBL)
+        self.sliderValueLBL.setText('0.0')
+
+        if self.edit:
+            #update button
+            self.updateBTN = QtGui.QPushButton(self)
+            self.updateBTN.setMinimumSize(QtCore.QSize(92, 0))
+            self.updateBTN.setText('UPDATE POSE')
+            font = QtGui.QFont()
+            font.setPointSize(12)
+            font.setBold(True)
+            self.updateBTN.setFont(font)
+            self.horizontalLayout.addWidget(self.updateBTN)
+            #reset button
+            self.resetBTN = QtGui.QPushButton(self)
+            self.resetBTN.setMinimumSize(QtCore.QSize(50, 0))
+            self.resetBTN.setText('RESET')
+            font = QtGui.QFont()
+            font.setPointSize(12)
+            font.setBold(True)
+            self.resetBTN.setFont(font)
+            self.horizontalLayout.addWidget(self.resetBTN)
+            self.resetBTN.setHidden(True)
+            self.resetBTN.setObjectName("orangeButton")
+        self.verticalLayout.addLayout(self.horizontalLayout)
+
+        #margins
+        self.horizontalLayout.setContentsMargins(2,2,2,2)
+        self.horizontalLayout.setSpacing(3)
+        self.verticalLayout.setContentsMargins(2,2,2,2)
+        self.verticalLayout.setSpacing(1)
+
+    def mousePressEvent(self, event):
+        nodes = list(set([x.split('.')[0] for x in self.driven]))
+        cmds.select(nodes)
+
+
+class poseEditorWin(QtGui.QMainWindow):
+    def __init__(self, parent=getMayaWindow(), mode='choose', edit=0):
+        QtGui.QMainWindow.__init__(self, parent)
+
+        self.edit = edit
+        self.currentFace = None
+
+        #TODO this won't work with multiple faces
+        for f in utils.getUType('faceModule'):
+            print 'poseChooser_UI: Face initialized from:', f
+            self.currentFace = face.FaceModule(faceNode=f)
+
+        ### BUILD USER INTERFACE
+        ##############################
+
+        #create/set a central widget
+        wid = QtGui.QWidget()
+        self.setCentralWidget(wid)
+
+        #setArt v2 stylesheet
+        utils.setArtV2StyleSheet(wid, imageDirectory='/art2images')
+
+        self.resize(400, 514)
+
+        #top face module picker
+        self.label = QtGui.QLabel(self)
+        self.label.setText("FacialNode:")
+        self.horizontalLayout = QtGui.QHBoxLayout()
+        self.horizontalLayout.setSpacing(2)
+        self.label.setMaximumSize(QtCore.QSize(78, 16777215))
+        self.horizontalLayout.addWidget(self.label)
+        self.currentFacialNodeCMB = QtGui.QComboBox(self)
+        self.horizontalLayout.addWidget(self.currentFacialNodeCMB)
+
+        #header and instruction
+        self.h2 = QtGui.QHBoxLayout()
+        self.attachJointLBL = QtGui.QLabel(self)
+
+        if self.edit:
+            self.attachJointLBL.setText("Edit Face Poses")
+        else:
+            self.attachJointLBL.setText("Available Face Poses")
+
+        font = self.attachJointLBL.font()
+        font.setBold(True)
+        font.setPointSize(13)
+        self.attachJointLBL.setFont(font)
+
+        #options horizontal layout
+        self.optionsLayout = QtGui.QHBoxLayout()
+
+        if self.edit:
+            #zero out button
+            self.zeroBTN = QtGui.QPushButton(self)
+            self.zeroBTN.setText('ZERO ALL')
+            self.zeroBTN.setMaximumSize(QtCore.QSize(75, 50))
+
+            #mirror edit out button
+            self.mirrorPoseBTN = QtGui.QPushButton(self)
+            self.mirrorPoseBTN.setText('MIRROR SELECTED POSE')
+
+            #checkboxes
+            self.moveMirrorsCHK = QtGui.QCheckBox(self)
+            self.moveMirrorsCHK.setText('Sync Movement of Mirror Poses')
+
+        self.info = QtGui.QLabel(self)
+        if self.edit:
+            self.info.setText("Check and edit the facial poses to make them better match your character")
+        else:
+            self.info.setText("These poses are available for your character's facial rig, based on the Facial Joint Movers you have selected")
+        self.info.setWordWrap(1)
+
+        self.h2.setSpacing(2)
+        self.h2.addWidget(self.attachJointLBL)
+
+
+        #main vertical layout
+        self.verticalLayout = QtGui.QVBoxLayout(wid)
+        self.verticalLayout.addLayout(self.horizontalLayout)
+        self.verticalLayout.addLayout(self.h2)
+        self.verticalLayout.addWidget(self.info)
+
+        #pose list
+        self.poseList = QtGui.QListWidget(self)
+        self.verticalLayout.addWidget(self.poseList)
+
+        QtCore.QMetaObject.connectSlotsByName(self)
+
+        self.setWindowTitle("Available Facial Poses")
+
+        if self.edit:
+            #options area
+            self.verticalLayout.addWidget(self.moveMirrorsCHK)
+            self.verticalLayout.addLayout(self.optionsLayout)
+            self.optionsLayout.addWidget(self.zeroBTN)
+            self.optionsLayout.addWidget(self.mirrorPoseBTN)
+
+
+        #continue button
+        self.continueBTN = QtGui.QPushButton(self)
+        self.continueBTN.setText('>> Continue >>')
+        self.continueBTN.setObjectName("blueButton")
+        font = self.continueBTN.font()
+        font.setBold(True)
+        font.setPointSize(12)
+        self.continueBTN.setFont(font)
+
+        #back button
+        self.backBTN = QtGui.QPushButton(self)
+        self.backBTN.setText('<< Back <<')
+        self.backBTN.setObjectName("blueButton")
+        font = self.backBTN.font()
+        font.setBold(True)
+        font.setPointSize(12)
+        self.backBTN.setFont(font)
+
+        #button layout on bottom
+        self.h3 = QtGui.QHBoxLayout()
+        self.verticalLayout.addLayout(self.h3)
+        self.h3.setSpacing(2)
+        self.h3.addWidget(self.backBTN)
+        self.h3.addWidget(self.continueBTN)
+
+        #self.verticalLayout.setContentsMargins(2,2,2,2)
+        self.verticalLayout.setSpacing(3)
+
+        #Make SDK transforms editable
+        #hide facial mask controls and mask
+
+        #show special eye and jaw movers
+
+        #unhide jointMovers, make LRAs unselectable
+
+        #connections
+        self.currentFacialNodeCMB.currentIndexChanged.connect(self.faceChangedFn)
+        #self.moverTree.itemClicked.connect(self.check_status)
+        #self.continueBTN.pressed.connect(self.continueFn)
+
+        #set the mask
+        self.mask = face.FaceMask
+        self.mask.active = True
+
+        ### CONNECTS
+        if self.edit:
+            self.zeroBTN.pressed.connect(self.zeroAllFn)
+
+        ### FILL UI
+        ##############################
+
+        self.refreshUI()
+
+    ### UI HOOKS
+    ##############################
+
+    def refreshUI(self):
+        for f in utils.getUType('faceModule'):
+            self.currentFacialNodeCMB.addItem(f)
+            self.currentFace = face.FaceModule(faceNode=f)
+            self.buildPoseList()
+
+    def updateFn(self, wid):
+        #check that the pose is fully on
+        if wid.poseSlider.value() != 100:
+            cmds.warning('Pose slider is not fully on, must be 1.0 in able to save edit.')
+            return False
+
+        cmds.setDrivenKeyframe(wid.driven, cd=wid.pose)
+        self.checkForEditedPoses()
+
+    def resetFn(self, wid):
+        print 'resetting', wid
+        try:
+            cmds.undoInfo(openChunk=True)
+            #get original pose from driven dict
+            initialPoseDict = json.loads(cmds.getAttr(self.currentFace.blendBoard + '.initialPoseDict'))
+            initialPoseAttrs = initialPoseDict[wid.pose]
+
+            #store current slider/attr val
+            preVal = cmds.getAttr(wid.pose)
+
+            #iterate over the attrs and set them bvack to original
+            #TODO: rotate order?
+            cmds.setAttr(wid.pose, 1.0)
+            for att in wid.driven:
+                #print att, initialPoseAttrs[att]
+                cmds.setAttr(att, initialPoseAttrs[att])
+
+            #set the old driven keyframe
+            cmds.setDrivenKeyframe(wid.driven, cd=wid.pose)
+
+            #reset the slider/attr val
+            cmds.setAttr(wid.pose, preVal)
+            self.checkForEditedPoses()
+
+        except Exception, err:
+            import traceback
+            print(traceback.format_exc())
+        finally:
+            cmds.undoInfo(closeChunk=True)
+
+    def sliderChanged(self, wid, value):
+        val = float(value)/100
+        floatVal = str(val)
+        wid.sliderValueLBL.setText(floatVal)
+        cmds.setAttr(wid.pose, val)
+
+    def zeroAllFn(self):
+        for i in range(0, self.poseList.count()):
+            wid = self.poseList.itemWidget(self.poseList.item(i))
+            wid.poseSlider.setValue(0)
+
+    def getAllDrivenAttrs(self):
+        ret = {}
+        for i in range(0, self.poseList.count()):
+            wid = self.poseList.itemWidget(self.poseList.item(i))
+            ret[wid.pose] = cmds.getAttr(wid.pose)
+        return ret
+
+    def setAllDrivenAttrs(self, dict):
+        for pose in dict:
+            cmds.setAttr(pose, dict[pose])
+
+    def zeroAllAttrs(self):
+        for i in range(0, self.poseList.count()):
+            wid = self.poseList.itemWidget(self.poseList.item(i))
+            cmds.setAttr(wid.pose, 0.0)
+
+    #not used atm
+    def getSliders(self):
+        ret = []
+        for i in range(0, self.poseList.count()):
+            wid = self.poseList.itemWidget(self.poseList.item(i))
+            ret.append(wid.poseSlider.value())
+        return ret
+
+    #not used atm
+    def setSliders(self, sliders):
+        for i in range(0, len(sliders)):
+            wid = self.poseList.itemWidget(self.poseList.item(i))
+            wid.poseSlider.setValue(sliders[i])
+
+    def checkForEditedPoses(self, threshold=0.001):
+        initialPoseDict = json.loads(cmds.getAttr(self.currentFace.blendBoard + '.initialPoseDict'))
+
+        for i in range(0, self.poseList.count()):
+            edited = False
+            wid = self.poseList.itemWidget(self.poseList.item(i))
+            preVals = self.getAllDrivenAttrs()
+            self.zeroAllAttrs()
+            cmds.setAttr(wid.pose, 1.0)
+
+            #iterate and check driven att values
+            initialPoseAttrs = initialPoseDict[wid.pose]
+            for att in wid.driven:
+                error = abs(initialPoseAttrs[att] - cmds.getAttr(att))
+                if error > threshold:
+                    edited = True
+            if edited:
+                wid.editLBL.setHidden(False)
+                wid.resetBTN.setHidden(False)
+            else:
+                wid.editLBL.setHidden(True)
+                wid.resetBTN.setHidden(True)
+            self.setAllDrivenAttrs(preVals)
+
+    def addPose(self, pose, driven):
+        wid = poseWidget(edit=self.edit)
+        wid.pose = pose
+        wid.driven = driven
+        currentVal = cmds.getAttr(pose)
+        wid.poseLBL.setText(pose.split('.')[-1])
+        wid2 = QtGui.QListWidgetItem()
+        self.poseList.addItem(wid2)
+        wid2.setSizeHint(QtCore.QSize(100, 60))
+        self.poseList.setItemWidget(wid2, wid)
+
+        #hookup slider
+        wid.poseSlider.valueChanged.connect(functools.partial(self.sliderChanged, wid))
+        wid.poseSlider.setValue(currentVal*100)
+
+        if self.edit:
+            #hookup buttons
+            wid.updateBTN.pressed.connect(functools.partial(self.updateFn, wid))
+            wid.resetBTN.pressed.connect(functools.partial(self.resetFn, wid))
+
+        #store driven for checking later
+        wid.driven = driven
+
+    def buildPoseList(self):
+        self.poseList.clear()
+        for pose in self.currentFace.poses:
+            self.addPose(pose, self.currentFace.poses[pose])
+        if self.edit:
+            self.checkForEditedPoses()
+
+    def faceChangedFn(self):
+        self.currentFace = face.FaceModule(faceNode=self.currentFacialNodeCMB.currentText())
+        self.buildPoseList()
+
+if __name__ == '__main__':
+    global poseWin
+    try:
+        poseWin.close()
+    except:
+        pass
+
+
+    poseWin = poseEditorWin(edit=1)
+    poseWin.show()
+
+
+
+
+
diff --git a/Engine/Extras/Maya_AnimationRiggingTools/ARTv1/MayaTools/General/Scripts/Modules/facial/utils.py b/Engine/Extras/Maya_AnimationRiggingTools/ARTv1/MayaTools/General/Scripts/Modules/facial/utils.py
new file mode 100644
--- /dev/null
+++ b/Engine/Extras/Maya_AnimationRiggingTools/ARTv1/MayaTools/General/Scripts/Modules/facial/utils.py
@@ -0,0 +1,457 @@
+import maya.api.OpenMaya as om2
+import maya.cmds as cmds
+
+import os
+import traceback
+
+## GENERAL
+#######################################################################
+
+def dagPathFromName(obj):
+    '''
+    Given a maya object name, returns an openMaya2 dagPath object, dagPath, and a full dagPath in a string
+    '''
+    sel = om2.MSelectionList()
+    sel.add(obj)
+    dag = sel.getDagPath(0)
+    dagStr = sel.getDagPath(0).fullPathName()
+    mobj = sel.getDependNode(0)
+    return mobj, dag, dagStr
+
+
+
+### ATTRIBUTE MANAGEMENT
+#######################################################################
+def attrExists(attr):
+    if '.' in attr:
+        node, att = attr.split('.')
+        return cmds.attributeQuery(att, node=node, ex=1)
+    else:
+        cmds.warning('attrExists: No attr passed in: ' + attr)
+        return False
+
+def msgConnect(attribFrom, attribTo, debug=0):
+    #TODO needs a mode to dump all current connections (overwrite/force)
+    objFrom, attFrom = attribFrom.split('.')
+    objTo, attTo = attribTo.split('.')
+    if debug: print 'msgConnect>>> Locals:', locals()
+    if not attrExists(attribFrom):
+        cmds.addAttr(objFrom, longName=attFrom, attributeType='message')
+    if not attrExists(attribTo):
+        cmds.addAttr(objTo, longName=attTo, attributeType='message')
+    
+    #check that both atts, if existing are msg atts
+    for a in (attribTo, attribFrom):
+        if cmds.getAttr(a, type=1) != 'message':
+            cmds.warning('msgConnect: Attr, ' + a + ' is not a message attribute. CONNECTION ABORTED.')
+            return False
+    
+    try:
+        return cmds.connectAttr(attribFrom, attribTo)
+    except Exception as e:
+        print e
+        return False
+
+def snapMeshPivot(mesh1, mesh2):
+	'''
+	move mesh2 pivot to mesh1 pivot
+	@author=chrise
+	'''
+
+	dupe = cmds.duplicate(mesh2)[0]
+	mesh1Xform = cmds.getAttr(mesh1+'.worldMatrix')
+	cmds.xform(mesh2, m=mesh1Xform)
+
+	cmds.undoInfo(openChunk=True)
+	try:
+		for v in range(0, cmds.polyEvaluate(mesh2, v=1)):
+			m2Vtx = mesh2+'.vtx['+str(v)+']'
+			dupeVtx = dupe+'.vtx['+str(v)+']'
+			cmds.xform(m2Vtx, m=cmds.xform(dupeVtx, m=1, q=1, ws=1), ws=1)
+		cmds.delete(dupe)
+	finally:
+		cmds.undoInfo(closeChunk=True)
+	'''
+	#TODO: here's some sample code for working with one nurbs shape
+	#I stopped because it didn't work with multiple shape nodes
+	surf = cmds.ls(sl=1)[0]
+    if cmds.nodeType(surf) == 'transform':
+        #get all cvs
+        cmds.select(surf + '.cv[*]')
+        cvs = cmds.ls(sl=1, flatten=1)
+        cmds.select(cl=1)
+        for cv in cvs:
+            print cv
+            print cmds.pointPosition(cv)
+    '''
+
+def snapRotatePivotPosition(nodeFrom, nodeTo):
+    '''
+    sel = cmds.ls(sl=1)
+    snapRotatePivotPosition(sel[0], sel[1])
+    '''
+    p = cmds.xform(nodeTo, q=1, ws=1, rp=1)
+    cmds.xform(nodeFrom, ws=1, piv=p)
+
+def snapPivotWithZeroParent(fromNode, toNode):
+    '''
+    Mainly used for snapping the pivot of a controller that is in a topnode GRP's space
+    '''
+    #TODO: This won't work if there is more than one child
+    try:
+        parent = cmds.listRelatives(fromNode, parent=1)
+
+        cmds.parent(fromNode, w=1)
+        toNodeMx = cmds.xform(toNode, q=1, ws=1, m=1)
+        cmds.xform(parent, m=toNodeMx, ws=1)
+
+        toNodePivot = cmds.xform(toNode, ws=1, q=1, t=1)
+        cmds.xform(fromNode, ws=1, a=1, piv=toNodePivot)
+        cmds.parent(fromNode, parent)
+        cmds.makeIdentity(fromNode, t=1, r=1, a=1)
+    except Exception as e:
+        print(traceback.format_exc())
+
+def getUType(type):
+    '''
+    Returns nodes of the requested uType
+    example: getUType('faceModule')
+    '''
+    #TODO: Check with namespaces
+    return [node.split('.')[0] for node in cmds.ls('*.uType', r=1) if cmds.getAttr(node) == type]
+
+def hideShapes(node):
+    for shape in cmds.listRelatives(node, shapes=1, f=1):
+        cmds.hide(shape)
+
+def unhideShapes(node):
+    for shape in cmds.listRelatives(node, shapes=1, f=1):
+        cmds.showHidden(shape)
+
+def shapesHidden(node):
+    for shape in cmds.listRelatives(node, shapes=1, f=1):
+        if cmds.getAttr(shape + '.visibility') == 1:
+            return False
+    return True
+
+### CONTROLLER MANAGEMENT
+#######################################################################
+def toggleShapeSelection(node, selectable):
+    '''
+    example: makeShapesUnselectable(cmds.ls(sl=1)[0], False)
+    '''
+    for shape in cmds.listRelatives(node, fullPath=1, shapes=1):
+        if selectable:
+            cmds.setAttr(shape + '.overrideEnabled', 0)
+            cmds.setAttr(shape + '.overrideDisplayType', 0)
+        else:
+            cmds.setAttr(shape + '.overrideEnabled', 1)    
+            cmds.setAttr(shape + '.overrideDisplayType', 2)
+
+def getAxisWorldVectors(node):
+    matrix = om2.MGlobal.getSelectionListByName(node).getDagPath(0).inclusiveMatrix()
+    x = ((om2.MVector.kXaxisVector) * matrix).normal()
+    y = ((om2.MVector.kYaxisVector) * matrix).normal()
+    z = ((om2.MVector.kZaxisVector) * matrix).normal()
+    return [x,y,z]
+
+def createArrow(name=None, thickness=0.1, length=2, vector=[1,0,0], point=[0,0,0]):
+    '''
+    Creates an arrow in the direction of the vector
+    Example:
+    arrow = createArrow(vector=[0,1,0], length=4)
+    '''
+    if not name:
+        name = 'arrow_CTRL'
+    #calc length for thickness
+    ratio = length/thickness
+    cyl = cmds.cylinder(radius=thickness, sections=4, heightRatio=ratio, pivot=[length/2, 0, 0], ch=0)[0]
+    cone = cmds.cone(radius=thickness*2, sections=4, ch=0, pivot=(length*1.0005,0,0))[0]
+    xform = cmds.createNode('transform', ss=1, name=name)
+    
+    shapes = []
+    transforms = []
+    for node in [cone,cyl]:
+        shapes.append(cmds.listRelatives(node, fullPath=1, shapes=1)[0])
+        transforms.append(node)
+    cmds.parent(shapes, xform, r=1, s=1)
+    rotateBy = cmds.angleBetween(euler=True, v1=(1,0,0), v2=vector)
+    cmds.rotate(rotateBy[0], rotateBy[1], rotateBy[2], xform)
+    cmds.xform(xform, t=point)
+    cmds.delete(transforms)
+    return xform
+
+
+def colorControl(node, name=None, color=None, useExisting=1, shapes=0):
+    '''
+    Colors an object by creating and assigning a shader with the color you input
+    Example:
+    colorControl('arrow', color=(1,0,0))
+    '''
+    if not name:
+        name = 'colorControl_%s_%s_%s' % (color[0], color[1], color[2])
+    
+    shader = None
+    shading_group = None
+    
+    if cmds.objExists(name):
+        if 'shadingDependNode' in cmds.nodeType(name, i=1):
+            if useExisting:
+                shader = name
+                shading_group = cmds.listConnections(name + '.outColor')[0]
+    else:
+        if len(color) == 3:
+            shader=cmds.shadingNode('blinn',asShader=True, name = name)
+            cmds.setAttr(shader + '.color', color[0], color[1], color[2], type='double3')
+            shading_group = cmds.sets(renderable=True,noSurfaceShader=True,empty=True, name=name + '_colorControlSG')
+            cmds.connectAttr('%s.outColor' %shader ,'%s.surfaceShader' %shading_group)
+        else:
+            cmds.warning('Color input is not a 3 element list or tuple. COLOR: ' + str(color))
+            return False
+    if shapes:
+        for shape in cmds.listRelatives(node, fullPath=1, shapes=1):
+            cmds.sets(shape, e=1, forceElement=shading_group)
+    else:
+        cmds.sets(node, e=1, forceElement=shading_group)
+    cmds.setAttr(shader + '.diffuse', 1)
+    cmds.setAttr(shader + '.ambientColor', 1,1,1, type='double3')
+
+    #Need to check if it's a blinn before setting these
+    #cmds.setAttr(shader + '.eccentricity', 0)
+    #cmds.setAttr(shader + '.specularColor', 0,0,0, type='double3')
+    return shader, shading_group
+    
+def createLRA(node=None, matrix=None, name=None, color=True, thickness=0.1, length=2, vector=[1,0,0], point=[0,0,0], arrowRadiusMult=2):
+    '''
+    Creates an LRA at the origin, or at a node, or a mmx
+    Example:
+    createLRA(length=2, thickness=0.05, arrowRadiusMult=3)
+    '''
+    if not name:
+        name = 'arrow_CTRL'
+    nodes = []
+    x,y,z = [],[],[]
+    #calc length for thickness
+    ratio = length/thickness
+    xform = cmds.createNode('transform', ss=1, name=name)
+
+    x.append(cmds.cylinder(radius=thickness, sections=4, heightRatio=ratio, pivot=[length/2, 0, 0], ch=0, axis=[1,0,0])[0])
+    x.append(cmds.cone(radius=thickness*arrowRadiusMult, sections=4, ch=0, pivot=(length*1.0005,0,0))[0])
+    
+    y.append(cmds.cylinder(radius=thickness, sections=4, heightRatio=ratio, pivot=[0, length/2, 0], ch=0, axis=[0,1,0])[0])
+    y.append(cmds.cone(radius=thickness*arrowRadiusMult, sections=4, ch=0, pivot=(0,length*1.0005,0), axis=[0,1,0])[0])
+    
+    z.append(cmds.cylinder(radius=thickness, sections=4, heightRatio=ratio, pivot=[0, 0, length/2], ch=0, axis=[0,0,1])[0])
+    z.append(cmds.cone(radius=thickness*arrowRadiusMult, sections=4, ch=0, pivot=(0,0,length*1.0005), axis=[0,0,1])[0])
+    
+    nodes.extend(x)
+    nodes.extend(y)
+    nodes.extend(z)
+    
+    if color:
+        for node in x: colorControl(node, name='red_m', color=(1,0,0))
+        for node in y: colorControl(node, name='green_m', color=(0,1,0))
+        for node in z: colorControl(node, name='blue_m', color=(0,0,1))
+    
+    shapes = []
+    transforms = []
+    for node in nodes:
+        shapes.append(cmds.listRelatives(node, fullPath=1, shapes=1)[0])
+        transforms.append(node)
+    cmds.parent(shapes, xform, r=1, s=1)
+    rotateBy = cmds.angleBetween(euler=True, v1=(1,0,0), v2=vector)
+    cmds.rotate(rotateBy[0], rotateBy[1], rotateBy[2], xform)
+    cmds.xform(xform, t=point)
+    cmds.delete(transforms)
+    return xform
+    
+
+def createCircleControl(name=None, x=1, y=1, z=1, radius=1, color=None, type='strip', heightRatio=0.007):
+    if not name:
+        name = 'circle_CTRL'
+    xform = cmds.group(name=name, em=1)
+    circleShapes = []
+    circleTransforms = []
+    vectors = getAxisWorldVectors(xform)
+    
+    xyz = [x,y,z]
+    for i in range(0, len(xyz)):
+        if xyz[i]:
+            circ = None
+            if type == 'curve':
+                circ = cmds.circle(radius=radius, name='jointMoverAxisCtrl', normal=vectors[i], sections=1, ch=0)
+                if color:
+                    cmds.color(circ, ud=color)
+            if type == 'strip' or type == 'toroid':
+                circ = None
+                if type == 'strip':
+                    circ = cmds.cylinder(radius=radius, name='jointMoverAxisCtrl', axis=vectors[i], sections=7, ch=0, heightRatio=heightRatio)
+                elif type == 'toroid':
+                    circ = cmds.torus(radius=radius, name='jointMoverAxisCtrl', axis=vectors[i], sections=7, ch=0, heightRatio=heightRatio)
+                if color:
+                    if color == 'xyz':
+                        if i == 0: colorControl(circ, name='red_m', color=(1,0,0))
+                        if i == 1: colorControl(circ, name='green_m', color=(0,1,0))
+                        if i == 2: colorControl(circ, name='blue_m', color=(0,0,1))
+                    else:
+                        if len(color) == 3: colorControl(circ, color=color)
+            circleShapes.append(cmds.listRelatives(circ[0], fullPath=1, shapes=1)[0])
+            circleTransforms.append(circ[0])
+    cmds.parent(circleShapes, xform, r=1, s=1)
+    cmds.delete(circleTransforms)
+    return xform
+    
+### ART V2 WRANGLING
+#######################################################################
+
+def setArtV2StyleSheet(widget, sheetTemplate=None, imageDirectory=None):
+    
+    css, lines = None, None
+    if not sheetTemplate:
+        sheetTemplate = (os.path.dirname(__file__) + '\\artV2_css_template.txt').replace('\\','/')
+        imageDirectory = (os.path.dirname(__file__) + imageDirectory).replace('\\','/')
+    if os.path.isfile(sheetTemplate):
+        f = open(sheetTemplate, 'r')
+        cssTemplate = f.read()
+        f.close()
+        css = cssTemplate.replace('___PATH___', imageDirectory)
+    else:
+        print('Cannot find CSS template: ' + sheetTemplate)
+        return False
+    
+    #check to see if all images found
+    for line in css.split('\n'):
+        if '.png' in line.lower():
+            image = line.split('(')[1].split(')')[0]
+            if not os.path.isfile(image):
+                print 'ERROR:setArtV2StyleSheet: Image not found [' + image + ']'
+    
+    widget.setStyleSheet(css)
+
+## SURFACE CONSTRAINING
+#######################################################################
+
+def square_distance(pointA, pointB):
+    '''
+    Both points must have the same dimensions
+    '''
+    if len(pointA) != len(pointB):
+        print 'SQUARED_DISTANCE: incoming points do not have the same dimensions:', pointA, pointB
+        return False
+    distance = 0
+    dimensions = len(pointA)
+    for dimension in range(dimensions):
+        distance += (pointA[dimension] - pointB[dimension])**2
+    return distance
+
+def closestVertOnMesh(mesh, node):
+    '''
+    Returns closest vertex component on mesh to transform/pivot of node
+    '''
+    closestVert = None
+    minLength = None
+    verts = cmds.getAttr(mesh+".vrts", multiIndices=True)
+    nodePos = cmds.xform(node, q=1, translation=1, ws=1)
+    for v in verts:
+        vPos = cmds.xform(mesh + '.vtx['+str(v)+']', q=1, translation=1, ws=1)
+        length = om2.MVector(om2.MVector(nodePos) - om2.MVector(vPos)).length()
+        if minLength is None or length < minLength:
+            minLength = length
+            closestVert = v
+    return (mesh + '.vtx[' + str(closestVert) + ']')
+
+def closestPointonMesh(mesh, node, createLoc=False):
+    pos = cmds.xform(node, q=1, rp=1, ws=1)
+    nodePos = om2.MPoint(pos[0], pos[1], pos[2],1.0)
+    print 'point', pos[0], pos[1], pos[2]
+
+    obj,dag,stringPath = dagPathFromName(mesh)
+    meshFn = om2.MFnMesh(dag)
+
+    resultPoint = meshFn.getClosestPoint(nodePos,om2.MSpace.kWorld)[0]
+    
+    if createLoc:
+        loc = cmds.spaceLocator()[0]
+        cmds.setAttr(loc+'.t', resultPoint[0], resultPoint[1], resultPoint[2])
+        return loc
+    else:
+        return (resultPoint[0], resultPoint[1], resultPoint[2])
+
+def constrainToMesh(mesh, node, mo=True, type='follicle'):
+    '''
+    Attaches node to closest vertex of mesh
+    '''
+    
+    #get closest vertex
+    vtx = closestVertOnMesh(mesh, node)
+    
+    surfaceXform = None
+    #build/place follicle
+    if type == 'follicle':
+        surfaceXform = cmds.createNode("transform", n=(mesh+"_follicle"))
+        follicle = cmds.createNode("follicle", n=(mesh+"_follicleShape"), p=surfaceXform)
+        cmds.connectAttr((follicle+".outTranslate"), (surfaceXform+".translate"), f=1)
+        cmds.connectAttr((follicle+".outRotate"), (surfaceXform+".rotate"), f=1)
+        cmds.connectAttr((mesh+".outMesh"), (follicle+".inm"), f=1)
+        cmds.connectAttr((mesh+".worldMatrix[0]"), (follicle+".inputWorldMatrix"), f=1)
+        uv = cmds.polyListComponentConversion(vtx, fv=True, tuv=True)
+        uvs = cmds.polyEditUV(uv ,q=1, u=1, v=1)
+        u = uvs[0]
+        v = uvs[1]
+        cmds.setAttr((follicle+".parameterU"), u)
+        cmds.setAttr((follicle+".parameterV"), v)
+    cmds.parentConstraint(surfaceXform, node, mo=mo)
+
+def rayMeshIntersect(mesh, pos, vector):
+    '''
+    Example:
+    mover = cmds.ls(sl=1)[0]
+    mesh = 'steel'
+    pos = cmds.xform(mover, q=1, rp=1, ws=1)
+    newPos = rayMeshIntersect(mesh, pos, -utils.getAxisWorldVectors(mover)[1])
+    cmds.xform(mover, t=newPos, ws=1)
+    '''
+    nodePos = om2.MPoint(pos[0], pos[1], pos[2],1.0)
+
+    vector = om2.MVector(vector)
+    obj,dag,stringPath = dagPathFromName(mesh)
+    meshFn = om2.MFnMesh(dag)
+
+    selectionList = om2.MSelectionList()
+    selectionList.add(mesh)
+    dagPath = selectionList.getDagPath(0)
+    fnMesh = om2.MFnMesh(dagPath)
+    intersection = fnMesh.closestIntersection(
+    om2.MFloatPoint(pos),
+    om2.MFloatVector(vector),
+    om2.MSpace.kWorld, 99999, False)
+    if intersection:
+        hitPoint, hitRayParam, hitFace, hitTriangle, hitBary1, hitBary2 = intersection
+        x,y,z,_ = hitPoint
+        return (x,y,z)
+
+def deleteConstraints(node):
+    ch = cmds.listRelatives(node, children=1)
+    if ch:
+        for item in ch:
+            if 'Constraint' in cmds.nodeType(item):
+                cmds.delete(item)
+
+def findNamingConflicts(ignoreShapes=1):
+    #TODO: listRelatives to find parent, check relationship to make sure theyre shapes instead of namechecking
+    nameDict = {}
+    for node in cmds.ls('*', l=1):
+        if ignoreShapes:
+            if 'shape' in node.lower():
+                continue
+        shortName = node.split('|')[-1]
+        if shortName not in nameDict.keys():
+            nameDict[shortName] = [node]
+        else:
+            cmds.warning('Found duplicate node names: \n' + node + '\n' + str(nameDict[shortName]))
+            nameDict[shortName].append(node)
+    returnDict = {}
+    for name in nameDict:
+        if len(nameDict[name]) > 1:
+            returnDict[name] = nameDict[name]
+    return returnDict
\ No newline at end of file
diff --git a/Engine/Extras/Maya_AnimationRiggingTools/ARTv1/MayaTools/General/Scripts/Modules/p4python/EpicP4.py b/Engine/Extras/Maya_AnimationRiggingTools/ARTv1/MayaTools/General/Scripts/Modules/p4python/EpicP4.py
new file mode 100644
--- /dev/null
+++ b/Engine/Extras/Maya_AnimationRiggingTools/ARTv1/MayaTools/General/Scripts/Modules/p4python/EpicP4.py
@@ -0,0 +1,86 @@
+from Modules.p4python.P4 import P4, P4Exception
+import maya.cmds as cmds
+
+def findFileP4(asset, debug=0, root='//depot/ArtSource/'):
+    p4 = P4()
+    try: p4.connect()
+    except:
+        print 'Cannot connect!'
+        return False
+
+    try:
+        file =  p4.run_files(root + '...' + asset)
+        depotLoc = file[0]['depotFile']
+        describe = p4.run_describe(file[0]['change'])
+        return [describe[0]['user'], depotLoc, describe[0]['desc'], file[0]['change']]
+    except Exception as e:
+        print "Cannot find file.", asset
+        if debug: print e
+        return False
+    finally:
+        p4.disconnect()
+
+def rePathFileNodesP4(debug=0, searchPath='//depot/ArtSource/'):
+    nodes = cmds.ls(type='file')
+
+    for node in nodes:
+        path = cmds.getAttr((node + '.fileTextureName'))
+        newPath = None
+
+        #check if the file exists on disk
+        if cmds.file(path, exists=1, q=1):
+            pass
+
+        #if the file is not on disk, continue and get it from P4
+        else:
+            if debug:
+                print 'rePathFileNodesP4>>>> Local file not found', path, 'searching p4 database.'
+
+            #connect to the perforce server
+            p4 = P4()
+            try: p4.connect()
+            except:
+                print '>>rePathFileNodesP4: Cannot connect to perforce!'
+                return False
+
+            try:
+                #convert the path to the same direction slashes
+                floppedPath = path.replace('\\', '/')
+                fName = floppedPath.split('/')[-1]
+
+                #search p4 for the file
+                f = None
+                try:
+                    f = p4.run_files(searchPath + '...' + fName)
+                except:
+                    pass
+                if f:
+                    if len(f) > 1:
+                        print 'rePathFileNodesP4>>>> Multiple files [', len(f), '] found in depot search path with the name', f
+                        for i in range(0, len(f)):
+                            try:
+                                print 'rePathFileNodesP4>>>> SYNCING: ', f[i]['depotFile']
+                                p4.run_sync(f[i]['depotFile'])
+                            except Exception as e:
+                                if debug:
+                                    print e
+                    else:
+                        print 'rePathFileNodesP4>>>> File found on depot: ', f[0]['depotFile']
+                        try:
+                            p4.run_sync(f[0]['depotFile'])
+                        except Exception as e:
+                                if debug:
+                                    print e
+                else:
+                    print 'rePathFileNodesP4>>>> FILE NOT FOUND IN PERFORCE SEARCH PATH - fName:', fName, 'searchPath:', searchPath
+
+                fileFstat = p4.run('fstat', f[0]['depotFile'])[0]
+                newPath = fileFstat['clientFile']
+
+                if debug: print 'rePathFileNodesP4>>>> NEW PATH>>', newPath, '\n\n'
+                cmds.setAttr((node + '.fileTextureName'), newPath, type='string')
+            except Exception as e:
+                print "Cannot find file.", fName,e
+                if debug:
+                    print 'rePathFileNodesP4>>>> EXCEPTION', e
+    p4.disconnect()
\ No newline at end of file
