cleaning up

diff --git a/Engine/Extras/Maya_AnimationRiggingTools/ARTv1/MayaTools/General/Scripts/Modules/p4python/P4.py b/Engine/Extras/Maya_AnimationRiggingTools/ARTv1/MayaTools/General/Scripts/Modules/p4python/P4.py
new file mode 100644
--- /dev/null
+++ b/Engine/Extras/Maya_AnimationRiggingTools/ARTv1/MayaTools/General/Scripts/Modules/p4python/P4.py
@@ -0,0 +1,1036 @@
+#!/usr/local/bin/python
+
+from __future__ import print_function
+
+"""A Python version of the Perforce "p4" client.
+    
+    This uses the Python type P4API.P4Adapter, which is a wrapper for the
+    Perforce ClientApi object.
+    
+    $Id: //depot/r15.1/p4-python/P4.py#1 $
+    
+    #*******************************************************************************
+    # Copyright (c) 2007-2010, Perforce Software, Inc.  All rights reserved.
+    # Portions Copyright (c) 1999, Mike Meyer. All rights reserved.
+    # Portions Copyright (c) 2004-2007, Robert Cowham. All rights reserved.
+    #
+    # Redistribution and use in source and binary forms, with or without
+    # modification, are permitted provided that the following conditions are met:
+    #
+    # 1.  Redistributions of source code must retain the above copyright
+    # notice, this list of conditions and the following disclaimer.
+    #
+    # 2.  Redistributions in binary form must reproduce the above copyright
+    # notice, this list of conditions and the following disclaimer in the
+    # documentation and/or other materials provided with the distribution.
+    #
+    # THIS SOFTWARE IS PROVIDED BY THE COPYRIGHT HOLDERS AND CONTRIBUTORS "AS IS"
+    # AND ANY EXPRESS OR IMPLIED WARRANTIES, INCLUDING, BUT NOT LIMITED TO, THE
+    # IMPLIED WARRANTIES OF MERCHANTABILITY AND FITNESS FOR A PARTICULAR PURPOSE
+    # ARE DISCLAIMED. IN NO EVENT SHALL PERFORCE SOFTWARE, INC. BE LIABLE FOR ANY
+    # DIRECT, INDIRECT, INCIDENTAL, SPECIAL, EXEMPLARY, OR CONSEQUENTIAL DAMAGES
+    # (INCLUDING, BUT NOT LIMITED TO, PROCUREMENT OF SUBSTITUTE GOODS OR SERVICES;
+    # LOSS OF USE, DATA, OR PROFITS; OR BUSINESS INTERRUPTION) HOWEVER CAUSED AND
+    # ON ANY THEORY OF LIABILITY, WHETHER IN CONTRACT, STRICT LIABILITY, OR TORT
+    # (INCLUDING NEGLIGENCE OR OTHERWISE) ARISING IN ANY WAY OUT OF THE USE OF THIS
+    # SOFTWARE, EVEN IF ADVISED OF THE POSSIBILITY OF SUCH DAMAGE.
+    #*******************************************************************************
+    
+    Copyright 2007-2010 Perforce Software, Inc. All rights reserved
+    
+    
+    License:
+    See accompanying LICENSE.txt including for redistribution permission.
+    """
+
+import sys, datetime
+import re
+import shutil
+from contextlib import contextmanager
+import uuid, tempfile
+import os, os.path, platform
+import subprocess
+
+# P4Exception - some sort of error occurred
+class P4Exception(Exception):
+    """Exception thrown by P4 in case of Perforce errors or warnings"""
+    
+    def __init__(self, value):
+        Exception.__init__(self)
+        
+        if isinstance(value, (list, tuple)) and len(value) > 2:
+            self.value = value[0]
+            self.errors = value[1]
+            self.warnings = value[2]
+        else:
+            self.value = value
+    def __str__(self):
+        return str(self.value)
+
+class Spec(dict):
+    """Subclass of dict, representing the fields of a spec definition.
+        
+        Attributes can be accessed either with the conventional dict format,
+        spec['attribute'] or with shorthand spec._attribute.
+        
+        Instances of this class will preventing any unknown keys.
+        """
+    def __init__(self, fieldmap=None):
+        self.__dict__['_Spec__fields'] = fieldmap
+    
+    def permitted_fields(self):
+        return self.__fields
+    
+    def __setitem__(self, key, value):
+        if not isinstance(value, str) and not isinstance(value, list):
+            raise P4Exception("Illegal value of type %s, must be string or list" % value.__class__)
+        if key in self or self.__fields == None:
+            dict.__setitem__(self, key, value)
+        elif str(key).lower() in self.__fields:
+            dict.__setitem__(self, self.__fields[key.lower()], value)
+        else:
+            raise P4Exception("Illegal field '%s'" % str(key))
+    
+    def __getattr__(self, attr):
+        key = str(attr).lower()
+        if key[0] != '_':
+            raise AttributeError(attr)
+        key = key[1:]
+        if key in self:
+            return self[key]
+        elif key in self.__fields:
+            return self[self.__fields[key]]
+    
+    def __setattr__(self, attr, value):
+        if attr == 'comment':
+            self.__dict__[attr] = value
+        else:
+            key = str(attr).lower()
+            if key[0] != '_':
+                raise AttributeError(attr)
+            key = key[1:]
+            self[key] = value
+
+#
+# P4Integration objects hold details about the integrations that have
+# been performed on a particular revision. Used primarily with the
+# P4Revision class
+#
+class Integration:
+    def __init__( self, how, file, srev, erev ):
+        self.how = how
+        self.file = file
+        self.srev = srev
+        self.erev = erev
+    
+    def __repr__(self):
+        return "Integration (how = %s file = %s srev = %s erev = %s)" \
+            % (self.how, self.file, self.srev, self.erev)
+#
+# Each P4Revision object holds details about a particular revision
+# of a file. It may also contain the history of any integrations
+# to/from the file
+#
+
+class Revision:
+    def __init__( self, depotFile ):
+        self.depotFile = depotFile
+        self.integrations = []
+        self.rev = None
+        self.change = None
+        self.action = None
+        self.type = None
+        self.time = None
+        self.user = None
+        self.client = None
+        self.desc = None
+        self.digest = None
+        self.fileSize = None
+    
+    def integration( self, how, file, srev, erev ):
+        rec = Integration( how, file, srev, erev )
+        self.integrations.append( rec )
+        return rec
+    
+    # iterator over the collection calling a provided function
+    # Python's poor man version of the Ruby block
+    def each_integration(self):
+        for i in self.integrations:
+            yield i
+    
+    def __repr__(self):
+        return "Revision (depotFile = %s rev = %s change = %s action = %s type = %s time = %s user = %s client = %s)" % \
+            (self.depotFile, self.rev, self.change, self.action, self.type, self.time, self.user, self.client)
+#
+# Each DepotFile entry contains details about one depot file.
+#
+class DepotFile:
+    def __init__( self, name ):
+        self.depotFile = name
+        self.revisions = []
+    
+    def new_revision(self):
+        r = Revision( self.depotFile )
+        self.revisions.append( r )
+        return r
+    
+    def each_revision(self):
+        for r in self.revisions:
+            yield r
+
+    def str_revision(self, rev, revFormat, changeFormat):
+        result = "... #{rev:<{rf}} change {change:{cf}} {action:9} on {date}  by {user}@{client} ({type}) '{desc}'". \
+            format(rev = rev.rev, 
+                   rf = revFormat,
+                   change = rev.change,
+                   cf = changeFormat,
+                   action = rev.action, 
+                   date = rev.time, 
+                   user = rev.user, 
+                   client = rev.client, 
+                   type = rev.type,
+                   desc = rev.desc)
+        return result
+ 
+    def str_integration(self, integ):
+        result = "... ... {how} {file}#{srev},{erev}". \
+            format(how = integ.how, file = integ.file, srev = integ.srev, erev = integ.erev)
+        return result
+    
+    def __str__(self):
+        result = "{}".format(self.depotFile)
+        
+        revFormat = len(str(self.revisions[0].rev))
+        changeFormat = len(str(self.revisions[0].change))
+        
+        for rev in self.revisions:
+            result += "\n{}".format(self.str_revision(rev, revFormat, changeFormat))
+            for integ in rev.integrations:
+                result += "\n{}".format(self.str_integration(integ))
+        return result
+    
+    def __repr__(self):
+        return "DepotFile (depotFile = %s, %s revisions)" % ( self.depotFile, len( self.revisions ) )
+
+#
+# Resolver class used in p4.run_resolve()
+#
+# The default simply checks that p4.input is set to sensible value
+# This class is meant to be subclassed for a custom resolver and
+# Resolver.resolve() overriden
+#
+
+class Resolver:
+    def __init__(self):
+        pass
+    
+    def resolve(self, mergeInfo):
+        if mergeInfo.merge_hint == "e":
+            print("Standard resolver encountered merge conflict, skipping resolve")
+            return "s"
+        else:
+            return mergeInfo.merge_hint
+
+    def actionResolve(self, mergeInfo):
+        return mergeInfo.merge_hint
+
+#
+# OutputHandler base class
+#
+# Extend this class if you want to use the handler interface
+#
+
+class OutputHandler:
+    REPORT  = 0
+    HANDLED = 1
+    CANCEL  = 2
+    
+    def __init__(self):
+        pass
+    
+    def outputText(self, s):
+        return OutputHandler.REPORT
+    
+    def outputBinary(self, b):
+        return OutputHandler.REPORT
+    
+    def outputStat(self, h):
+        return OutputHandler.REPORT
+    
+    def outputInfo(self, i):
+        return OutputHandler.REPORT
+    
+    def outputMessage(self, e):
+        return OutputHandler.REPORT
+
+class ReportHandler( OutputHandler ):
+    def __init__(self):
+        OutputHandler.__init__(self)
+    
+    def outputText(self, s):
+        print( "text: ", s)
+        return OutputHandler.HANDLED
+    
+    def outputBinary(self, b):
+        print( "binary: ", b)
+        return OutputHandler.HANDLED
+    
+    def outputStat(self, h):
+        print( "stat:", h)
+        return OutputHandler.HANDLED
+    
+    def outputInfo(self, i):
+        print( "info: ", i)
+        return OutputHandler.HANDLED
+    
+    def outputMessage(self, e):
+        print( "error:", e)
+        return OutputHandler.HANDLED
+
+class Progress:
+    TYPE_SENDFILE = 1
+    TYPE_RECEIVEFILE = 2
+    TYPE_TRANSFER = 3
+    TYPE_COMPUTATION = 4
+    
+    UNIT_PERCENT = 1
+    UNIT_FILES = 2
+    UNIT_KBYTES = 3
+    UNIT_MBYTES = 4
+    
+    def __init__(self):
+        pass
+    
+    def init(self, type):
+        self.type = type
+    
+    def setDescription( self, description, units ):
+        self.description = description
+        self.units = units
+    
+    def setTotal( self, total ):
+        self.total = total
+    
+    def update( self, position ):
+        self.position = position
+    
+    def done( self, fail ):
+        pass
+
+class TextProgress(Progress):
+    TYPES = [ "Unknown", "Submit", "Sync" ]
+    UNITS = [ "Unknown", "Percent", "Files", "KBytes", "MBytes" ]
+    
+    def __init__(self):
+        Progress.__init__(self)
+    
+    def init(self, type):
+        Progress.init(self, type)
+        print( "Progress.init with '%s'" % self.TYPES[type] )
+    
+    def setDescription(self, description, units):
+        Progress.setDescription(self, description, units)
+        print( "Progress.setDescription with '%s' and units '%s'" % (description, self.UNITS[units]) )
+    
+    def setTotal( self, total ):
+        Progress.setTotal(self, total)
+        print( "Progress.setTotal with '%s' " % total )
+    
+    def update( self, position ):
+        Progress.update(self, position )
+        print( "Progress.update with '%s'" % position )
+    
+    def done( self, fail ):
+        Progress.done(self, fail)
+        print( "Progress.done with '%s"'' % fail )
+
+def processFilelog(h):
+    if "depotFile" in h:
+        df = DepotFile( h[ "depotFile" ] )
+        for n, rev in enumerate( h[ "rev" ]):
+            # Create a new revision of this file ready for populating
+            r = df.new_revision()
+            # Populate the base attributes of each revision
+            r.rev = int( rev )
+            r.change = int( h[ "change" ][ n ] )
+            r.action = h[ "action" ][ n ]
+            r.type = h[ "type" ][ n ]
+            r.time = datetime.datetime.utcfromtimestamp( int( h[ "time" ][ n ]) )
+            r.user = h[ "user" ][ n ]
+            r.client = h[ "client" ][ n ]
+            r.desc = h[ "desc" ][ n ]
+            if "digest" in h and n < len(h[ "digest" ]):
+                r.digest = h[ "digest" ][ n ]
+            if "fileSize" in h and n < len(h[ "fileSize" ]):
+                r.fileSize = h[ "fileSize" ][ n ]
+            
+            # Now if there are any integration records for this revision,
+            # add them in too
+            
+            if (not "how" in h) or (n >= len(h["how"]) or h["how"][n] == None):
+                continue
+            else:
+                for m, how in enumerate( h[ "how" ][ n ] ):
+                    file = h[ "file" ][ n ][ m ]
+                    srev = h[ "srev" ][ n ][ m ].lstrip('#')
+                    erev = h[ "erev" ][ n ][ m ].lstrip('#')
+                    
+                    if srev == "none":
+                        srev = 0
+                    else:
+                        srev = int( srev )
+                
+                    if erev == "none":
+                        erev = 0
+                    else:
+                        erev = int( erev )
+        
+                    r.integration( how, file, srev, erev )
+        return df
+    else:
+        raise Exception("Not a filelog object: " + h)
+
+class FilelogOutputHandler(OutputHandler):
+    def __init__(self):
+        OutputHandler.__init__(self)
+    
+    def outputStat(self, h):
+        df = processFilelog(h)
+        
+        return self.outputFilelog(df)
+    
+    def outputFilelog(self, f):
+        return OutputHandler.REPORT
+
+# This is where the C/C++ shared library is loaded
+# It has to be in this place because the library needs to access
+# the classes defined above. Accessing classes defined below this
+# entry would cause an endless loop
+
+import P4API
+
+class P4(P4API.P4Adapter):
+    """Use this class to communicate with a Perforce server
+        
+        Instances of P4 will use the environment settings (including P4CONFIG)
+        to determine the connection parameters such as P4CLIENT and P4PORT.
+        
+        This attributes can also be set separately before connecting.
+        
+        To run any Perforce commands, users of this class first need to run
+        the connect() method.
+        
+        It is good practice to disconnect() after the program is complete.
+        """
+    # Constants useful for exception_level
+    # RAISE_ALL:     Errors and Warnings are raised as exceptions (default)
+    # RAISE_ERROR:   Only Errors are raised as exceptions
+    # RAISE_NONE:    No exceptions are raised, instead False is returned
+    
+    RAISE_ALL    = 2
+    RAISE_ERROR  = 1
+    RAISE_ERRORS = 1
+    RAISE_NONE   = 0
+    
+    # Named values for generic error codes returned by
+    # P4API.Message.generic
+    
+    EV_NONE      = 0     # misc
+    
+    # The fault of the user
+    
+    EV_USAGE     = 0x01  # request not consistent with dox
+    EV_UNKNOWN   = 0x02  # using unknown entity
+    EV_CONTEXT   = 0x03  # using entity in wrong context
+    EV_ILLEGAL   = 0x04  # trying to do something you can't
+    EV_NOTYET    = 0x05  # something must be corrected first
+    EV_PROTECT   = 0x06  # protections prevented operation
+    
+    # No fault at all
+    
+    EV_EMPTY     = 0x11  # action returned empty results
+    
+    # not the fault of the user
+    
+    EV_FAULT     = 0x21  # inexplicable program fault
+    EV_CLIENT    = 0x22  # client side program errors
+    EV_ADMIN     = 0x23  # server administrative action required
+    EV_CONFIG    = 0x24  # client configuration inadequate
+    EV_UPGRADE   = 0x25  # client or server too old to interact
+    EV_COMM      = 0x26  # communications error
+    EV_TOOBIG    = 0x27  # not even Perforce can handle this much
+    
+    # Named values for error severities returned by
+    # P4API.Message.severity
+    E_EMPTY      = 0  # nothing yet
+    E_INFO       = 1  # something good happened
+    E_WARN       = 2  # something not good happened
+    E_FAILED     = 3  # user did something wrong
+    E_FATAL      = 4  # system broken -- nothing can continue
+    
+    # mappings for __iterate
+    # list-of-specs => ( name-of-one-spec, field-name-in-list-of-specs )
+    
+    specfields = {
+        'clients'   :   ('client', 'client'),
+        'labels'    :   ('label', 'label'),
+        'branches'  :   ('branch', 'branch'),
+        'changes'   :   ('change', 'change'),
+        'streams'   :   ('stream', 'Stream'),
+        'jobs'      :   ('job', 'Job'),
+        'users'     :   ('user', 'User'),
+        'groups'    :   ('group', 'group'),
+        'depots'    :   ('depot', 'name'),
+        'servers'   :   ('server', 'Name')
+    }
+    
+    def __init__(self, *args, **kwlist):
+        P4API.P4Adapter.__init__(self, *args, **kwlist)
+    
+    def __del__(self):
+        if self.debug > 3:
+            print("P4.__del__()", file=sys.stderr)
+    
+    def __getattr__(self, name):
+        if name.startswith("run_"):
+            cmd = name[len("run_"):]
+            return lambda *args, **kargs: self.run(cmd, *args, **kargs)
+        elif name.startswith("delete_"):
+            cmd = name[len("delete_"):]
+            return lambda *args, **kargs: self.run(cmd, "-d", *args, **kargs)
+        elif name.startswith("fetch_"):
+            cmd = name[len("fetch_"):]
+            return lambda *args, **kargs: self.__fetch(cmd, *args, **kargs)
+        elif name.startswith("save_"):
+            cmd = name[len("save_"):]
+            return lambda *args, **kargs: self.__save(cmd, *args, **kargs)
+        elif name.startswith("parse_"):
+            cmd = name[len("parse_"):]
+            return lambda *args, **kargs: self.__parse_spec(cmd, *args, **kargs)
+        elif name.startswith("format_"):
+            cmd = name[len("format_"):]
+            return lambda *args, **kargs: self.__format_spec(cmd, *args, **kargs)
+        elif name.startswith("iterate_"):
+            cmd = name[len("iterate_"):]
+            return lambda *args, **kargs: self.__iterate(cmd, *args, **kargs)
+        else:
+            raise AttributeError(name)
+    
+    def __save(self, cmd, *args, **kargs):
+        self.input = args[0]
+        return self.run(cmd, "-i", args[1:], **kargs)
+    
+    def __parse_spec(self, cmd, *args, **kargs):
+        form = args[0]
+        comments = "\n".join( [ x for x in form.split('\n') if x.startswith('#') ] ) + "\n"
+        spec = self.parse_spec(cmd, *args, **kargs)
+        spec.__dict__['comment'] = comments
+        return spec
+    
+    def __format_spec(self, cmd, *args, **kargs):
+        spec = args[0]
+        form = self.format_spec(cmd, *args, **kargs)
+        if 'comment' in spec.__dict__:
+            form = spec.__dict__['comment'] + "\n" +  form
+        return form
+    
+    def __fetch(self, cmd, *args, **kargs):
+        result = self.run(cmd, "-o", *args, **kargs)
+        for r in result:
+            if isinstance(r, tuple) or isinstance(r,dict):
+                return r 
+        return result[0]
+    
+    def __iterate(self, cmd, *args, **kargs):
+        
+        if cmd in self.specfields:
+            specs = self.run(cmd, *args, **kargs)
+            spec = self.specfields[cmd][0]
+            field = self.specfields[cmd][1]
+            
+            # Return a generators (Python iterator object)
+            # On iteration, this will retrieve one spec at a time
+            return ( self.run(spec, '-o', x[field])[0] for x in specs )
+        else:
+            raise Exception('Unknown spec list command: %s', cmd)
+    
+    def __repr__(self):
+        state = "disconnected"
+        if self.connected():
+            state = "connected"
+        
+        return "P4 [%s@%s %s] %s" % \
+            (self.user, self.client, self.port, state)
+    
+    def identify(cls):
+        return P4API.identify()
+    identify = classmethod(identify)
+    
+    def run(self, *args, **kargs):
+        "Generic run method"
+        context = {}
+        resultLogging = True
+        
+        if "resultLogging" in kargs:
+            resultLogging= False
+            del kargs["resultLogging"]
+        
+        for (k,v) in list(kargs.items()):
+            context[k] = getattr(self, k)
+            setattr(self, k, v)
+                
+        flatArgs = self.__flatten(args)
+        if self.logger:
+            self.logger.info("p4 " + " ".join(flatArgs))
+        
+        result = P4API.P4Adapter.run(self, *flatArgs)
+        
+        if resultLogging and self.logger:
+            self.logger.debug(result)
+
+        for (k,v) in list(context.items()):
+            setattr( self, k, v)
+                    
+        return result
+    
+    def run_submit(self, *args, **kargs):
+        "Simplified submit - if any arguments is a dict, assume it to be the changeform"
+        nargs = list(args)
+        form = None
+        for n, arg in enumerate(nargs):
+            if isinstance( arg, dict):
+                self.input = arg
+                nargs.pop(n)
+                nargs.append("-i")
+                break
+        return self.run("submit", *nargs, **kargs)
+    
+    def run_shelve(self, *args, **kargs):
+        "Simplified shelve - if any arguments is a dict, assume it to be the changeform"
+        nargs = list(args)
+        form = None
+        for n, arg in enumerate(nargs):
+            if isinstance( arg, dict):
+                self.input = arg
+                nargs.pop(n)
+                nargs.append("-i")
+                break
+        return self.run("shelve", *nargs, **kargs)
+    
+    def delete_shelve(self, *args, **kargs):
+        "Simplified deletion of shelves - if no -c is passed in, add it to the args"
+        nargs = list(args)
+        if '-c' not in nargs:
+            nargs = ['-c'] + nargs # prepend -c if it is not there
+        nargs = ['-d'] + nargs
+        return self.run("shelve", *nargs, **kargs)
+    
+    def run_login(self, *args):
+        "Simple interface to make login easier"
+        self.input = self.password
+        return self.run("login", *args)
+    
+    def run_password( self, oldpass, newpass ):
+        "Simple interface to allow setting of the password"
+        if( oldpass and len(oldpass) > 0 ):
+            self.input = [ oldpass, newpass, newpass ]
+        else:
+            self.input = [ newpass, newpass ]
+        
+        try:
+            return self.run( "password" )
+        except P4Exception as e:
+            if self.errors and self.errors[0] == "Passwords don't match.":
+                raise P4Exception("Password invalid.")
+
+
+    #
+    # run_filelog: convert "p4 filelog" responses into objects with useful
+    #              methods
+    #
+    # Requires tagged output to be of any real use. If tagged output it not
+    # enabled then you just get the raw data back
+    #
+    def run_filelog( self, *args, **kargs ):
+        kargs["resultLogging"] = False
+        raw = self.run( 'filelog', args, **kargs )
+        if (not self.tagged or not raw):
+            # untagged mode returns simple strings, which breaks the code below
+            # raw could be None if a handler is used
+            return raw
+        result = []
+        for h in raw:
+            df = None
+            if isinstance( h, dict ):
+                df = processFilelog( h )
+            else:
+                df = h
+            result.append( df )
+        
+        # if logger is in the kargs, P4.run has changed its context but already reset it
+        # so we need to pull out the logger again for this special case
+        # run_filelog has its own result debug output because of the post-processing step
+        
+        logger = self.logger
+        if "logger" in kargs:
+            logger = kargs["logger"]
+        
+        if logger:
+            output = "\n\n".join([ str(x) for x in result ])
+            logger.debug(output)
+        
+        return result
+
+    def run_print(self, *args, **kargs):
+        kargs["resultLogging"] = False
+        raw = self.run('print', args, **kargs)
+
+        logger = self.logger
+        if "logger" in kargs:
+            logger = kargs["logger"]
+
+        result = []
+        if raw:
+            debugResult = []
+            for line in raw:
+                if isinstance(line, dict):
+                    result.append(line)
+                    if logger:
+                        debugResult.append(line)
+                    result.append("")
+                else:
+                    # to support encoding for Python 3, we have to do a little dance
+                    # we cannot add bytes to the str "", but we expect that all lines
+                    # are either str or bytes. So if we encounter bytes, we replace the content
+                    try:
+                        result[-1] += line
+                    except TypeError:
+                        if type(line) == bytes and type(result[-1]) == str and result[-1] == "":
+                            result[-1] = line
+                        else:
+                            raise
+            if logger:
+                logger.debug(debugResult)
+            return result
+        else:
+            return []
+
+    def run_resolve(self, *args, **kargs):
+        if self.resolver:
+            myResolver = self.resolver
+        else:
+            myResolver = Resolver()
+        if "resolver" in kargs:
+            myResolver = kargs["resolver"]
+        
+        savedResolver = self.resolver
+        self.resolver = myResolver
+        result = self.run("resolve", args)
+        self.resolver = savedResolver
+        
+        return result
+
+    def run_tickets(self, *args):
+        fname = self.ticket_file
+        with open(fname) as f:
+            tickets_raw = f.readlines()
+        pattern = re.compile('([^=]*)=(.*):([^:]*)\n')
+        tickets = [ pattern.match(x).groups() for x in tickets_raw ]
+        keys = [ "Host", "User", "Ticket" ]
+        result = [ dict(zip(keys, x)) for x in tickets ]
+        return result
+    
+    def run_init(self, *args, **kargs):
+        raise Exception("Please run P4.init() instead")
+        
+    def run_clone(self, *args, **kargs):
+        raise Exception("Please run P4.clone) instead")
+    
+    def __flatten(self, args):
+        result = []
+        if isinstance(args, tuple) or isinstance(args, list):
+            for i in args:
+                result.extend(self.__flatten(i))
+        else:
+            result.append(args)
+        return tuple(result)
+
+    def __enter__( self ):
+        return self
+    
+    def __exit__( self, exc_type, exc_val, exc_tb ):
+        if self.connected():
+            self.disconnect()
+        return False
+    
+    def connect( self ):
+        P4API.P4Adapter.connect( self )
+        return self
+    
+    def is_ignored( self, path ):
+        return P4API.P4Adapter.is_ignored( self, os.path.abspath(path) )
+    
+    @contextmanager
+    def while_tagged( self, t ):
+        old = self.tagged
+        self.tagged = t
+        yield
+        self.tagged = old
+    
+    @contextmanager
+    def at_exception_level( self, e ):
+        old = self.exception_level
+        self.exception_level = e
+        yield
+        self.exception_level = old
+    
+    @contextmanager
+    def using_handler( self, c ):
+        old = self.handler
+        self.handler = c
+        yield
+        self.handler = old
+    
+    @contextmanager
+    def saved_context( self , **kargs):
+        """Saves the context of this p4 object and restores it again at the end of the block"""
+        
+        saved_context = {}
+        for attr in self.__members__:
+            saved_context[attr] = getattr(self, attr)
+        
+        for (k,v) in list(kargs.items()):
+            setattr( self, k, v)
+        
+        yield
+        
+        # now restore the context again. Ignore AttributeError exception
+        # Exception is expected because some attributes only have getters, no setters
+        
+        for (k,v) in list(saved_context.items()):
+            if k not in ("port", "track"):
+                try:
+                    setattr( self, k, v )
+                except AttributeError:
+                    pass # expected for server_level and p4config_file
+
+    @contextmanager
+    def temp_client( self, prefix, template ):
+        """Creates a temporary workspace with a temporary root. 
+            To be used with the "with" statement. Will clean up the temporary root and client
+            workspace after the block has finished.
+            The prefix is prepended to the workspace name and should be used in conjunction with
+            the SpecMap of a spec depot to avoid creating entries there.
+        """
+
+        name = "{prefix}_{id}".format(prefix=prefix, id=str(uuid.uuid1()))
+        ws = self.fetch_client('-t', template, name)
+
+        try:
+            root = tempfile.mkdtemp(prefix=prefix)
+            ws._root = root
+            self.save_client(ws)
+
+            oldName = self.client
+            self.client = name
+            
+            yield ws
+
+            self.client = oldName
+        finally:
+            self.delete_client(name)
+            shutil.rmtree(root)
+
+class Map(P4API.P4Map):
+    def __init__(self, *args):
+        P4API.P4Map.__init__(self, *args)
+        if len(args) > 0:
+            self.insert( *args )
+    
+    LEFT2RIGHT = True
+    RIGHT2LEFT = False
+    
+    def __str__( self ):
+        result = ""
+        
+        for a in self.as_array():
+            result += a + "\n"
+        
+        return result
+    
+    def is_empty(self):
+        """Returns True if this map has no entries yet, otherwise False"""
+        
+        return self.count() == 0
+    
+    def includes(self, *args):
+        return self.translate(*args) != None
+    
+    def reverse(self):
+        return Map(P4API.P4Map.reverse(self).as_array())
+    
+    def insert(self, *args):
+        """Insert an argument to the map. The argument can be:
+            
+            A String,
+            Either of the form "[+-]//lhs/... //rhs/..." or "[+-]//lhs/..."
+            for label style maps.
+            A List:
+            This is a list of strings of one of the single string formats
+            described above.
+            A pair of Strings:
+            P4.Map.insert(lhs, rhs)
+            """
+        
+        if len(args) == 1 :
+            arg = args[0]
+            if isinstance( arg, str ):
+                P4API.P4Map.insert( self, arg )
+            elif isinstance( arg, list ):
+                for s in arg:
+                    P4API.P4Map.insert( self, s )
+        
+        else: # expecting 2 args in this case: left, right
+            left = args[0].strip()
+            right = args[1].strip()
+            P4API.P4Map.insert(self, left, right )
+
+
+def init(*args, **kargs):
+    return __run_dvcs("init", args, **kargs)
+
+def clone(*args, **kargs):
+    return __run_dvcs("clone", args, **kargs)
+
+def __run_dvcs(cmd, *args, **kargs):
+    __check_paths()
+    options = []
+    
+    def add_option(options, name, opt):
+        if name in kargs:
+            options += [opt, kargs[name]]
+    
+    add_option(options, "client","-c")
+    add_option(options, "directory","-d")
+    add_option(options, "user", "-u")
+        
+    named_args = []
+    
+    if "unicode" in kargs:
+        unicode = kargs["unicode"]
+        del kargs["unicode"]
+        if unicode:
+            named_args.append("-xi")
+        else:
+            named_args.append("-n")
+    
+    if "casesensitive" in kargs:
+        casesensitive = kargs["casesensitive"]
+        del kargs["casesensitive"]
+        if casesensitive:
+            named_args.append("-C1")
+        else:
+            named_args.append("-C0")
+    
+    arguments = [ "p4" ]
+    arguments += options
+    arguments += [ cmd ]
+    arguments += named_args 
+    for a in args:
+        arguments += a
+    
+    p = subprocess.Popen(arguments, stdout=subprocess.PIPE, stderr=subprocess.PIPE)
+    (sout,serr) = (p.stdout, p.stderr)
+    results = sout.read()
+    errors = serr.read()
+    sout.close()
+    serr.close()
+    
+    if errors:
+        if type(errors) == bytes:
+            errors = errors.decode('UTF-8')
+        errors = errors.strip()
+        
+        raise Exception("Cmd '{}' raised\n{}".format(" ".join(arguments), errors))
+    
+    if type(results) == bytes:
+        results = results.decode('UTF-8')
+    
+    # this is where we analyze the outcome of the command
+    # move into the resulting directory (if specified)
+    # then create a P4 object and return that (connected??)
+    
+    results = results.strip()
+    
+    if "directory" in kargs:
+        path = os.path.abspath(kargs["directory"])
+        os.chdir(path)
+        os.environ["PWD"] = path
+        
+    # create a copy of kargs without the "directory" key - make it compatible with Python 2.6
+    # new_kargs = { x:kargs[x] for x in kargs if x != "directory" }
+    new_kargs = dict((x,kargs[x]) for x in kargs if x != "directory" )
+    
+    return P4(**new_kargs)
+    
+def __check_paths():
+    if not __exec_exists("p4"):
+        raise Exception("P4 executable not in path")
+    if not __exec_exists("p4d"):
+        raise Exception("P4D executable not in path")
+    
+def __exec_exists(name):
+    execName = name
+    if platform.system() == "Windows":
+        execName += ".exe"
+    
+    for p in os.environ["PATH"].split(os.pathsep):
+        pathToFile = os.path.join(p, execName)
+        if os.path.exists(pathToFile) and os.access(pathToFile, os.X_OK):
+            return __check_version(pathToFile)
+    return False
+    
+def __check_version(pathToFile):
+    p = subprocess.Popen([pathToFile,'-V'], stdout=subprocess.PIPE, stderr=subprocess.PIPE)
+    (sout,serr) = (p.stdout, p.stderr)
+    
+    output = sout.read()
+    error = serr.read()
+
+    sout.close()
+    serr.close()
+    
+    if type(output) == bytes: # Python 3 output is bytes, not str
+        output = output.decode('UTF-8')
+    chunks = output.split(os.linesep)
+    
+    pattern = re.compile("Rev. (?P<Program>.+)/(?P<Platform>.+)/(?P<Release>.+)/(?P<Patch>\d+) \((\d+/\d+/\d+)\).")
+    
+    for c in chunks:
+        match = pattern.match(c)
+        if match:
+            version = match.group('Release')
+            year = int(version.split('.')[0])
+            if year >= 2015:
+                return True
+            else:
+                program = match.group('Program')
+                raise Exception("{} must be at least 2015.1, not {}".format(program, version))
+
+    raise Exception("Unknown P4 output : {}".format(output) )
+
+if __name__ == "__main__":
+    p4 = P4()
+    p4.connect()
+    try:
+        ret = p4.run(sys.argv[1:])
+        for line in ret:
+            if isinstance(line, dict):
+                print("-----")
+                for k in list(line.keys()):
+                    print(k, "=", line[k])
+            else:
+                print(line)
+    except:
+        for e in p4.errors:
+            print(e)
diff --git a/Engine/Extras/Maya_AnimationRiggingTools/ARTv1/MayaTools/General/Scripts/Modules/uExport/__init__.py b/Engine/Extras/Maya_AnimationRiggingTools/ARTv1/MayaTools/General/Scripts/Modules/uExport/__init__.py
new file mode 100644
--- /dev/null
+++ b/Engine/Extras/Maya_AnimationRiggingTools/ARTv1/MayaTools/General/Scripts/Modules/uExport/__init__.py
@@ -0,0 +1 @@
+from uExport import *
\ No newline at end of file
diff --git a/Engine/Extras/Maya_AnimationRiggingTools/ARTv1/MayaTools/General/Scripts/Modules/uExport/uExport.py b/Engine/Extras/Maya_AnimationRiggingTools/ARTv1/MayaTools/General/Scripts/Modules/uExport/uExport.py
new file mode 100644
--- /dev/null
+++ b/Engine/Extras/Maya_AnimationRiggingTools/ARTv1/MayaTools/General/Scripts/Modules/uExport/uExport.py
@@ -0,0 +1,580 @@
+'''
+uExport
+Christopher Evans, Version 0.1, Feb 2014
+@author = Chris Evans
+version = 0.2
+
+Add this to a shelf:
+import uExport as ue
+uExportToolWindow = ue.show()
+
+'''
+
+import os
+
+import shiboken
+import time
+from PySide import QtGui, QtCore
+from cStringIO import StringIO
+import xml.etree.ElementTree as xml
+import pysideuic
+
+import maya.cmds as cmds
+import maya.OpenMayaUI as openMayaUI
+import maya.mel as mel
+
+def show():
+    global uExportToolWindow
+    try:
+        uExportToolWindow.close()
+    except:
+        pass
+
+    uExportToolWindow = uExportTool()
+    uExportToolWindow.show()
+    return uExportToolWindow
+
+def loadUiType(uiFile):
+    """
+    Pyside lacks the "loadUiType" command, so we have to convert the ui file to py code in-memory first
+    and then execute it in a special frame to retrieve the form_class.
+    http://tech-artists.org/forum/showthread.php?3035-PySide-in-Maya-2013
+    """
+    parsed = xml.parse(uiFile)
+    widget_class = parsed.find('widget').get('class')
+    form_class = parsed.find('class').text
+
+    with open(uiFile, 'r') as f:
+        o = StringIO()
+        frame = {}
+
+        pysideuic.compileUi(f, o, indent=0)
+        pyc = compile(o.getvalue(), '<string>', 'exec')
+        exec pyc in frame
+
+        #Fetch the base_class and form class based on their type in the xml from designer
+        form_class = frame['Ui_%s'%form_class]
+        base_class = eval('QtGui.%s'%widget_class)
+    return form_class, base_class
+
+def getMayaWindow():
+    ptr = openMayaUI.MQtUtil.mainWindow()
+    if ptr is not None:
+        return shiboken.wrapInstance(long(ptr), QtGui.QWidget)
+
+try:
+    selfDirectory = os.path.dirname(__file__)
+    uiFile = selfDirectory + '/uExport.ui'
+except:
+    uiFile = 'D:\\Build\\usr\\jeremy_ernst\\MayaTools\\General\\Scripts\\Modules\\uExport\\uExport.ui'
+    form_class, base_class = loadUiType(uiFile)
+
+if os.path.isfile(uiFile):
+    form_class, base_class = loadUiType(uiFile)
+else:
+    cmds.error('Cannot find UI file: ' + uiFile)
+
+#this is used by both classes below, this method is usually in the utils lib at Epic
+def attrExists(attr):
+    if '.' in attr:
+        node, att = attr.split('.')
+        return cmds.attributeQuery(att, node=node, ex=1)
+    else:
+        cmds.warning('attrExists: No attr passed in: ' + attr)
+        return False
+
+class uExport():
+    '''
+    Just a little basket to store things.
+    TODO: Add properties to get/set values
+    '''
+    def __init__(self, node):
+        self.meshes = cmds.listConnections(node + '.rendermesh')
+        self.root = cmds.listConnections(node + '.export_root')
+
+        #joints
+        self.joints = [self.root]
+        children = cmds.listRelatives(self.root, type='joint',allDescendents=True)
+        if children:
+            self.joints.extend(children)
+
+        self.version = cmds.getAttr(node + '.uexport_ver')
+        self.node = node
+        self.name = node.split('|')[-1]
+        self.asset_name = node
+        self.folder_path = None
+        if attrExists(node + '.asset_name'):
+            self.asset_name = cmds.getAttr(node + '.asset_name')
+        if attrExists(node + '.fbx_name'):
+            self.fbx_name = cmds.getAttr(node + '.fbx_name')
+        if attrExists(node + '.folder_path'):
+            self.folder_path = cmds.getAttr(node + '.folder_path')
+
+
+########################################################################
+## UEXPORT TOOL
+########################################################################
+
+
+class uExportTool(base_class, form_class):
+    title = 'uExportTool 0.2'
+
+    currentMesh = None
+    currentSkin = None
+    currentInf = None
+    currentVerts = None
+    currentNormalization = None
+
+    scriptJobNum = None
+    copyCache = None
+
+    jointLoc = None
+
+    iconLib = {}
+    iconPath = os.environ.get('MAYA_LOCATION', None) + '/icons/'
+    iconLib['joint'] = QtGui.QIcon(QtGui.QPixmap(iconPath + 'kinJoint.png'))
+    iconLib['ikHandle'] = QtGui.QIcon(QtGui.QPixmap(iconPath + 'kinHandle.png'))
+    iconLib['transform'] = QtGui.QIcon(QtGui.QPixmap(iconPath + 'orientJoint.png'))
+
+    def __init__(self, parent=getMayaWindow()):
+        self.closeExistingWindow()
+        super(uExportTool, self).__init__(parent)
+
+        self.setupUi(self)
+        self.setWindowTitle(self.title)
+
+        wName = openMayaUI.MQtUtil.fullName(long(shiboken.getCppPointer(self)[0]))
+
+        ## Connect UI
+        ########################################################################
+        self.connect(self.export_BTN, QtCore.SIGNAL("clicked()"), self.export_FN)
+        self.connect(self.createUexportNode_BTN, QtCore.SIGNAL("clicked()"), self.createUexportNode_FN)
+        self.connect(self.replaceUnknownNodes, QtCore.SIGNAL("clicked()"), self.replaceUnknownNodes_FN)
+        self.refreshBTN.clicked.connect(self.refreshUI)
+
+        #context menu
+        self.export_tree.setContextMenuPolicy(QtCore.Qt.CustomContextMenu)
+        self.connect(self.export_tree, QtCore.SIGNAL("customContextMenuRequested(QPoint)" ), self.openMenu)
+
+        self.export_tree.itemClicked.connect(self.check_status)
+
+        self.snapRoot_CMB.setHidden(True)
+        self.refreshUI()
+
+## GENERAL
+########################################################################
+
+    def closeExistingWindow(self):
+        for qt in QtGui.qApp.topLevelWidgets():
+            try:
+                if qt.__class__.__name__ == self.__class__.__name__:
+                    qt.deleteLater()
+                    print 'uExport: Closed ' + str(qt.__class__.__name__)
+            except:
+                pass
+
+    def convertSkelSettingsToNN(delete=1):
+        orig = 'SkeletonSettings_Cache'
+        if cmds.objExists(orig):
+            if cmds.nodeType(orig) == 'unknown':
+                print orig
+                new = cmds.createNode('network')
+                for att in cmds.listAttr(orig):
+                    if not cmds.attributeQuery(att, node=new, exists=1):
+                        typ = cmds.attributeQuery(att, node=orig, at=1)
+                        if typ == 'typed':
+                            cmds.addAttr(new, longName=att, dt='string')
+                            if cmds.getAttr(orig + '.' + att):
+                                cmds.setAttr(new + '.' + att, cmds.getAttr(orig + '.' + att), type='string')
+                        elif typ == 'enum':
+                            cmds.addAttr(new, longName=att, at='enum', enumName=cmds.attributeQuery(att, node=orig, listEnum=1)[0])
+                cmds.delete(orig)
+                cmds.rename(new, 'SkeletonSettings_Cache')
+
+    def replaceUnknownNodes_FN(self):
+        self.convertSkelSettingsToNN()
+
+## UI RELATED
+########################################################################
+
+    def openMenu(self, position):
+        menu = QtGui.QMenu()
+        rendermeshRewire = menu.addAction("Re-wire rendermesh attr to current mesh selection")
+        rootRewire = menu.addAction("Re-wire root joint attr to current selected joint")
+        pos = self.export_tree.mapToGlobal(position)
+        action = menu.exec_(pos)
+        if action == rendermeshRewire:
+            for index in self.export_tree.selectedIndexes():
+                uExportNode = self.export_tree.itemFromIndex(index).uExport.node
+                meshes = cmds.ls(sl=1)
+                #TODO: check if theyre actually meshes
+                self.connectRenderMeshes(uExportNode, meshes)
+    	elif action == rootRewire:
+            index = self.export_tree.selectedIndexes()[0]
+            uExportNode = self.export_tree.itemFromIndex(index).uExport.node
+            root = cmds.ls(sl=1)
+            if len(root) == 1:
+                self.connectRoot(uExportNode, root[0])
+            else:
+                cmds.error('Select a single joint, you == fail, bro. ' + str(root))
+        self.refreshUI()
+
+    def refreshUI(self):
+        self.export_tree.clear()
+        self.uNodes = []
+        for node in self.getExportNodes():
+            self.uNodes.append(uExport(node))
+        self.buildExportTree(self.uNodes)
+
+
+    def buildExportTree(self, uNodes):
+        for uNode in uNodes:
+
+            #top level
+            wid1 = QtGui.QTreeWidgetItem()
+            font = wid1.font(0)
+            font.setPointSize(15)
+
+            wid1.setText(0,uNode.asset_name)
+            wid1.uExport = uNode
+
+            wid1.setText(1, uNode.version)
+            self.export_tree.addTopLevelItem(wid1)
+            wid1.setExpanded(True)
+            wid1.setFont(0,font)
+
+            font = wid1.font(0)
+            font.setPointSize(10)
+
+            wid1.setFlags(QtCore.Qt.ItemIsEnabled | QtCore.Qt.ItemIsEditable | QtCore.Qt.ItemIsUserCheckable | QtCore.Qt.ItemIsSelectable)
+            wid1.setCheckState(0, QtCore.Qt.Checked)
+
+            #mesh branch
+            meshTop = QtGui.QTreeWidgetItem()
+            meshes = uNode.meshes
+            if meshes:
+                meshTop.setText(0, 'RENDER MESHES: (' + str(len(uNode.meshes)) + ')')
+            else:
+                meshTop.setText(0, 'RENDER MESHES: NONE')
+            meshTop.setFlags(QtCore.Qt.ItemIsEnabled | QtCore.Qt.ItemIsEditable | QtCore.Qt.ItemIsUserCheckable)
+            meshTop.setCheckState(0, QtCore.Qt.Checked)
+            wid1.addChild(meshTop)
+
+            #meshes
+            if uNode.meshes:
+                for mesh in uNode.meshes:
+                    meshWid = QtGui.QTreeWidgetItem()
+                    meshWid.setText(0, mesh)
+                    meshTop.addChild(meshWid)
+
+            if uNode.root:
+                #anim branch
+                animTop = QtGui.QTreeWidgetItem()
+                jnts = uNode.joints
+                if jnts:
+                    animTop.setText(0, 'ANIMATION:  (' + str(len(uNode.joints)) + ' JOINTS)')
+                else:
+                    animTop.setText(0, 'ANIMATION: NONE')
+                animTop.setFlags(QtCore.Qt.ItemIsEnabled | QtCore.Qt.ItemIsEditable | QtCore.Qt.ItemIsUserCheckable)
+                animTop.setCheckState(0, QtCore.Qt.Checked)
+                wid1.addChild(animTop)
+
+                #anims
+                animWid = QtGui.QTreeWidgetItem()
+                animWid.setText(0, '<< CURRENT TIME RANGE >>')
+                animTop.addChild(animWid)
+
+            #meta branch
+            metaTop = QtGui.QTreeWidgetItem()
+            metaTop.setText(0, 'METADATA')
+            wid1.addChild(metaTop)
+
+            #meta
+            fpathWid = QtGui.QTreeWidgetItem()
+            fpathWid.setText(0, 'FOLDER_PATH:   ' + str(uNode.folder_path))
+            metaTop.addChild(fpathWid)
+            #meta
+            fbxName = QtGui.QTreeWidgetItem()
+            fbxName.setText(0, 'FBX_FILE:   ' + str(uNode.fbx_name))
+            metaTop.addChild(fbxName)
+
+    def createUexportNode_FN(self):
+        if cmds.ls(sl=1):
+            if self.useRoot_CHK.isChecked():
+                rootName = self.rootName_CMB.currentText()
+                self.create(renderMeshes=cmds.ls(sl=1), rootJoint=rootName)
+            else:
+                #TODO: modal picker with joint filter
+                pass
+
+    def export_FN(self):
+        #get what should be export in an exportTask fn
+        exportWidgets = self.getExportNodeWidgets()
+        for wid in exportWidgets:
+            snapConst = None
+            if self.snapRoot_CHK.isChecked():
+                node = self.snapRoot_CMB.currentText()
+                try:
+                    if cmds.objExists(node):
+                        snapConst = cmds.parentConstraint(node, wid.uExport.root)
+                    else:
+                        cmds.error('Unable to find node: ' + node)
+                except:
+                    cmds.warning('Could not constrain root' + wid.uExport.root + ', is it already constrained?')
+            #export
+            mpath = cmds.file(sceneName=1, q=1)
+            fname = mpath.split('/')[-1]
+            fpath = mpath.replace(fname,'').replace('/','\\')
+            if wid.uExport.fbx_name:
+                if wid.uExport.fbx_name != '':
+                    fname = wid.uExport.fbx_name
+            if wid.uExport.folder_path:
+                if wid.uExport.folder_path != '':
+                    fpath = wid.uExport.folder_path
+
+            #prompt user for save location
+            userPath = QtGui.QFileDialog.getSaveFileName(caption = 'Export ' + wid.uExport.name + ' to:', filter='FBX Files (*.fbx)', dir=fpath + '\\' + fname)[0]
+
+            if userPath:
+                wid.uExport.fbx_name = userPath.split('/')[-1]
+                wid.uExport.folder_path = userPath.replace(wid.uExport.fbx_name,'')
+
+                #if set to save loc, save it
+                if self.rememberSaveCHK.isChecked():
+                    cmds.setAttr(wid.uExport.name + '.folder_path', wid.uExport.folder_path, type='string')
+                    cmds.setAttr(wid.uExport.name + '.fbx_name', wid.uExport.fbx_name, type='string')
+
+                meshChk = 1
+                animChk = 1
+                for c in range(0, wid.childCount()):
+                    if 'MESH' in wid.child(c).text(0):
+                        if wid.child(c).checkState(0) == QtCore.Qt.Checked:
+                            pass
+                        else:
+                            meshChk = 0
+
+                start = time.time()
+                #put export into the uExport class later
+                self.export(wid.uExport.name, path=userPath, mesh=meshChk)
+                elapsed = (time.time() - start)
+                print 'uExport>>> Exported ', wid.text(0), 'to', userPath ,'in %.2f seconds.' % elapsed
+
+            else:
+                cmds.warning('Invalid path specified: [' + str(userPath) + ']')
+            #cleanup constraint
+            if snapConst: cmds.delete(snapConst)
+
+    def getExportNodeWidgets(self):
+        nodes = []
+        for i in range(0, self.export_tree.topLevelItemCount()):
+            if self.export_tree.topLevelItem(i).checkState(0) == QtCore.Qt.Checked:
+                nodes.append(self.export_tree.topLevelItem(i))
+        return nodes
+
+
+    def check_status(self):
+        for i in range(0, self.export_tree.topLevelItemCount()):
+            if self.export_tree.topLevelItem(i).checkState(0) == QtCore.Qt.Unchecked:
+                for c in range(0, self.export_tree.topLevelItem(i).childCount()):
+                    self.export_tree.topLevelItem(i).child(c).setCheckState(0,QtCore.Qt.Unchecked)
+
+## P4 CRAP
+
+## UEXPORT NODE
+########################################################################
+    def getExportNodes(self):
+        return cmds.ls('*.uexport_ver', o=1, r=1)
+
+    def connectRenderMeshes(self, uNode, renderMeshes, rewire=1):
+        try:
+            cmds.undoInfo(openChunk=True)
+            if rewire:
+                for conn in cmds.listConnections(uNode + '.rendermesh', plugs=1, destination=1):
+                    cmds.disconnectAttr(uNode + '.rendermesh', conn)
+        
+            for node in renderMeshes:
+                if not attrExists(node+'.export'):
+                    cmds.addAttr(node, longName='export', attributeType='message')
+                cmds.connectAttr(uNode + '.rendermesh', node + '.export' )
+        except Exception as e:
+            print e
+        finally:
+            cmds.undoInfo(closeChunk=True)
+
+    def connectRoot(self, uNode, root, rewire=1):
+        try:
+            cmds.undoInfo(openChunk=True)
+            if rewire:
+                conns = cmds.listConnections(uNode + '.export_root', plugs=1, source=1)
+                if conns:
+                    for conn in conns:
+                        cmds.disconnectAttr(conn, uNode + '.export_root')
+        
+            if not attrExists(root+'.export'):
+                cmds.addAttr(root, longName='export', attributeType='message')
+            cmds.connectAttr(root + '.export', uNode + '.export_root' )
+        except Exception as e:
+            print e
+        finally:
+            cmds.undoInfo(closeChunk=True)
+        
+    def create(self, renderMeshes=None, rootJoint=None, strName='uExport'):
+        uExport = None
+        if cmds.objExists(strName):
+            #later re-hook up
+            #set uExport
+            pass
+        else:
+            uExport = cmds.group(em=1, name=strName)
+            cmds.addAttr(uExport, ln='rendermesh', at='message')
+            cmds.addAttr(uExport, ln='export_root', at='message')
+            cmds.addAttr(uExport, ln='materials', at='message')
+            cmds.addAttr(uExport, ln='uexport_ver', dt='string')
+            cmds.setAttr(uExport + '.uexport_ver', '1.0', type='string')
+            cmds.addAttr(uExport, ln='folder_path', dt='string')
+            cmds.addAttr(uExport, ln='asset_name', dt='string')
+            cmds.addAttr(uExport, ln='fbx_name', dt='string')
+
+        if uExport:
+            try:
+                if renderMeshes:
+                    self.connectRenderMeshes(uExport, renderMeshes, rewire=0)
+
+                #rootJoint
+                if rootJoint:
+                    if not attrExists(rootJoint+'.export'):
+                        cmds.addAttr(rootJoint, longName='export', attributeType='message')
+                    cmds.connectAttr(rootJoint + '.export', uExport + '.export_root')
+                else:
+                    cmds.warning('No root joint or could not find root: ' + str(rootJoint))
+            except Exception as e:
+                print cmds.warning(e)
+
+
+
+## EXPORT
+########################################################################
+
+    #TODO: Find and export blendshape meshes!
+    def setExportFlags(self):
+
+        # Mesh
+        mel.eval("FBXExportSmoothingGroups -v true")
+        mel.eval("FBXExportHardEdges -v false")
+        mel.eval("FBXExportTangents -v false")
+        mel.eval("FBXExportInstances -v false")
+        mel.eval("FBXExportInAscii -v true")
+        mel.eval("FBXExportSmoothMesh -v false")
+
+        # Animation
+        mel.eval("FBXExportBakeComplexAnimation -v true")
+        mel.eval("FBXExportBakeComplexStart -v "+str(cmds.playbackOptions(minTime=1, q=1)))
+        mel.eval("FBXExportBakeComplexEnd -v "+str(cmds.playbackOptions(maxTime=1, q=1)))
+        mel.eval("FBXExportReferencedAssetsContent -v true")
+        mel.eval("FBXExportBakeComplexStep -v 1")
+        mel.eval("FBXExportUseSceneName -v false")
+        if self.euler_CHK.isChecked():
+            mel.eval("FBXExportQuaternion -v euler")
+        if self.quat_CHK.isChecked():
+            mel.eval("FBXExportQuaternion -v quaternion")
+        if self.resample_CHK.isChecked():
+            mel.eval("FBXExportQuaternion -v resample")
+        mel.eval("FBXExportShapes -v true")
+        mel.eval("FBXExportSkins -v true")
+
+        #up axis, the rather large elephant in the room
+        idx = self.upAxisCMB.currentIndex()
+        if idx != 0:
+            if idx == 1:
+                mel.eval("FBXExportUpAxis y")
+            if idx == 2:
+                mel.eval("FBXExportUpAxis z")
+
+
+        #garbage we don't want
+        # Constraints
+        mel.eval("FBXExportConstraints -v false")
+        # Cameras
+        mel.eval("FBXExportCameras -v false")
+        # Lights
+        mel.eval("FBXExportLights -v false")
+        # Embed Media
+        mel.eval("FBXExportEmbeddedTextures -v false")
+        # Connections
+        mel.eval("FBXExportInputConnections -v false")
+        # Axis Conversion
+
+    def export(self, node, mesh=1, anim=1, path=None):
+
+        toExport = []
+
+        if not path:
+            oldPath = path
+            path = cmds.file(sceneName=1, q=1)
+            cmds.warning('No valid path set for export [' + str(oldPath) + ']/nExporting to Maya file loc: ' + path)
+            toExport.extend(cmds.listRelatives(cmds.listConnections(node + '.export_root'), type='joint',allDescendents=True,f=1))
+
+        #kvassey -- adding support for baking to root in rig before export
+        if self.bakeRoot_CHK.isChecked():
+            print "Bake Root Checked"
+            #copy root skeleton with input connections under new group
+            cmds.select(cl=True)
+            currRoot = cmds.listConnections(node + '.export_root')
+            exportSkel = cmds.duplicate(currRoot, un=True, rc=False, po=False)
+            print exportSkel
+            tempGrp = cmds.group(exportSkel[0])
+            #rename root joint
+            dupRoot = cmds.rename(exportSkel[0], currRoot)
+            #bake
+            startTime = cmds.playbackOptions(min=True, q=True)
+            endTime = cmds.playbackOptions(max=True, q=True)
+            cmds.bakeResults(dupRoot, sm=True,  hi="below", s=True, sb=1, dic=True, t=(startTime, endTime))
+
+            #move FBX export inside here, skip rest.
+            #toExport.extend(cmds.listRelatives(cmds.listConnections(dupRoot), type='joint',allDescendents=True))
+            #cmds.select(toExport)
+            cmds.select(dupRoot, r=True, hi=True)
+            toExport = cmds.ls(sl=True, type='joint')
+            cmds.select(toExport, r=True)
+            self.setExportFlags()
+            # Export!
+            print "FBXExport -f \""+ path +"\" -s"
+            mel.eval("FBXExport -f \""+ path +"\" -s")
+
+            cmds.delete(tempGrp)
+            toExport = []
+
+
+        #Assuming the meshes are skinned, maybe work with static later
+        if not self.bakeRoot_CHK.isChecked():
+            if not self.fbxPerMesh_CHK.isChecked():
+                if mesh:
+                    toExport.extend(cmds.listConnections(node + '.rendermesh'))
+                if anim:
+                    toExport.extend(cmds.listRelatives(cmds.listConnections(node + '.export_root'), type='joint', allDescendents=True, f=1))
+
+                cmds.select(toExport)
+                self.setExportFlags()
+                # Export!
+                print "FBXExport -f \""+ path +"\" -s"
+                mel.eval("FBXExport -f \""+ path +"\" -s")
+
+            else:
+                for msh in cmds.listConnections(node + '.rendermesh'):
+                    toExport.append(msh)
+                    toExport.extend(cmds.listRelatives(cmds.listConnections(node + '.export_root'), type='joint',allDescendents=True, f=1))
+                    cmds.select(toExport)
+                    self.setExportFlags()
+                    filePath = os.path.dirname(path)
+                    new_fpath = filePath + '/' + msh + '.fbx'
+                    # Export!
+                    print "FBXExport -f \"" + new_fpath + "\" -s"
+                    mel.eval("FBXExport -f \"" + new_fpath + "\" -s")
+
+                    #clear sel and export list
+                    cmds.select(d=1)
+                    toExport = []
+
+
+if __name__ == '__main__':
+    show()
\ No newline at end of file
diff --git a/Engine/Extras/Maya_AnimationRiggingTools/ARTv1/MayaTools/General/Scripts/P4.py b/Engine/Extras/Maya_AnimationRiggingTools/ARTv1/MayaTools/General/Scripts/P4.py
new file mode 100644
--- /dev/null
+++ b/Engine/Extras/Maya_AnimationRiggingTools/ARTv1/MayaTools/General/Scripts/P4.py
@@ -0,0 +1,757 @@
+#!/usr/local/bin/python
+
+from __future__ import print_function
+
+"""A Python version of the Perforce "p4" client.
+
+This uses the Python type P4API.P4Adapter, which is a wrapper for the
+Perforce ClientApi object. 
+
+$Id: //depot/r12.2/p4-python/P4.py#1 $
+
+#*******************************************************************************
+# Copyright (c) 2007-2010, Perforce Software, Inc.	All rights reserved.
+# Portions Copyright (c) 1999, Mike Meyer. All rights reserved.
+# Portions Copyright (c) 2004-2007, Robert Cowham. All rights reserved.
+#
+# Redistribution and use in source and binary forms, with or without
+# modification, are permitted provided that the following conditions are met:
+#
+# 1.  Redistributions of source code must retain the above copyright
+# notice, this list of conditions and the following disclaimer.
+#
+# 2.  Redistributions in binary form must reproduce the above copyright
+# notice, this list of conditions and the following disclaimer in the
+# documentation and/or other materials provided with the distribution.
+#
+# THIS SOFTWARE IS PROVIDED BY THE COPYRIGHT HOLDERS AND CONTRIBUTORS "AS IS"
+# AND ANY EXPRESS OR IMPLIED WARRANTIES, INCLUDING, BUT NOT LIMITED TO, THE
+# IMPLIED WARRANTIES OF MERCHANTABILITY AND FITNESS FOR A PARTICULAR PURPOSE
+# ARE DISCLAIMED. IN NO EVENT SHALL PERFORCE SOFTWARE, INC. BE LIABLE FOR ANY
+# DIRECT, INDIRECT, INCIDENTAL, SPECIAL, EXEMPLARY, OR CONSEQUENTIAL DAMAGES
+# (INCLUDING, BUT NOT LIMITED TO, PROCUREMENT OF SUBSTITUTE GOODS OR SERVICES;
+# LOSS OF USE, DATA, OR PROFITS; OR BUSINESS INTERRUPTION) HOWEVER CAUSED AND
+# ON ANY THEORY OF LIABILITY, WHETHER IN CONTRACT, STRICT LIABILITY, OR TORT
+# (INCLUDING NEGLIGENCE OR OTHERWISE) ARISING IN ANY WAY OUT OF THE USE OF THIS
+# SOFTWARE, EVEN IF ADVISED OF THE POSSIBILITY OF SUCH DAMAGE.
+#*******************************************************************************
+
+Copyright 2007-2010 Perforce Software, Inc. All rights reserved
+
+
+License:
+See accompanying LICENSE.txt including for redistribution permission.
+"""
+
+import sys, os, string, datetime
+import types, re
+from contextlib import contextmanager
+
+# P4Exception - some sort of error occurred
+class P4Exception(Exception):
+	"""Exception thrown by P4 in case of Perforce errors or warnings"""
+	
+	def __init__(self, value):
+		Exception.__init__(self)
+		
+		if isinstance(value, (list, tuple)) and len(value) > 2:
+			self.value = value[0]
+			self.errors = value[1]
+			self.warnings = value[2]
+		else: 
+			self.value = value
+	def __str__(self):
+		return str(self.value)
+
+class Spec(dict):
+	"""Subclass of dict, representing the fields of a spec definition.
+	
+	Attributes can be accessed either with the conventional dict format,
+	spec['attribute'] or with shorthand spec._attribute.
+	
+	Instances of this class will preventing any unknown keys.
+	"""
+	def __init__(self, fieldmap=None):
+		self.__dict__['_Spec__fields'] = fieldmap
+	
+	def permitted_fields(self):
+		return self.__fields
+	
+	def __setitem__(self, key, value):
+		if not isinstance(value, str) and not isinstance(value, list):
+			raise P4Exception("Illegal value of type %s, must be string or list" % value.__class__)
+		if key in self or self.__fields == None:
+			dict.__setitem__(self, key, value)
+		elif str(key).lower() in self.__fields:
+			dict.__setitem__(self, self.__fields[key.lower()], value)
+		else:
+			raise P4Exception("Illegal field '%s'" % str(key))
+	
+	def __getattr__(self, attr):
+		key = str(attr).lower()
+		if key[0] != '_':
+			raise AttributeError(attr)
+		key = key[1:]
+		if key in self:
+			return self[key]
+		elif key in self.__fields:
+			return self[self.__fields[key]]
+					
+	def __setattr__(self, attr, value):
+		key = str(attr).lower()
+		if key[0] != '_':
+			raise AttributeError(attr)
+		key = key[1:]
+		self[key] = value
+
+#
+# P4Integration objects hold details about the integrations that have
+# been performed on a particular revision. Used primarily with the
+# P4Revision class
+#
+class Integration:
+	def __init__( self, how, file, srev, erev ):
+		self.how = how
+		self.file = file
+		self.srev = srev
+		self.erev = erev
+
+	def __repr__(self):
+		return "Integration (how = %s file = %s srev = %s erev = %s)" \
+			% (self.how, self.file, self.srev, self.erev)
+#
+# Each P4Revision object holds details about a particular revision
+# of a file. It may also contain the history of any integrations 
+# to/from the file
+#
+
+class Revision:
+	def __init__( self, depotFile ):
+		self.depotFile = depotFile
+		self.integrations = []
+		self.rev = None
+		self.change = None
+		self.action = None
+		self.type = None
+		self.time = None
+		self.user = None
+		self.client = None
+		self.desc = None
+		self.digest = None
+		self.fileSize = None
+
+	def integration( self, how, file, srev, erev ):
+		rec = Integration( how, file, srev, erev )
+		self.integrations.append( rec )
+		return rec
+
+# iterator over the collection calling a provided function
+# Python's poor man version of the Ruby block 
+	def each_integration(self):
+		for i in self.integrations:
+			yield i
+
+	def __repr__(self):
+		return "Revision (depotFile = %s rev = %s change = %s action = %s type = %s time = %s user = %s client = %s)" % \
+			(self.depotFile, self.rev, self.change, self.action, self.type, self.time, self.user, self.client)
+#
+# Each DepotFile entry contains details about one depot file. 
+# 
+class DepotFile:
+	def __init__( self, name ):
+		self.depotFile = name
+		self.revisions = []
+
+	def new_revision(self):
+		r = Revision( self.depotFile )
+		self.revisions.append( r )
+		return r
+	
+	def each_revision(self):
+		for r in self.revisions:
+			yield r
+
+	def __repr__(self):
+		return "DepotFile (depotFile = %s, %s revisions)" % ( self.depotFile, len( self.revisions ) )
+
+#
+# Resolver class used in p4.run_resolve()
+#
+# The default simply checks that p4.input is set to sensible value
+# This class is meant to be subclassed for a custom resolver and
+# Resolver.resolve() overriden
+#
+
+class Resolver:
+	def __init__(self):
+		pass
+	
+	def resolve(self, mergeInfo):
+		if mergeInfo.merge_hint == "e":
+			print("Standard resolver encountered merge conflict, skipping resolve")
+			return "s"
+		else:
+			return mergeInfo.merge_hint
+
+	def actionResolve(self, mergeInfo):
+		return mergeInfo.merge_hint
+
+#
+# OutputHandler base class 
+#
+# Extend this class if you want to use the handler interface
+#
+
+class OutputHandler:
+	REPORT	= 0
+	HANDLED = 1
+	CANCEL	= 2
+	
+	def __init__(self):
+		pass
+	
+	def outputText(self, s):
+		return OutputHandler.REPORT
+	
+	def outputBinary(self, b):
+		return OutputHandler.REPORT
+	
+	def outputStat(self, h):
+		return OutputHandler.REPORT
+	
+	def outputInfo(self, i):
+		return OutputHandler.REPORT
+	
+	def outputMessage(self, e):
+		return OutputHandler.REPORT
+
+class ReportHandler( OutputHandler ):
+	def __init__(self):
+		OutputHandler.__init__(self)
+
+	def outputText(self, s):
+		print( "text: ", s)
+		return OutputHandler.HANDLED
+	
+	def outputBinary(self, b):
+		print( "binary: ", b)
+		return OutputHandler.HANDLED
+	
+	def outputStat(self, h):
+		print( "stat:", h)
+		return OutputHandler.HANDLED
+	
+	def outputInfo(self, i):
+		print( "info: ", i)
+		return OutputHandler.HANDLED
+	
+	def outputMessage(self, e):
+		print( "error:", e)
+		return OutputHandler.HANDLED
+	
+class Progress:
+	def __init__(self):
+		pass
+	
+	def init(self, type):
+		self.type = type
+	
+	def setDescription( self, description, units ):
+		self.description = description
+		self.units = units
+	
+	def setTotal( self, total ):
+		self.total = total
+	
+	def update( self, position ):
+		self.position = position
+		
+	def done( self, fail ):
+		pass
+
+class TextProgress(Progress):
+	def __init__(self):
+		Progress.__init__(self)
+	
+	def init(self, type):
+		Progress.init(self, type)
+		print( "Progress.init with %s" % type )
+	
+	def setDescription(self, description, units):
+		Progress.setDescription(self, description, units)
+		print( "Progress.setDescription with %s and units %s" % (description, units) )
+
+	def setTotal( self, total ):
+		Progress.setTotal(self, total)
+		print( "Progress.setTotal with %s " % total )
+	
+	def update( self, position ):
+		Progress.update(self, position )
+		print( "Progress.update with %s" % position )
+
+	def done( self, fail ):
+		Progress.done(self, fail)
+		print( "Progress.done with %s" % fail )
+
+def processFilelog(h):
+	if "depotFile" in h:
+		df = DepotFile( h[ "depotFile" ] )
+		for n, rev in enumerate( h[ "rev" ]):
+			# Create a new revision of this file ready for populating
+			r = df.new_revision()
+			# Populate the base attributes of each revision
+			r.rev = int( rev )
+			r.change = int( h[ "change" ][ n ] )
+			r.action = h[ "action" ][ n ] 
+			r.type = h[ "type" ][ n ]
+			r.time = datetime.datetime.utcfromtimestamp( int( h[ "time" ][ n ]) )
+			r.user = h[ "user" ][ n ]
+			r.client = h[ "client" ][ n ]
+			r.desc = h[ "desc" ][ n ]
+			if "digest" in h and n < len(h[ "digest" ]):
+			  r.digest = h[ "digest" ][ n ]
+			if "fileSize" in h and n < len(h[ "fileSize" ]):
+			  r.fileSize = h[ "fileSize" ][ n ]
+	
+			# Now if there are any integration records for this revision,
+			# add them in too
+			
+			if (not "how" in h) or (n >= len(h["how"]) or h["how"][n] == None):
+				continue
+			else:
+				for m, how in enumerate( h[ "how" ][ n ] ):
+					file = h[ "file" ][ n ][ m ]
+					srev = h[ "srev" ][ n ][ m ].lstrip('#')
+					erev = h[ "erev" ][ n ][ m ].lstrip('#')
+					
+					if srev == "none":
+						srev = 0
+					else:
+						srev = int( srev )
+				  
+					if erev == "none":
+						erev = 0
+					else:
+						erev = int( erev )
+						
+					r.integration( how, file, srev, erev )
+		return df
+	else:
+		raise Exception("Not a filelog object: " + h)
+
+class FilelogOutputHandler(OutputHandler):
+	def __init__(self):
+		OutputHandler.__init__(self)
+		
+	def outputStat(self, h):
+		df = processFilelog(h)
+		
+		return self.outputFilelog(df)
+	
+	def outputFilelog(self, f):
+		return OutputHandler.REPORT
+
+# This is where the C/C++ shared library is loaded
+# It has to be in this place because the library needs to access
+# the classes defined above. Accessing classes defined below this
+# entry would cause an endless loop
+
+import P4API
+#import p4_adsk.P4API as P4API
+
+class P4(P4API.P4Adapter):
+	"""Use this class to communicate with a Perforce server
+	
+	Instances of P4 will use the environment settings (including P4CONFIG)
+	to determine the connection parameters such as P4CLIENT and P4PORT.
+	
+	This attributes can also be set separately before connecting.
+	
+	To run any Perforce commands, users of this class first need to run
+	the connect() method.
+	
+	It is good practice to disconnect() after the program is complete.
+	"""
+	# Constants useful for exception_level
+	# RAISE_ALL:	 Errors and Warnings are raised as exceptions (default)
+	# RAISE_ERROR:	 Only Errors are raised as exceptions
+	# RAISE_NONE:	 No exceptions are raised, instead False is returned
+	
+	RAISE_ALL	 = 2
+	RAISE_ERROR	 = 1
+	RAISE_ERRORS = 1
+	RAISE_NONE	 = 0
+
+	# Named values for generic error codes returned by 
+	# P4API.Message.generic
+	
+	EV_NONE		 = 0	 # misc
+	
+	# The fault of the user
+	
+	EV_USAGE	 = 0x01	 # request not consistent with dox
+	EV_UNKNOWN	 = 0x02	 # using unknown entity
+	EV_CONTEXT	 = 0x03	 # using entity in wrong context
+	EV_ILLEGAL	 = 0x04	 # trying to do something you can't
+	EV_NOTYET	 = 0x05	 # something must be corrected first
+	EV_PROTECT	 = 0x06	 # protections prevented operation
+	
+	# No fault at all
+	
+	EV_EMPTY	 = 0x11	 # action returned empty results
+	
+	# not the fault of the user
+	
+	EV_FAULT	 = 0x21	 # inexplicable program fault
+	EV_CLIENT	 = 0x22	 # client side program errors
+	EV_ADMIN	 = 0x23	 # server administrative action required
+	EV_CONFIG	 = 0x24	 # client configuration inadequate
+	EV_UPGRADE	 = 0x25	 # client or server too old to interact
+	EV_COMM		 = 0x26	 # communications error
+	EV_TOOBIG	 = 0x27	 # not even Perforce can handle this much
+	
+	# Named values for error severities returned by
+	# P4API.Message.severity
+	E_EMPTY		 = 0  # nothing yet
+	E_INFO		 = 1  # something good happened
+	E_WARN		 = 2  # something not good happened
+	E_FAILED	 = 3  # user did something wrong
+	E_FATAL		 = 4  # system broken -- nothing can continue
+
+	# mappings for __iterate
+	# list-of-specs => ( name-of-one-spec, field-name-in-list-of-specs )
+	
+	specfields = { 
+				'clients'	:	('client', 'client'),
+				'labels'	:	('label', 'label'),
+				'branches'	:	('branch', 'branch'),
+				'changes'	:	('change', 'change'),
+				'streams'	:	('stream', 'Stream'),
+				'jobs'		:	('job', 'Job'),
+				'users'		:	('user', 'User'),
+				'groups'	:	('group', 'group'),
+				'depots'	:	('depot', 'name'),
+				'servers'	:	('server', 'Name')
+	}
+	
+	def __init__(self, *args, **kwlist): 
+		P4API.P4Adapter.__init__(self, *args, **kwlist)
+
+	def __del__(self):
+		if self.debug > 3:
+			print("P4.__del__()", file=sys.stderr)
+			
+	# store the references to the created lambdas as a weakref to allow Python 
+	# to clean up the garbage. |The lambda as a closure stores a reference to self
+	# which causes a circular reference problem without the weakref
+	
+	def __getattr__(self, name):
+		if name.startswith("run_"):
+			cmd = name[len("run_"):]
+			return lambda *args, **kargs: self.run(cmd, *args, **kargs)
+		elif name.startswith("delete_"):
+			cmd = name[len("delete_"):]
+			return lambda *args, **kargs: self.run(cmd, "-d", *args, **kargs)
+		elif name.startswith("fetch_"):
+			cmd = name[len("fetch_"):]
+			return lambda *args, **kargs: self.run(cmd, "-o", *args, **kargs)[0]
+		elif name.startswith("save_"):
+			cmd = name[len("save_"):]
+			return lambda *args, **kargs: self.__save(cmd, *args, **kargs)
+		elif name.startswith("parse_"):
+			cmd = name[len("parse_"):]
+			return lambda *args, **kargs: self.parse_spec(cmd, *args, **kargs)
+		elif name.startswith("format_"):
+			cmd = name[len("format_"):]
+			return lambda *args, **kargs: self.format_spec(cmd, *args, **kargs)
+		elif name.startswith("iterate_"):
+			cmd = name[len("iterate_"):]
+			return lambda *args, **kargs: self.__iterate(cmd, *args, **kargs)
+		else:
+			raise AttributeError(name)
+	
+	def __save(self, cmd, *args, **kargs):
+		self.input = args[0]
+		return self.run(cmd, "-i", args[1:], **kargs)
+	
+	def __iterate(self, cmd, *args, **kargs):
+
+		specs = self.run(cmd, *args, **kargs)
+		if cmd in self.specfields:
+			spec = self.specfields[cmd][0]
+			field = self.specfields[cmd][1]
+			
+			# Return a generators (Python iterator object)
+			# On iteration, this will retrieve one spec at a time
+			return ( self.run(spec, '-o', x[field])[0] for x in specs )
+		else:
+			raise Exception('Unknown spec list command: %s', cmd)
+	
+	def __repr__(self):
+		state = "disconnected"
+		if self.connected():
+			state = "connected"
+			
+		return "P4 [%s@%s %s] %s" % \
+		  (self.user, self.client, self.port, state)
+	
+	def identify(cls):
+		return P4API.identify()
+	identify = classmethod(identify)
+	
+	def run(self, *args, **kargs):
+		"Generic run method"
+		context = {}
+		
+		for (k,v) in list(kargs.items()):
+			context[k] = getattr(self, k)
+			setattr(self, k, v)
+			
+		result = P4API.P4Adapter.run(self, *self.__flatten(args))
+		
+		for (k,v) in list(context.items()):
+			setattr( self, k, v)
+
+		return result
+		
+	def run_submit(self, *args, **kargs):
+		"Simplified submit - if any arguments is a dict, assume it to be the changeform"
+		nargs = list(args)
+		form = None
+		for n, arg in enumerate(nargs):
+			if isinstance( arg, dict):
+				self.input = arg
+				nargs.pop(n)
+				nargs.append("-i")
+				break
+		return self.run("submit", *nargs, **kargs)
+	
+	def run_shelve(self, *args):
+		"Simplified shelve - if any arguments is a dict, assume it to be the changeform"
+		nargs = list(args)
+		form = None
+		for n, arg in enumerate(nargs):
+			if isinstance( arg, dict):
+				self.input = arg
+				nargs.pop(n)
+				nargs.append("-i")
+				break
+		return self.run("shelve", *nargs)
+
+	def delete_shelve(self, *args):
+		"Simplified deletion of shelves - if no -c is passed in, add it to the args"
+		nargs = list(args)
+		if '-c' not in nargs:
+			nargs = ['-c'] + nargs # prepend -c if it is not there
+		nargs = ['-d'] + nargs
+		return self.run("shelve", *nargs)
+	
+	def run_login(self, *args):
+		"Simple interface to make login easier"
+		self.input = self.password
+		return self.run("login", *args)
+
+	def run_password( self, oldpass, newpass ):
+		"Simple interface to allow setting of the password"
+		if( oldpass and len(oldpass) > 0 ):
+			self.input = [ oldpass, newpass, newpass ]
+		else:
+			self.input = [ newpass, newpass ]
+			
+		return self.run( "password" )
+
+	#
+	# run_filelog: convert "p4 filelog" responses into objects with useful
+	#			   methods
+	#
+	# Requires tagged output to be of any real use. If tagged output it not 
+	# enabled then you just get the raw data back
+	#
+	def run_filelog( self, *args, **kargs ):
+		raw = self.run( 'filelog', args, **kargs )
+		if (not self.tagged or not raw): 
+			# untagged mode returns simple strings, which breaks the code below
+			# raw could be None if a handler is used
+			return raw
+		result = []
+		for h in raw:
+			df = None
+			if isinstance( h, dict ):
+				df = processFilelog( h )
+			else:
+				df = h
+			result.append( df )
+		return result
+
+	def run_print(self, *args, **kargs):
+		raw = self.run('print', args, **kargs)
+		result = []
+		if raw:
+			for line in raw:
+				if isinstance(line, dict):
+					result.append(line)
+					result.append("")
+				else:
+					result[-1] += line
+			return result
+		else:
+			return []
+	
+	def run_resolve(self, *args, **kargs):
+		if self.resolver:
+			myResolver = self.resolver
+		else:
+			myResolver = Resolver()
+		if "resolver" in kargs: 
+			myResolver = kargs["resolver"]
+			
+		savedResolver = self.resolver
+		self.resolver = myResolver
+		result = self.run("resolve", args)
+		self.resolver = savedResolver
+		
+		return result
+	
+	def run_tickets(self, *args):
+		fname = self.ticket_file
+		with open(fname) as f:
+			tickets_raw = f.readlines()
+		pattern = re.compile('([^=]*)=(.*):([^:]*)\n')
+		tickets = [ pattern.match(x).groups() for x in tickets_raw ]
+		keys = [ "Host", "User", "Ticket" ]
+		result = [ dict(zip(keys, x)) for x in tickets ]
+		return result
+	
+	def __flatten(self, args):
+		result = []
+		if isinstance(args, tuple) or isinstance(args, list):
+			for i in args:
+				result.extend(self.__flatten(i))
+		else:
+			result.append(args)
+		return tuple(result)
+
+	def __enter__( self ):
+		return self
+		
+	def __exit__( self, exc_type, exc_val, exc_tb ):
+		if self.connected():
+			self.disconnect()
+		return False
+		
+	def connect( self ):
+		P4API.P4Adapter.connect( self )
+		return self
+	
+	@contextmanager
+	def while_tagged( self, t ):
+		old = self.tagged
+		self.tagged = t
+		yield
+		self.tagged = old
+	
+	@contextmanager
+	def at_exception_level( self, e ):
+		old = self.exception_level
+		self.exception_level = e
+		yield
+		self.exception_level = old
+		
+	@contextmanager
+	def using_handler( self, c ):
+		old = self.handler
+		self.handler = c
+		yield
+		self.handler = old
+	
+	@contextmanager
+	def saved_context( self , **kargs):
+		"""Saves the context of this p4 object and restores it again at the end of the block"""
+	
+		saved_context = {}
+		for attr in self.__members__:
+			saved_context[attr] = getattr(self, attr)
+	
+		for (k,v) in list(kargs.items()):
+			setattr( self, k, v)
+		
+		yield
+		
+		# now restore the context again. Ignore AttributeError exception
+		# Exception is expected because some attributes only have getters, no setters
+		
+		for (k,v) in list(saved_context.items()):
+			try:
+				setattr( self, k, v )
+			except AttributeError:
+				pass # expected for server_level and p4config_file
+
+class Map(P4API.P4Map):
+	def __init__(self, *args):
+		P4API.P4Map.__init__(self, *args)
+		if len(args) > 0:
+			self.insert( *args )
+	
+	LEFT2RIGHT = True
+	RIGHT2LEFT = False
+	
+	def __str__( self ):
+		result = ""
+		
+		for a in self.as_array():
+			result += a + "\n"
+	
+		return result
+		
+	def is_empty(self):
+		"""Returns True if this map has no entries yet, otherwise False"""
+		
+		return self.count() == 0
+	
+	def includes(self, *args):
+		return self.translate(*args) != None
+	
+	def reverse(self):
+		return Map(P4API.P4Map.reverse(self).as_array())
+	
+	def insert(self, *args):
+		"""Insert an argument to the map. The argument can be:
+		
+		A String, 
+			Either of the form "[+-]//lhs/... //rhs/..." or "[+-]//lhs/..." 
+			for label style maps.
+		A List:
+			This is a list of strings of one of the single string formats 
+			described above.
+		A pair of Strings:
+			P4.Map.insert(lhs, rhs)
+		"""
+		
+		if len(args) == 1 :
+			arg = args[0]
+			if isinstance( arg, str ):
+				P4API.P4Map.insert( self, arg )
+			elif isinstance( arg, list ):
+				for s in arg:
+					P4API.P4Map.insert( self, s )
+					
+		else: # expecting 2 args in this case: left, right 
+			left = args[0].strip()
+			right = args[1].strip()
+			P4API.P4Map.insert(self, left, right )
+
+if __name__ == "__main__":
+	p4 = P4()
+	p4.connect()
+	try:
+		ret = p4.run(sys.argv[1:])
+		for line in ret:
+			if isinstance(line, dict):
+				print("-----")
+				for k in list(line.keys()):
+					print(k, "=", line[k])
+			else:
+				print(line)
+	except:
+		for e in p4.errors:
+			print(e)
diff --git a/Engine/Extras/Maya_AnimationRiggingTools/ARTv1/MayaTools/General/Scripts/Tools/ART_CopyWeightsTool.py b/Engine/Extras/Maya_AnimationRiggingTools/ARTv1/MayaTools/General/Scripts/Tools/ART_CopyWeightsTool.py
new file mode 100644
--- /dev/null
+++ b/Engine/Extras/Maya_AnimationRiggingTools/ARTv1/MayaTools/General/Scripts/Tools/ART_CopyWeightsTool.py
@@ -0,0 +1,189 @@
+from PySide import QtGui
+from PySide import QtCore
+import maya.OpenMayaUI as mui
+import shiboken
+import maya.cmds as cmds
+from functools import partial
+
+#Written by: Charles R. Anderson
+#Dec. 4, 2015 
+#Questions or comments email me: [REDACTED_EMAIL]
+
+def findCluster(meshName):
+    skinClusters = []
+    skinHistory = []
+    skinHistory = cmds.listHistory(meshName, lv=0)
+    for node in skinHistory:
+        if cmds.nodeType(node) == "skinCluster":
+            skinClusters.append(node)
+
+    return skinClusters[0]
+    
+#-----------------------------------------------------------------------------------------------
+def copyWeights(fromList, toList):
+    #filterString = textField.text()
+    #print filterString
+    if fromList.count() > 0:
+        for fromIndex in xrange(fromList.count()):
+            fromMesh = fromList.item(fromIndex).text()
+            print "From Mesh: ", fromMesh
+            if toList.count() > 0:
+                for toIndex in xrange(toList.count()):
+                    toMesh = toList.item(toIndex).text()
+                    if fromMesh.lower() in toMesh.lower():
+                        print "Matched To Mesh found: ", toMesh
+                        print "Copying weights from ", fromMesh, " to ", toMesh
+                        toCluster = findCluster(toMesh)
+                        fromCluster = findCluster(fromMesh)
+                        cmds.copySkinWeights(ss=fromCluster, ds=toCluster, noMirror=True, surfaceAssociation="closestPoint", influenceAssociation=["closestJoint", "oneToOne", "name"])
+            else:
+                cmds.warning("Your Copy TO list is empty.")
+    else:
+        cmds.warning("Your Copy FROM list is empty.")
+
+######################## UI FUNCTIONS ###############################
+#-----------------------------------------------------------------------------------------------
+# This function finds the main Maya Window
+def getMayaWindow():
+    pointer = mui.MQtUtil.mainWindow()
+    return shiboken.wrapInstance(long(pointer), QtGui.QWidget)
+
+#-----------------------------------------------------------------------------------------------
+def addMeshItem (listName):
+    listName.clear() # TEMPORARILY clears the list so that we never get dupes.  Figure out how to check for copies before adding and then remove this.
+    addItems = cmds.ls(sl=True)
+    #for one in addItems:
+        #print listName.findItems(one, QtCore.Qt.MatchExactly)
+    #for one in addItems:
+    listName.addItems(addItems)
+	
+#-----------------------------------------------------------------------------------------------
+def removeMeshItem (listName):
+    selected = listName.selectedItems()
+    for one in selected:
+        index = listName.row(one)
+        listName.takeItem(index)
+
+#-----------------------------------------------------------------------------------------------
+def clearMeshItems (listName):
+    listName.clear()
+    	
+#-----------------------------------------------------------------------------------------------
+def resetUI (fromList, toList):
+    fromList.clear()
+    toList.clear()
+
+#-----------------------------------------------------------------------------------------------
+def createUI(*args):
+    # check to see if this window exists and if so delete it.
+    objectName = "copyWeightsWindowNew"
+    if cmds.window("copyWeightsWindowNew", exists=True):
+        cmds.deleteUI("copyWeightsWindowNew", wnd=True)
+
+    # Create the Copy Weights window
+    parent = getMayaWindow()
+    currentWindow = QtGui.QMainWindow(parent)
+    currentWindow.setObjectName(objectName)
+    currentWindow.setWindowTitle("Copy Weights Tool")
+    currentWindow.resize(500, 700)
+
+    # create a widget
+    widget = QtGui.QWidget()
+    currentWindow.setCentralWidget(widget)
+
+    # create the Main vertical layout
+    mainLayout = QtGui.QVBoxLayout(widget)
+
+    # create a horizontal layout for the lists
+    listsLayout = QtGui.QHBoxLayout()
+    mainLayout.addLayout(listsLayout)
+
+    # FROM section---------------------
+    fromLayout = QtGui.QVBoxLayout()
+    listsLayout.addLayout(fromLayout)
+
+    fromListLabel = QtGui.QLabel("Copy from:")
+    fromLayout.addWidget(fromListLabel)
+    fromList = QtGui.QListWidget()
+    fromLayout.addWidget(fromList)
+    fromList.setToolTip("This is a list of meshes to copy weights from.")
+    fromList.setSelectionMode(QtGui.QAbstractItemView.ExtendedSelection)
+
+    # Add Button
+    addFromButton = QtGui.QPushButton("Add Selected")
+    fromLayout.addWidget(addFromButton)
+    addFromButton.setMinimumHeight(30)
+    addFromButton.setMaximumHeight(30)
+    addFromButton.setMinimumWidth(200)
+    addFromButton.setToolTip("Adds the selected meshes to the list above.")
+    addFromButton.clicked.connect(partial(addMeshItem, fromList))
+
+    removeFromButton = QtGui.QPushButton("Remove Selected")
+    fromLayout.addWidget(removeFromButton)
+    removeFromButton.setMinimumHeight(30)
+    removeFromButton.setMaximumHeight(30)
+    removeFromButton.setMinimumWidth(200)
+    removeFromButton.clicked.connect(partial(removeMeshItem, fromList))
+
+    clearFromButton = QtGui.QPushButton("Clear List")
+    fromLayout.addWidget(clearFromButton)
+    clearFromButton.setMinimumHeight(30)
+    clearFromButton.setMaximumHeight(30)
+    clearFromButton.setMinimumWidth(200)
+    clearFromButton.clicked.connect(partial(clearMeshItems, fromList))
+
+    # TO section-------------------------
+    toLayout = QtGui.QVBoxLayout()
+    listsLayout.addLayout(toLayout)
+
+    toListLabel = QtGui.QLabel("Copy to:")
+    toLayout.addWidget(toListLabel)
+    toList = QtGui.QListWidget()
+    toLayout.addWidget(toList)
+    toList.setToolTip("This is a list of meshes to copy weights to.")
+    toList.setSelectionMode(QtGui.QAbstractItemView.ExtendedSelection)
+
+    # Add Button
+    addToButton = QtGui.QPushButton("Add Selected")
+    toLayout.addWidget(addToButton)
+    addToButton.setMinimumHeight(30)
+    addToButton.setMaximumHeight(30)
+    addToButton.setMinimumWidth(200)
+    addToButton.clicked.connect(partial(addMeshItem, toList))
+
+    removeToButton = QtGui.QPushButton("Remove Selected")
+    toLayout.addWidget(removeToButton)
+    removeToButton.setMinimumHeight(30)
+    removeToButton.setMaximumHeight(30)
+    removeToButton.setMinimumWidth(200)
+    removeToButton.clicked.connect(partial(removeMeshItem, toList))
+
+    clearToButton = QtGui.QPushButton("Clear List")
+    toLayout.addWidget(clearToButton)
+    clearToButton.setMinimumHeight(30)
+    clearToButton.setMaximumHeight(30)
+    clearToButton.setMinimumWidth(200)
+    clearToButton.clicked.connect(partial(clearMeshItems, toList))
+
+
+    # Main widgets ---------------------
+    resetButton = QtGui.QPushButton("Reset")
+    mainLayout.addWidget(resetButton)
+    resetButton.setMinimumHeight(30)
+    resetButton.setMaximumHeight(30)
+    resetButton.setMinimumWidth(200)
+    resetButton.clicked.connect(partial(resetUI, fromList, toList))
+
+    copyButton = QtGui.QPushButton("Copy")
+    mainLayout.addWidget(copyButton)
+    copyButton.setMinimumHeight(30)
+    copyButton.setMaximumHeight(30)
+    copyButton.setMinimumWidth(200)
+    copyButton.setToolTip("Run through all of the meshes in the FROM list and copy its weights to any meshes that it matches in the TO list.  This matches based on name and is case insensitive.  For example it will find \"Art\" inside of \"fart\".")
+    copyButton.clicked.connect(partial(copyWeights, fromList, toList))
+
+    # show the window
+    currentWindow.show()
+
+#createUI()
+
diff --git a/Engine/Extras/Maya_AnimationRiggingTools/ARTv1/MayaTools/General/Scripts/Tools/ART_CreateHydraulicRigs.py b/Engine/Extras/Maya_AnimationRiggingTools/ARTv1/MayaTools/General/Scripts/Tools/ART_CreateHydraulicRigs.py
new file mode 100644
--- /dev/null
+++ b/Engine/Extras/Maya_AnimationRiggingTools/ARTv1/MayaTools/General/Scripts/Tools/ART_CreateHydraulicRigs.py
@@ -0,0 +1,60 @@
+import maya.cmds as cmds
+import maya.mel as mel
+import os
+
+# This tool only works on "leaf" joints created with A.R.T 
+# To use, select 3 objects, the start of your hydraulic system, the end, and the upvector.  It is assumed that you are selecting the _anim controls for the first 2
+# It is assumed that the upvector is 1 by default unless passed in otherwise.  Once set, you can always change this by simply manually reversing the value in the aimConstraint
+# It is also assumed that you are using X as the axis that points at each other and that the first control points at the 2nd, but the 2nd points away from the first.
+# Assuming Z as the up vector axis.
+
+def CreateHydraulicRig(upDir, *args):
+
+    selection = cmds.ls(sl=True)
+    if len(selection) != 3:
+        cmds.confirmDialog(icon = "warning!!", title = "Hydraulic Rig Tool", message = "You must select exactly 3 objects.  Top of hydraulic, bottom of hydraulic, and the up vector.")
+        return
+    else:
+        hyd_start_01_anim = selection[0]
+        hyd_start_02_anim = selection[1]
+        upObject = selection[2]
+        
+        hyd_start_01 = hyd_start_01_anim.rpartition("_")[0]
+        hyd_start_02 = hyd_start_02_anim.rpartition("_")[0]
+
+        # Create a rig for the lower arm extra hydraulic piston.
+        cmds.delete("driver_"+hyd_start_01+"_parentConstraint1")
+        robo_lowarm_pistonrod_01_anim_sub = cmds.group(em=True, name=hyd_start_01_anim+"_sub", parent=hyd_start_01_anim)
+        const = cmds.parentConstraint(hyd_start_01_anim, robo_lowarm_pistonrod_01_anim_sub, weight=1, mo=False)
+        cmds.delete(const)
+        const = cmds.parentConstraint(robo_lowarm_pistonrod_01_anim_sub, "driver_"+hyd_start_01, weight=1, mo=True)
+        
+        cmds.delete("driver_"+hyd_start_02+"_parentConstraint1")
+        robo_lowarm_pistonrod_02_anim_sub = cmds.group(em=True, name=hyd_start_02_anim+"_sub", parent=hyd_start_02_anim)
+        const = cmds.parentConstraint(hyd_start_02_anim, robo_lowarm_pistonrod_02_anim_sub, weight=1, mo=False)
+        cmds.delete(const)
+        const = cmds.parentConstraint(robo_lowarm_pistonrod_02_anim_sub, "driver_"+hyd_start_02, weight=1, mo=True)
+
+        # Hook up the hydraulics for the lowerarm piston.
+        const1 = cmds.aimConstraint(robo_lowarm_pistonrod_01_anim_sub, robo_lowarm_pistonrod_02_anim_sub, weight=1, mo=True, aimVector=(-1, 0, 0), upVector=(0, 0, upDir), worldUpType="object", worldUpVector=(0, 0, -1), worldUpObject=upObject)
+        const2 = cmds.aimConstraint(robo_lowarm_pistonrod_02_anim_sub, robo_lowarm_pistonrod_01_anim_sub, weight=1, mo=True, aimVector=(1, 0, 0), upVector=(0, 0, upDir), worldUpType="object", worldUpVector=(0, 0, -1), worldUpObject=upObject)
+    
+    cmds.select(const1, const2)
+
+def CreateHydraulicSlideRig(*args):
+    print "test"
+    # Create a slide mechanism on the forearm part of the lowerarm piston.
+    # CRA NOTE: this will come at a later date.
+    '''
+    cmds.addAttr("robo_lowarm_pistonrod_02_anim", ln = "Slide", dv = 0, min=0, max=10, keyable = True)
+    cmds.setDrivenKeyframe("robo_lowarm_pistonrod_02_anim_grp.translateX", cd="robo_lowarm_pistonrod_02_anim.Slide", itt="spline", ott="spline")
+    cmds.setDrivenKeyframe("robo_lowarm_pistonrod_02_anim_grp.translateZ", cd="robo_lowarm_pistonrod_02_anim.Slide", itt="spline", ott="spline")
+
+    cmds.setAttr("robo_lowarm_pistonrod_02_anim.Slide", 10)
+    cmds.setAttr("robo_lowarm_pistonrod_02_anim_grp.tx", -6.662443)
+    cmds.setAttr("robo_lowarm_pistonrod_02_anim_grp.tz", -5.510393)
+
+    cmds.setDrivenKeyframe("robo_lowarm_pistonrod_02_anim_grp.translateX", cd="robo_lowarm_pistonrod_02_anim.Slide", itt="spline", ott="spline")
+    cmds.setDrivenKeyframe("robo_lowarm_pistonrod_02_anim_grp.translateZ", cd="robo_lowarm_pistonrod_02_anim.Slide", itt="spline", ott="spline")
+    cmds.setAttr("robo_lowarm_pistonrod_02_anim.Slide", 0)
+    '''
