cleaning up

diff --git a/Engine/Extras/Maya_AnimationRiggingTools/ARTv1/MayaTools/General/Scripts/Tools/ART_ImportMultipleWeights.py b/Engine/Extras/Maya_AnimationRiggingTools/ARTv1/MayaTools/General/Scripts/Tools/ART_ImportMultipleWeights.py
new file mode 100644
--- /dev/null
+++ b/Engine/Extras/Maya_AnimationRiggingTools/ARTv1/MayaTools/General/Scripts/Tools/ART_ImportMultipleWeights.py
@@ -0,0 +1,272 @@
+import maya.cmds as cmds
+import os, cPickle
+from PySide import QtGui
+import maya.OpenMayaUI as mui
+import shiboken
+from PySide import QtCore
+from functools import partial
+
+#############################################################################################
+#############################################################################################
+#############################################################################################
+# Select all of the meshes in your scene that you want to import onto.  
+# If there is a .txt file in the SkinWeights folder on disk that matches its name it will import the weights onto the mesh.
+def importMultipleSkinWeights(suffix, rui, *args):
+    
+    #tools path
+    toolsPath = cmds.internalVar(usd = True) + "mayaTools.txt"
+    if os.path.exists(toolsPath):
+
+        f = open(toolsPath, 'r')
+        mayaToolsDir = f.readline()
+        f.close()	
+
+
+    selection = cmds.ls(sl = True)
+    if selection == []:
+        cmds.confirmDialog(icon = "warning", title = "Import Skin Weights", message = "Nothing Selected")
+        return
+        
+    warnings = False
+
+    if selection != []:
+        #find all files in the system dir
+        path = mayaToolsDir + "\General\ART\SkinWeights\\"
+        print "path is "+path
+        
+        for one in selection:
+            '''
+            if len(selection) > 1:
+                cmds.confirmDialog(icon = "warning", title = "Import Skin Weights", message = "Too many objects selected. Please only selected 1 mesh.")
+                return
+            '''
+            #try:
+                #fileToImport = cmds.fileDialog2(fileFilter="*.txt", dialogStyle=2, fm = 1, startingDirectory = path)[0]
+
+            #except:
+                #cmds.warning("Operation Cancelled.")
+                #return
+            filename = one
+            print "SUFFIX: ", suffix
+            if suffix:
+                filename = one.partition(suffix)[0]
+            
+            print "FILENAME: ", filename
+
+            fileToImport = path+filename+".txt"
+            if cmds.file(fileToImport, q=True, ex=True):
+                print "file is "+fileToImport
+
+                #load the file with cPickle and parse out the information
+                file = open(fileToImport, 'r')
+                weightInfoList = cPickle.load(file)
+
+                #get the number of vertices in the mesh. This will be used to update our progress window
+                verts = cmds.polyEvaluate(one, vertex = True)
+                if verts < 20:
+                    verts = verts * 20
+
+                increment = int(verts/20)
+                matchNum = increment
+
+                #find the transforms(joints) from the weightInfoList. format = [ [vertNum, [jointName, value], [jointName, value] ] ]
+                transforms = []
+                newList = weightInfoList[:]
+                for info in newList:
+
+                    for i in info:
+                        if i != info[0]:
+                            transform = i[0]
+                            transforms.append(transform)
+
+                #remove duplicates from the transforms(joints) list
+                transforms = set(transforms)
+                print "Files Transforms List:"
+                print transforms
+
+                #clear selection. Add transforms(joints) to selection and geo
+                cmds.select(clear = True)
+                for t in transforms:
+                    if cmds.objExists(t):
+                        cmds.select(t, add = True)
+
+
+                #check if the geometry that is selected is already skinned or not
+                skinClusters = cmds.ls(type='skinCluster')
+                skinnedGeometry = []
+                for cluster in skinClusters:
+                    geometry = cmds.skinCluster(cluster, q = True, g = True)[0]
+                    geoTransform = cmds.listRelatives(geometry, parent = True)[0]
+                    skinnedGeometry.append(geoTransform)
+
+                    if geoTransform == one:
+                        skin = cluster
+                        
+                        currentTransforms = cmds.skinCluster(skin, q=True, inf=True)
+                        print "Current Transforms List:"
+                        print currentTransforms
+                        
+                        for tr in transforms:
+                            if tr not in currentTransforms:
+                                cmds.skinCluster(skin, e=True, ai=tr, lw=False)
+                                print (cmds.skinCluster(skin, q=True, inf=True))
+
+
+                #if a skin cluster does not exist already on the mesh:
+                if one not in skinnedGeometry:
+
+                    #put a skin cluster on the geo normalizeWeights = 1, dr=4.5, mi = 4,
+                    cmds.select(one, add = True)
+                    skin = cmds.skinCluster( tsb = True, skinMethod = 0, name = (one +"_skinCluster"))[0]
+
+
+                #setup a progress window to track length of progress
+                cmds.progressWindow(title='Skin Weights', progress = 0, status = "Reading skin weight information from file...")
+
+                #get the vert number, and the weights for that vertex
+                for info in weightInfoList:
+                    vertNum = info[0]
+                    info.pop(0)
+
+                    #progress bar update
+                    if vertNum == matchNum:
+                        cmds.progressWindow( edit=True, progress = ((matchNum/increment) * 5), status=  "Importing skin weights for " + one + ".\n Please be patient")
+                        matchNum += increment
+
+                    #set the weight for each vertex
+                    try:
+                        cmds.skinPercent(skin, one + ".vtx[" + str(vertNum) + "]", transformValue= info)
+                    except:
+                        pass
+
+                # remove unused influences
+                if rui == 1:
+                    infToRemove = []
+                    weightedInfs = cmds.skinCluster(skin, q=True, weightedInfluence=True)
+                    finalTransforms = cmds.skinCluster(skin, q=True, inf=True)
+                    for inf in finalTransforms:
+                        if inf not in weightedInfs:
+                            infToRemove.append(inf)
+                    for rem in infToRemove:
+                        cmds.skinCluster(skin, e=True, ri=rem)
+
+                #close the file
+                file.close()
+            else:
+                cmds.warning(fileToImport+" does not exist.")
+                warnings = True
+
+            #close progress window
+            cmds.progressWindow(endProgress = 1)
+
+    cmds.select(selection)
+    if warnings == False:
+        cmds.headsUpMessage( "Importing weights is complete!!", time=3.0 )
+        #cmds.confirmDialog(icon = "Complete", title = "Import Skin Weights", message = "Importing weights is complete!!")
+    elif warnings == True:
+        cmds.headsUpMessage( "Importing weights is complete but with Warnings!!", time=3.0 )
+        #cmds.confirmDialog(icon = "Complete", title = "Import Skin Weights", message = "Importing weights is complete but with Warnings!!")
+
+    
+######################## UI FUNCTIONS ###############################
+#-----------------------------------------------------------------------------------------------
+# This function finds the main Maya Window
+def getMayaWindow():
+    pointer = mui.MQtUtil.mainWindow()
+    return shiboken.wrapInstance(long(pointer), QtGui.QWidget)
+
+#-----------------------------------------------------------------------------------------------
+def resetUI (fromList, toList):
+    fromList.clear()
+    toList.clear()
+
+#-----------------------------------------------------------------------------------------------
+def deleteUI ():
+    cmds.deleteUI("importMultipleWeightsWindow", wnd=True)
+    
+#-----------------------------------------------------------------------------------------------
+def importFunction (suffixField, RICheckBox, closeUI):
+    suffixField = suffixField.text()
+    RICheckBox = RICheckBox.isChecked()
+
+    importMultipleSkinWeights(suffixField, RICheckBox)
+    if closeUI == True:
+        cmds.deleteUI("importMultipleWeightsWindow", wnd=True)
+
+#-----------------------------------------------------------------------------------------------
+def importMultipleWeightsUI(*args):
+    # check to see if this window exists and if so delete it.
+    objectName = "importMultipleWeightsWindow"
+    if cmds.window("importMultipleWeightsWindow", exists=True):
+        cmds.deleteUI("importMultipleWeightsWindow", wnd=True)
+
+    # Create the Weights window
+    parent = getMayaWindow()
+    currentWindow = QtGui.QMainWindow(parent)
+    currentWindow.setObjectName(objectName)
+    currentWindow.setWindowTitle("Import Multiple Weights Tool")
+    currentWindow.resize(50, 50)
+
+    # create a widget
+    widget = QtGui.QWidget()
+    currentWindow.setCentralWidget(widget)
+
+    # create the Main vertical layout
+    mainLayout = QtGui.QVBoxLayout(widget)
+    subLayout = QtGui.QHBoxLayout()
+    mainLayout.addLayout(subLayout)
+
+    # create a horizontal layout for the lists
+    optionsLayout = QtGui.QVBoxLayout()
+    mainLayout.addLayout(optionsLayout)
+
+    buttonsLayout = QtGui.QHBoxLayout()
+    mainLayout.addLayout(buttonsLayout)
+    
+    # Suffix UI
+    suffixLayout = QtGui.QHBoxLayout()
+    optionsLayout.addLayout(suffixLayout)
+    textFieldLabel = QtGui.QLabel("Suffix to be removed:")
+    suffixLayout.addWidget(textFieldLabel)
+    suffixField = QtGui.QLineEdit()
+    suffixLayout.addWidget(suffixField)
+    suffixField.setPlaceholderText("Enter Suffix...")
+
+    # Remove Unused Influences UI
+    RUILayout = QtGui.QHBoxLayout()
+    optionsLayout.addLayout(RUILayout)
+    RICheckBox = QtGui.QCheckBox("Remove Unused Influences")
+    RUILayout.addWidget(RICheckBox)
+    RICheckBox.setChecked(True)
+    
+    # Import Button
+    importButton = QtGui.QPushButton("Import")
+    buttonsLayout.addWidget(importButton)
+    importButton.setMinimumHeight(30)
+    importButton.setMaximumHeight(30)
+    importButton.setMinimumWidth(200)
+    importButton.setToolTip("Adds the selected meshes to the list above.")
+    importButton.clicked.connect(partial(importFunction, suffixField, RICheckBox, True))
+
+    # Apply Button
+    applyButton = QtGui.QPushButton("Apply")
+    buttonsLayout.addWidget(applyButton)
+    applyButton.setMinimumHeight(30)
+    applyButton.setMaximumHeight(30)
+    applyButton.setMinimumWidth(200)
+    applyButton.clicked.connect(partial(importFunction, suffixField, RICheckBox, False))
+
+    # Close Button
+    closeFromButton = QtGui.QPushButton("Close")
+    buttonsLayout.addWidget(closeFromButton)
+    closeFromButton.setMinimumHeight(30)
+    closeFromButton.setMaximumHeight(30)
+    closeFromButton.setMinimumWidth(200)
+    closeFromButton.clicked.connect(deleteUI)
+    
+    # show the window
+    currentWindow.show()
+
+#importMultipleWeightsUI()
+    
+#importMultipleSkinWeights("_LOD4")
diff --git a/Engine/Extras/Maya_AnimationRiggingTools/ARTv1/MayaTools/General/Scripts/Tools/ART_MeshSpliter.py b/Engine/Extras/Maya_AnimationRiggingTools/ARTv1/MayaTools/General/Scripts/Tools/ART_MeshSpliter.py
new file mode 100644
--- /dev/null
+++ b/Engine/Extras/Maya_AnimationRiggingTools/ARTv1/MayaTools/General/Scripts/Tools/ART_MeshSpliter.py
@@ -0,0 +1,93 @@
+import maya.cmds as cmds
+import maya.mel as mel
+
+# Splits the selected meshes by their skinCLuster influences and creates a new heirarchy and constrains them to the skeleton.
+# NOTE this script doesnt work well if the names of the meshes have ":" in them.
+def splitMeshByInfluences(threshold, *args):
+    import maya.mel as mel
+    
+    finalMeshList = []
+    
+    # find the meshes that the user has selected
+    selected = cmds.ls(sl=True, type="transform")
+
+    # Create an empty set.  We wil add to it later.
+    cmds.select(cl=True)
+    setName = cmds.sets(n="Rigid_Meshes")
+    
+    rigidMeshMasterGrp = cmds.group(em=True, name="Rigid_Mesh")
+
+    # for each mesh:
+    for mesh in selected:
+        rigidMeshMeshGrp = cmds.group(em=True, name="Rigid_"+mesh, p=rigidMeshMasterGrp)
+        #print "Mesh:"+mesh
+
+        skinCluster = mel.eval('findRelatedSkinCluster %s' % mesh)
+        influences = cmds.skinCluster(skinCluster, q=True, inf=True)
+        #print influences
+
+        for inf in influences:
+            # Duplicate the mesh with input connections
+            dupeMesh = cmds.duplicate(mesh, name=mesh+"__"+inf, rr=False)[0]
+            #print "DupeMesh:       "+dupeMesh
+            for attr in ["tx", "ty", "tz", "rx", "ry", "rz", "sx", "sy", "sz", "v"]:
+                cmds.setAttr(dupeMesh+"."+attr, lock=False, k=True)
+            cmds.makeIdentity(dupeMesh, t=1, r=1, s=1, apply=True)
+            cmds.select(dupeMesh, r=True)
+            mel.eval("CenterPivot")
+            cmds.parent(dupeMesh, rigidMeshMeshGrp)
+            
+            finalMeshList.append(dupeMesh)
+            cmds.sets(dupeMesh, add=setName)
+
+            # Select the verts that are influenced by the inf joint.
+            if len(influences) <= 1:
+                const = cmds.parentConstraint(inf, rigidMeshMeshGrp, dupeMesh, mo=True)[0]
+                cmds.setAttr(const+"."+rigidMeshMeshGrp+"W1", 0)
+            else:
+                verts = []
+                for x in range(cmds.polyEvaluate(mesh, v=True)):
+                    v = cmds.skinPercent(skinCluster, '%s.vtx[%d]' % (mesh, x), transform=inf, q=True)
+                    if v > threshold:
+                        verts.append('%s.vtx[%d]' % (dupeMesh, x))
+                #print verts
+                if verts:
+                    cmds.select(verts)            
+                    # Convert the selection to contained faces
+                    if len(verts) != cmds.polyEvaluate(mesh, v=True):
+                        mel.eval("InvertSelection;")
+                        cmds.select(cmds.polyListComponentConversion(fv=True, tf=True, internal=True ))
+                        cmds.delete()
+                    cmds.select(dupeMesh, r=True)
+                    mel.eval("CenterPivot")
+                    const = cmds.parentConstraint(inf, rigidMeshMeshGrp, dupeMesh, mo=True)[0]
+                    cmds.setAttr(const+"."+rigidMeshMeshGrp+"W1", 0)
+                else:
+                    print "No Valid Verts so we didnt use this influence."
+                    cmds.delete(dupeMesh)
+
+    cmds.select(rigidMeshMasterGrp, hi=True)
+    cmds.createDisplayLayer(cmds.ls(sl=True), name="RigidMesh_L")
+    
+    return finalMeshList
+
+#splitMeshByInfluences(0.5)
+
+def findThreshold(*args):
+    threshold = cmds.floatSliderGrp("skinWeightFloatSliderGrp", q=True, value=True)
+    splitMeshByInfluences(threshold)
+
+
+# Create a UI
+def splitMeshByInfluencesUI(*args):
+    if cmds.window("splitMeshByInfluencesUI", exists = True):
+        cmds.deleteUI("splitMeshByInfluencesUI")
+
+    window = cmds.window("splitMeshByInfluencesUI", title='Split Mesh by Influences', s=False)
+    cmds.columnLayout()
+    cmds.floatSliderGrp("skinWeightFloatSliderGrp", label='Skin Weight Threshold', field=True, minValue=0, maxValue=1, value=0.5, step=0.01, cw3=[110, 30, 150])
+    cmds.button("Apply", label="Apply", w=290, c=findThreshold)
+
+    cmds.showWindow(window)
+
+#splitMeshByInfluencesUI()
\ No newline at end of file
diff --git a/Engine/Extras/Maya_AnimationRiggingTools/ARTv1/MayaTools/General/Scripts/Tools/ART_OrientJointWithUp.py b/Engine/Extras/Maya_AnimationRiggingTools/ARTv1/MayaTools/General/Scripts/Tools/ART_OrientJointWithUp.py
new file mode 100644
--- /dev/null
+++ b/Engine/Extras/Maya_AnimationRiggingTools/ARTv1/MayaTools/General/Scripts/Tools/ART_OrientJointWithUp.py
@@ -0,0 +1,295 @@
+import maya.cmds as cmds
+from functools import partial
+import pymel.core as pm
+
+#Select the joint you want to orient, then the child you want it to be aimed at, and then the object to use as the upVector
+#CRA NOTE: OptionVars are not retaining their settings when you re-open the UI.  They always reset to the default.  Figure out how to make them remember.
+#--------------------------------------------------------------------------------------------------------------------
+#aimAxis        = axis to point at the child bone.
+#upAxis         = axis to point at the locator bone.
+#heirarchy      = Should this be applied using the heirarchy or only to the selected joints?
+#customChildInt = Do we use a custom child?   When set to 1, this will tell it to use the named joint in customChild rather than popping open a window when more than 1 child joint exists.
+#customChild    = The joint to use as the aim joint when customChildInt is set to 1.  If no custom child is defined it will either use the only child available or if there are multiples it will bring up a ui to let the user choose which child to use.
+# IMPORTANT NOTE: When calling this script directly and not using the UI, it is reccommended to have heirarchy set to 0.  If you try to orient an entire heirarchy we can not guarantee that it will pick the correct child unless you are 100% positive that there are no parents with multiple children.
+#----------------------------------------------------------------------------------------
+# This def finds the joint we want to aim (jointToAim), the joint to be aimed at (aimTarget), and the up node (upNode) based on the selection.  IF heirarchy is checked, then it uses the findChild def to find the joint to be aimed at.
+def setOrientWithUpVector(aimAxis, upAxis, heirarchy, customChildInt, customChild):
+    # THESE ARE THE 3 NODES WE NEED TO DEFINE!!
+    jointToAim = None
+    aimTarget = None
+    upNode = None
+
+    #Find the selected joint and upVector nodes from the selection
+    # CRA NOTE: Could add some more logic here later to check to see if the first 2 selected are joints and that the 2nd is a relative of the first.
+    selected = cmds.ls(sl=True)
+    if len(selected) != 2:
+        cmds.confirmDialog(icon = "warning!!", title = "Orient Joint with Up Tool", message = "You must select exactly 2 objects.  Joint you want to orient and then the object to use as the upVector.")
+        return
+    else:
+        jointToAim = selected[0]
+        upNode = selected[1]
+
+        # Find the joints that we need to apply the orient script to.
+        if heirarchy == 0:
+            # find only the immediate children of jointToAim
+            allChildren = cmds.listRelatives(jointToAim, c=True, type="joint")
+            #print "\nAll Children..."
+            #print allChildren
+
+            if customChildInt == 1:
+                aimTarget = customChild
+                if allChildren != None:
+                    doOrientJoint(jointToAim, aimTarget, upNode, aimAxis, upAxis, allChildren)
+                else:
+                    cmds.joint(jointToAim, e=True, oj="none", secondaryAxisOrient="yup", ch=True, zso=True)
+            elif allChildren != None:
+                if len(allChildren) == 1:
+                    aimTarget = allChildren[0]
+                else:
+                    aimTarget = cmds.layoutDialog(t="Child Joints List", ui=partial(findChild, jointToAim))
+                #print "\nPRINTING JOINTS LIST---- "
+                #print "jointToAim: "+jointToAim
+                #print "aimTarget: "+aimTarget
+                #print "upNode   : "+upNode   
+                #print "----DONE PRINTING JOINTS LIST"
+                doOrientJoint(jointToAim, aimTarget, upNode, aimAxis, upAxis, allChildren)
+            else:
+                cmds.joint(jointToAim, e=True, oj="none", secondaryAxisOrient="yup", ch=True, zso=True)
+
+
+        if heirarchy == 1:
+            # find all of the joints we want to use as the jointToAim
+            cmds.select(jointToAim, hi=True)
+            allJointsToAim = cmds.ls(sl=True, type="joint")
+            #print "\nALL Joints To Aim..."
+            #print allJointsToAim
+
+            for one in allJointsToAim:
+                jointToAim = one
+                #print "\nCURRENT JOINT..."
+                #print jointToAim
+                
+                # find only the immediate children of jointToAim
+                allChildren = cmds.listRelatives(jointToAim, c=True, type="joint")
+                #print "\nAll Children..."
+                #print allChildren
+
+                if allChildren != None:
+                    if len(allChildren) == 1:
+                        aimTarget = allChildren[0]
+                    else:
+                        aimTarget = cmds.layoutDialog(t="Child Joints List for: "+jointToAim, ui=partial(findChild, jointToAim))
+                   
+                    #print "\nPRINTING JOINTS LIST---- "
+                    #print "jointToAim: "+jointToAim
+                    #print "aimTarget: "+aimTarget
+                    #print "upNode: "+upNode   
+                    #print "----DONE PRINTING JOINTS LIST"
+                    doOrientJoint(jointToAim, aimTarget, upNode, aimAxis, upAxis, allChildren)
+                else:
+                    # Set the joint orient to none so that it just takes the orientation of its parent.  This is only for the leaf joints in the heirarchy.
+                    cmds.joint(jointToAim, e=True, oj="none", secondaryAxisOrient="yup", ch=True, zso=True)
+
+#----------------------------------------------------------------------------------------
+# This def is only used by the setOrientWithUpVector def IF the heirarchy check box is set.  It creates a UI that outputs the child that the user selects in a UI to be used as the aim child for joints with more than 1 child.
+def findChild(jointToAim):
+    #print "\nFIND CHILD START----->" 
+
+    allChildren = cmds.listRelatives(jointToAim, c=True, type="joint")
+
+    form = cmds.setParent(q=True)
+    cmds.formLayout(form, e=True, width=200, height=300)
+    
+    tsl = cmds.textScrollList("childJointList", nr=10, ams=1, dcc="pm.runtime.ShowAttributeEditorOrChannelBox()")
+    cmds.textScrollList(tsl, e=True, sc=partial(selectItem, tsl))
+    for s in allChildren:
+        cmds.textScrollList(tsl, e=True, a=s, sii=1) 
+    cmds.select(cmds.textScrollList(tsl, q=True, si=True))
+    txt = cmds.text("childText", l="Select the child joint that you want the aim axis to point at and press continue.")
+    butt = cmds.button("continueButton", l="Continue", c=partial(getTextScrollListName, tsl))
+    cmds.formLayout(form, e=True, 
+        af=[(tsl, 'top', 1), (tsl, 'left', 1), (tsl, 'right', 1), (txt, 'left', 1), (txt, 'right', 1), (butt, 'left', 1), (butt, 'right', 1), (butt, 'bottom', 1)], 
+        ac=[(tsl, 'bottom', 1, txt), (txt, 'bottom', 1, butt)])
+
+def selectItem(textScrollListName):
+    cmds.select(cmds.textScrollList(textScrollListName, q=True, si=True))
+    
+def getTextScrollListName(textScrollListName, *args):
+    cmds.layoutDialog(dismiss=cmds.textScrollList(textScrollListName, q=True, si=True)[0])
+
+#----------------------------------------------------------------------------------------
+def doOrientJoint (jointToAim, aimTarget, upNode, aimAxis, upAxis, allChildren):
+    
+    print jointToAim
+    print aimTarget
+    print upNode
+    print "upAxis: " + upAxis 
+    print "aimAxis: " + aimAxis
+    print allChildren
+
+    Coords_up = []
+    if upAxis == "x":
+        Coords_up = [1, 0, 0]
+    if upAxis == "-x":
+        Coords_up = [-1, 0, 0]
+    if upAxis == "y":
+        Coords_up = [0, 1, 0]
+    if upAxis == "-y":
+        Coords_up = [0, -1, 0]
+    if upAxis == "z":
+        Coords_up = [0, 0, 1]
+    if upAxis == "-z":
+        Coords_up = [0, 0, -1]
+
+    Coords_aim = [] 
+    if aimAxis == "x":
+        Coords_aim = [1, 0, 0]
+    if aimAxis == "-x":
+        Coords_aim = [-1, 0, 0]
+    if aimAxis == "y":
+        Coords_aim = [0, 1, 0]
+    if aimAxis == "-y":
+        Coords_aim = [0, -1, 0]
+    if aimAxis == "z":
+        Coords_aim = [0, 0, 1]
+    if aimAxis == "-z":
+        Coords_aim = [0, 0, -1]
+
+    #Do the aim constraint to make the joint orient properly
+    cmds.parent(allChildren, w=True)
+    aimConst = cmds.aimConstraint(aimTarget, jointToAim, n=jointToAim+"_AimConst", o=[0, 0, 0], w=1, aim=[Coords_aim[0], Coords_aim[1], Coords_aim[2]], u=[Coords_up[0], Coords_up[1], Coords_up[2]], wut="object", wuo=upNode)[0]
+    cmds.delete(aimConst)
+    cmds.select(jointToAim)
+    cmds.makeIdentity(apply=True, t=1, r=1, s=1, n=0) 
+    cmds.delete(ch=True)
+    cmds.parent(allChildren, jointToAim)
+
+#----------------------------------------------------------------------------------------
+def orientJointOB(*args):
+    #Spine Option Variables
+    if cmds.optionVar(exists="orientJointHeirarchyOV") == 0:
+        cmds.optionVar(iv=("orientJointHeirarchyOV", 1)) 
+    if cmds.optionVar(exists="OV_frameLayout_collapse") == 0:
+        cmds.optionVar(iv=("OV_frameLayout_collapse", 0))
+    if cmds.optionVar(exists="OV_upVectorType") == 0:
+        cmds.optionVar(iv=("OV_upVectorType", 3))
+    if cmds.optionVar(exists="OV_aimVectorType") == 0:
+        cmds.optionVar(iv=("OV_aimVectorType", 1))
+        
+    # Create the window
+    if cmds.window("orientJointsWindow", exists=True):
+        cmds.deleteUI("orientJointsWindow", window=True)
+    cmds.window("orientJointsWindow", t="Orient Joints with Up Vector", menuBar=True, s=0, mxb=0, rtf=1, w=300, h=175)
+    
+    # Create the menu
+    cmds.menu("editMenu", l="Edit", p="orientJointsWindow", to=0, aob=True)
+    cmds.menuItem("resetSettingsItem", l="Reset Settings", ann="Reset all of the settings in this window to their default.", echoCommand=True, p="editMenu", c=resetOrientJointVariables)
+    cmds.menu("helpMenu", l="Help", p="orientJointsWindow", hm=1, to=0, aob=True)
+    cmds.menuItem("helpItem", l="Help on Orient Joints with Up Vector...", ann="Help on this tool.  Currently there is none...", echoCommand=True, p="helpMenu", c="")
+
+    # base layout for the entire UI
+    cmds.columnLayout(adjustableColumn=True)
+    cmds.frameLayout(label="Options", bgc=[0.1, 0.1, 0.1], cll=True, cl=cmds.optionVar(q="OV_frameLayout_collapse"), cc="cmds.optionVar(iv=(\"OV_frameLayout_collapse\", 1))", ec="cmds.optionVar(iv=(\"OV_frameLayout_collapse\", 0))", borderStyle="etchedOut")
+
+    # Aim vector dropdown menu
+    cmds.optionMenuGrp("OMG_aimVectorType", l="Aim Axis", cw=[1, 80], ann="Aim Axis is the axis of your joints that will point at the child. Default is \"x\"", cc="cmds.optionVar(iv=(\"OV_aimVectorType\", "+"cmds.optionMenuGrp(\"OMG_aimVectorType\", q=True, sl=True)"+"))")
+    cmds.menuItem("xItem", l="x")
+    cmds.menuItem("nxItem", l="-x")
+    cmds.menuItem("yItem", l="y")
+    cmds.menuItem("nyItem", l="-y")
+    cmds.menuItem("zItem", l="z")
+    cmds.menuItem("nzItem", l="-z")
+    cmds.optionMenuGrp("OMG_aimVectorType", e=True, sl=cmds.optionVar(q="OV_aimVectorType"))
+
+    # Up vector dropdown menu
+    cmds.optionMenuGrp("OMG_upVectorType", l="Up Axis", cw=[1, 80], ann="Up Axis is the axis of your joint that will be aimed at the up object you have selected to be the up vector for your joint. Default is \"y\"", cc="cmds.optionVar(iv=(\"OV_upVectorType\", "+"cmds.optionMenuGrp(\"OMG_upVectorType\", q=True, sl=True)"+"))")
+    cmds.menuItem("xItem", l="x")
+    cmds.menuItem("nxItem", l="-x")
+    cmds.menuItem("yItem", l="y")
+    cmds.menuItem("nyItem", l="-y")
+    cmds.menuItem("zItem", l="z")
+    cmds.menuItem("nzItem", l="-z")
+    cmds.optionMenuGrp("OMG_upVectorType", e=True, sl=cmds.optionVar(q="OV_upVectorType"))
+
+    # Heirarchy checkbox
+    cmds.checkBox("heirarchyCB", l="Heirarchy?", ann="Do you want to orient all of the children of the selected joint?", v=cmds.optionVar(q="orientJointHeirarchyOV"), onc="cmds.optionVar(iv=(\"orientJointHeirarchyOV\", 1))", ofc="cmds.optionVar(iv=(\"orientJointHeirarchyOV\", 0))")
+    cmds.setParent( '..' )  
+    cmds.setParent( '..' ) 
+    cmds.setParent( '..' )  
+    cmds.setParent( '..' ) 
+
+    frmLyout = cmds.formLayout(nd=100) 
+
+    # Buttons
+    cmds.button("orientJointButton", l="Orient Joints", ann="Orient Joints with the current settings and close this window.", al="center", w=75, h=20, c=applyAndClose)
+    cmds.button("applyGroupButton", l="Apply", ann="Orient Joints with the current settings.", al="center", w=75, h=20, c=findOrientJointValues)
+    cmds.button("CloseButton", l="Close", ann="Close this window.", al="center", w=75, h=20, c="cmds.deleteUI(\"orientJointsWindow\", window=True)")
+    cmds.setParent( '..' )  
+    cmds.setParent( '..' ) 
+
+    cmds.formLayout(frmLyout, e=True, af=[("orientJointButton", "top", 5), ("orientJointButton", "left", 5), ("orientJointButton", "bottom", 5), ("applyGroupButton", "top", 5), ("applyGroupButton", "bottom", 5), ("CloseButton", "top", 5), ("CloseButton", "right", 5), ("CloseButton", "bottom", 5)], ap=[("orientJointButton", "right", 1, 33), ("applyGroupButton", "right", 1, 66)], ac=[("applyGroupButton", "left", 1, "orientJointButton"), ("CloseButton", "left", 1, "applyGroupButton")])
+    cmds.showWindow()
+
+#----------------------------------------------------------------------------------------
+def applyAndClose(*args):
+    findOrientJointValues()
+    cmds.deleteUI("orientJointsWindow", window=True)
+
+#----------------------------------------------------------------------------------------
+def resetOrientJointVariables(*args):
+    #Fixup options
+    cmds.optionVar(iv=("orientJointHeirarchyOV", 1))
+    cmds.checkBox("heirarchyCB", e=True, v=1)
+
+    cmds.optionVar(iv=("OV_upVectorType", 3))
+    cmds.optionMenuGrp("OMG_upVectorType", e=True, sl=3)
+
+    cmds.optionVar(iv=("OV_aimVectorType", 1))
+    cmds.optionMenuGrp("OMG_aimVectorType", e=True, sl=1)
+
+#----------------------------------------------------------------------------------------
+def findOrientJointValues(*args):
+    upVector = cmds.optionVar(q="OV_upVectorType")
+    aimVector = cmds.optionVar(q="OV_aimVectorType")
+
+    #Set the Up Axis string
+    if upVector == 1:
+        upAxis = "x" 
+    if upVector == 2:
+        upAxis = "-x" 
+    if upVector == 3:
+        upAxis = "y" 
+    if upVector == 4:
+        upAxis = "-y" 
+    if upVector == 5:
+        upAxis = "z" 
+    if upVector == 6:
+        upAxis = "-z" 
+
+    #Set the Aim Axis string
+    if aimVector == 1:
+        aimAxis = "x" 
+    if aimVector == 2:
+        aimAxis = "-x" 
+    if aimVector == 3:
+        aimAxis = "y" 
+    if aimVector == 4:
+        aimAxis = "-y" 
+    if aimVector == 5:
+        aimAxis = "z" 
+    if aimVector == 6:
+        aimAxis = "-z" 
+
+    if cmds.optionVar(q="orientJointHeirarchyOV") == 0:
+        heirarchy = 0 
+    if cmds.optionVar(q="orientJointHeirarchyOV") == 1:
+        heirarchy = 1 
+
+    setOrientWithUpVector(aimAxis, upAxis, heirarchy, 0, None)
+
+# Run it through the UI
+#orientJointOB()
+
+# Run it without the UI (Values are passed in as args)
+#setOrientWithUpVector("x", "y", 0, 1, "lowerarm_l")
+#setOrientWithUpVector("x", "y", 1, 0, None)
diff --git a/Engine/Extras/Maya_AnimationRiggingTools/ARTv1/MayaTools/General/Scripts/Tools/ART_RenameSuffix.py b/Engine/Extras/Maya_AnimationRiggingTools/ARTv1/MayaTools/General/Scripts/Tools/ART_RenameSuffix.py
new file mode 100644
--- /dev/null
+++ b/Engine/Extras/Maya_AnimationRiggingTools/ARTv1/MayaTools/General/Scripts/Tools/ART_RenameSuffix.py
@@ -0,0 +1,10 @@
+import maya.cmds as cmds
+
+def suffixHeirarchyNames(*args):
+    selection = cmds.ls(sl=True, type="transform")
+    result = cmds.promptDialog(title="Suffix Seleceted", message="Enter Suffix:", button=["Ok", "Cancel"])
+    if result == "Ok":
+        suffix = cmds.promptDialog(q=True, text=True)
+        for one in selection:
+            cmds.rename(one, one+suffix)
+#suffixHeirarchyNames()
\ No newline at end of file
diff --git a/Engine/Extras/Maya_AnimationRiggingTools/ARTv1/MayaTools/General/Scripts/Tools/ART_ReverseAttrHookup.py b/Engine/Extras/Maya_AnimationRiggingTools/ARTv1/MayaTools/General/Scripts/Tools/ART_ReverseAttrHookup.py
new file mode 100644
--- /dev/null
+++ b/Engine/Extras/Maya_AnimationRiggingTools/ARTv1/MayaTools/General/Scripts/Tools/ART_ReverseAttrHookup.py
@@ -0,0 +1,69 @@
+import maya.cmds as cmds
+from functools import partial
+
+# Creates a connection between the 2 selected objects with a reverse node to drive to values in opposite directions of each other based on one attr.
+# This is mostly used for having a single attr drive the 2 weights on a constraint with 2 sources in opposite of each other.
+# CRA FEATURE: going forward, add the ability to have 3 or even more inputs driven by the single output.
+# CRA FEATURE: Add check boxes and filters for filtering the attr lists better.
+# To use, select the driver, then the driven, then run the tool.
+
+textScrollList = None
+menuDriven0 = None
+menuDriven1 = None
+
+def ReverseAttrHookupUI(*args):
+    
+    print "UI"
+
+    selection = cmds.ls(sl=True)
+    if len(selection) != 2:
+        cmds.confirmDialog(icon = "warning!!", title = "Reverse Attr Hookup Tool", message = "You must select only 2 objects.  Driver then Driven.")
+        return
+    else:
+        driver = selection[0]
+        driven = selection[1]
+
+
+        attrsListDriver = cmds.listAttr(driver, o=True, c=True, u=True, s=True, k=False)
+        attrsListDriven = cmds.listAttr(driven, o=True, c=True, u=True, s=True, k=True)
+
+
+        if cmds.window("ReverseAttrHookupWin", exists = True):
+            cmds.deleteUI("ReverseAttrHookupWin")
+
+        window = cmds.window("ReverseAttrHookupWin", w=4300, h=500, title="Reverse Attr Hookup Tool", mxb=True, mnb=True, sizeable=False, rtf=True)
+        mainLayout = cmds.columnLayout(w = 300, h = 400, rs = 5, co = ["both", 5])
+        cmds.text(l=driver)
+        textScrollList = cmds.textScrollList("artProjCharacterList", w = 290, h = 300, parent = mainLayout, ra=True, ann="This is the attribute that will drive the Driven attributes listed below.")
+        cmds.textScrollList(textScrollList, e=True, a=attrsListDriver)
+        
+        menuDriven0 = cmds.optionMenuGrp(l=driven+" 0", cw=[1, 175], ann="This attribute will be driven as a direct connection from the driver.  When the driver is 1, this input will be 1.  When the driver is 0, this input will be 0.")
+        for attr in attrsListDriven:
+            cmds.menuItem(l=attr)
+        
+        menuDriven1 = cmds.optionMenuGrp(l=driven+" 1", cw=[1, 175], ann="This attribute will be driven by the driver through a reverse node.  When the driver is 1, this input will be 0.  When the driver is 0, this input will be 1.")
+        for attr in attrsListDriven:
+            cmds.menuItem(l=attr)
+
+        buttonUpdate = cmds.button(w = 290, h = 40, label = "Update", c=ReverseAttrHookupUI, ann = "Refresh the UI.", parent = mainLayout)
+        buttonConnect = cmds.button(w = 290, h = 40, label = "Connect", c=partial(ReverseAttrHookup, driver, driven, textScrollList, menuDriven0, menuDriven1), ann = "Connect the Selected Attrs.", parent = mainLayout)
+         
+        cmds.showWindow(window)
+    
+# This script hooks up the attrs
+def ReverseAttrHookup(driver, driven, textScrollList, menuDriven0, menuDriven1, *args):
+    
+    print "SCRIPT"
+
+    driverAttr = cmds.textScrollList(textScrollList, q=True, si=True)[0]
+    drivenAttr0 = cmds.optionMenuGrp(menuDriven0, q=True, v=True)
+    drivenAttr1 = cmds.optionMenuGrp(menuDriven1, q=True, v=True)
+    
+    print driverAttr
+    print drivenAttr0
+    print drivenAttr1
+
+    reverseNode = cmds.shadingNode("reverse", asUtility=True, name=driver+"_reverse")
+    cmds.connectAttr(driver+"."+driverAttr, driven+"."+drivenAttr0, f=True)
+    cmds.connectAttr(driver+"."+driverAttr, reverseNode+".inputX", f=True)
+    cmds.connectAttr(reverseNode+".outputX", driven+"."+drivenAttr1, f=True)
diff --git a/Engine/Extras/Maya_AnimationRiggingTools/ARTv1/MayaTools/General/Scripts/Tools/ART_SurfaceVertexSelector.py b/Engine/Extras/Maya_AnimationRiggingTools/ARTv1/MayaTools/General/Scripts/Tools/ART_SurfaceVertexSelector.py
new file mode 100644
--- /dev/null
+++ b/Engine/Extras/Maya_AnimationRiggingTools/ARTv1/MayaTools/General/Scripts/Tools/ART_SurfaceVertexSelector.py
@@ -0,0 +1,71 @@
+import maya.cmds as cmds
+import maya.OpenMaya as om
+from functools import partial
+
+# This script walks along the surface of a mesh from a single component and find the nearest N verts based on the number you pass into it (count).
+# It then selects these verts and returns them as a list called "elements"
+def surfaceVertSelector(startElement, *args):
+    count = args[0]
+    #print count
+    vertCount = 0
+    startDist = 0
+    cmds.select(startElement)
+
+    # Find all of the verts on the selected mesh and see if they are less than count.  If they are then set count to be all verts.
+    allVerts = cmds.polyEvaluate(v=True)
+    #print allVerts
+    if count > allVerts:
+        count = allVerts
+    #print args[0]
+        
+    # Run through the mesh and using soft select to traverse out and find verts until we reach the count value.
+    while vertCount < count:
+        startDist = startDist+0.1
+        #print "startDist: ", startDist
+        cmds.softSelect(softSelectEnabled=True, ssf=1, ssd=startDist)
+        selection = om.MSelectionList()
+        surfaceVertSelector = om.MRichSelection()
+        om.MGlobal.getRichSelection(surfaceVertSelector)
+        surfaceVertSelector.getSelection(selection)
+
+        dagPath = om.MDagPath()
+        component = om.MObject()
+
+        iter = om.MItSelectionList( selection,om.MFn.kMeshVertComponent )
+        elements = []
+        while not iter.isDone():
+            iter.getDagPath( dagPath, component )
+            dagPath.pop()
+            node = dagPath.fullPathName()
+            fnComp = om.MFnSingleIndexedComponent(component)
+
+            for i in range(fnComp.elementCount()):
+                elements.append('%s.vtx[%i]' % (node, fnComp.element(i)))
+            iter.next()
+        vertCount = len(elements)
+        #print "vertCount: ", vertCount
+    cmds.softSelect(softSelectEnabled=False)
+    cmds.select(elements)
+
+#startElement = cmds.ls(sl=True)
+#surfaceVertSelector(startElement, 75)
+    
+# Create a UI that will drive the selection tool dynamically using dragCommand.
+# The reset button will allow you to change the starting component(s) on your mesh.
+def surfaceVertSelectorUI(*args):
+    startElement = cmds.ls(sl=True, type="float3")
+    if startElement:
+        if cmds.window("surfaceVertSelector", exists = True):
+            cmds.deleteUI("surfaceVertSelector")
+
+        window = cmds.window("surfaceVertSelector", title='Surface Vert Selector', s=False)
+        cmds.columnLayout()
+        cmds.intSliderGrp("vertCountIntSliderGrp", label='Vert Count', field=True, minValue=1, maxValue=2000, value=1, step=1, cw3=[75, 50, 1000], dragCommand=partial(surfaceVertSelector, startElement))
+        #surfaceVertSelector(startElement, 20)
+        cmds.button("refreshButton", label="Reset", c=surfaceVertSelectorUI)
+
+        cmds.showWindow(window)
+    else:
+        cmds.warning("You must select some components.")
+
+#surfaceVertSelectorUI()
\ No newline at end of file
diff --git a/Engine/Extras/Maya_AnimationRiggingTools/ARTv1/MayaTools/General/Scripts/Tools/ART_TransferWeights.py b/Engine/Extras/Maya_AnimationRiggingTools/ARTv1/MayaTools/General/Scripts/Tools/ART_TransferWeights.py
new file mode 100644
--- /dev/null
+++ b/Engine/Extras/Maya_AnimationRiggingTools/ARTv1/MayaTools/General/Scripts/Tools/ART_TransferWeights.py
@@ -0,0 +1,172 @@
+import maya.cmds as cmds
+from PySide import QtGui
+import maya.OpenMayaUI as mui
+import shiboken
+
+#Written by: Charles R. Anderson
+#Dec. 4, 2015 
+#Questions or comments email me: [REDACTED_EMAIL]
+
+
+# This function finds the main Maya Window
+def getMayaWindow():
+    pointer = mui.MQtUtil.mainWindow()
+    return shiboken.wrapInstance(long(pointer), QtGui.QWidget)
+
+# This function finds the cluster for the passed in mesh.
+def findCluster(meshName):
+    skinClusters = []
+    skinHistory = []
+    skinHistory = cmds.listHistory(meshName, lv=0)
+    for node in skinHistory:
+        if cmds.nodeType(node) == "skinCluster":
+            skinClusters.append(node)
+
+    return skinClusters[0]
+
+# This function transfers the weighting for each mesh in the meshes list for the influences in the jointsToRemove list to the jointToTransferTo joint.
+def LOD_transferWeights(meshes, jointsToRemove, jointToTransferTo, *args):
+    for mesh in meshes:
+        #print "MESH: ", mesh
+        
+        # Find the skin cluster for the current mesh
+        cluster = findCluster(mesh)
+        #print "CLUSTER: ", cluster
+        
+        # Prune weights on the current mesh
+        cmds.skinPercent(cluster, mesh, prw=0.001)
+
+        # Find all of the current influences on the current skin cluster.
+        meshInfluences = cmds.skinCluster(cluster, q=True, inf=True)
+        #print "Current Influences: ", meshInfluences
+
+        for joint in jointsToRemove:
+            if joint in meshInfluences:
+                #print "Current Joint: ", joint
+                
+                # If the jointToTransferTo is not already an influence on the current mesh then add it.
+                currentInfluences = cmds.skinCluster(cluster, q=True, inf=True)
+                if jointToTransferTo not in currentInfluences:
+                    cmds.skinCluster(cluster, e=True, ai=jointToTransferTo, lw=False)
+
+                # Now transfer all of the influences we want to remove onto the jointToTransferTo.
+                for x in range(cmds.polyEvaluate(mesh, v=True)):
+                    #print "TRANSFERRING DATA....."
+                    value = cmds.skinPercent(cluster, (mesh+".vtx["+str(x)+"]"), t=joint, q=True)
+                    if value > 0:
+                        cmds.skinPercent(cluster, (mesh+".vtx["+str(x)+"]"), tmw=[joint, jointToTransferTo])
+
+        # Remove unused influences
+        currentInfluences = cmds.skinCluster(cluster, q=True, inf=True)
+        #print "Current Influences: ", currentInfluences
+        influencesToRemove = []
+        weightedInfs = cmds.skinCluster(cluster, q=True, weightedInfluence=True)
+        #print "Weighted Influences: ", weightedInfs
+        for inf in currentInfluences:
+            #print "Influence: ", inf
+            if inf not in weightedInfs:
+                #print "Update Influences to Remove List: ", inf
+                influencesToRemove.append(inf)
+            
+        #print "ToRemove Influences: ", influencesToRemove
+        if influencesToRemove != []:
+            for inf in influencesToRemove:
+                cmds.skinCluster(cluster, e=True, ri=inf)
+
+
+########  Example
+'''
+meshes = ["hands_geo_LOD4"]
+jointsToRemove = ["index_01_r", "index_02_r", "middle_01_r", "middle_02_r", "pinky_01_r", "pinky_02_r", "ring_01_r", "ring_02_r", "thumb_01_r", "thumb_02_r"]
+jointToTransferTo = "hand_r"
+LOD_transferWeights(meshes, jointsToRemove, jointToTransferTo)
+'''
+##################
+
+
+
+
+
+
+# THIS IS OLD CODE THAT I WAS USING TO CREATE A UI FOR TRANSFERING WEIGHTS.  MAY STILL USE SOME OF IT FOR A TOOL TO CREATE THE REMOVE BONES LIST.
+'''
+def create_LOD_TransferWeightsUI():
+
+    # check for existing window.
+    windowName = "LOD_TransferWeightsWindow"
+    if cmds.window("LOD_TransferWeightsWindow", exists=True):
+        cmds.deleteUI("LOD_TransferWeightsWindow", wnd=True)
+        
+    # Create a window
+    mainMayaWindow = getMayaWindow()
+    LOD_TransferWeightsWindow = QtGui.QMainWindow(mainMayaWindow)
+    LOD_TransferWeightsWindow.setObjectName(windowName)
+
+    # Create a font
+    font = QtGui.QFont()
+    font.setPointSize(20)
+    font.setBold(True)
+
+    # create a widget
+    widget = QtGui.QWidget()
+    LOD_TransferWeightsWindow.setCentralWidget(widget)
+
+    # create a layout
+    layout = QtGui.QVBoxLayout(widget)
+
+    # create a button and add it to the layout
+    button = QtGui.QPushButton("Create Locator")
+    layout.addWidget(button)
+    button.setFont(font)
+    button.setMinimumSize(200, 40)
+    button.setMaximumSize(200, 40)
+    imagePath = cmds.internalVar(upd=True) + "icons/extractDeltasDuplicate.png"
+    button.setStyleSheet("background-image: url("+imagePath+"); border:solid black 1px; color")
+    button.clicked.connect(LOD_transferWeights)
+
+    # Create a close button
+    closeButton = QtGui.QPushButton("Close")
+    layout.addWidget(closeButton)
+    closeButton.setFont(font)
+    closeButton.setMinimumSize(200, 40)
+    closeButton.setMaximumSize(200, 40)
+    imagePath = cmds.internalVar(upd=True) + "icons/extractDeltasDuplicate.png"
+    closeButton.setStyleSheet("background-image: url("+imagePath+"); border:solid black 1px")
+    closeButton.clicked.connect(LOD_TransferWeightsWindow.close)
+
+
+    # show the button
+    LOD_TransferWeightsWindow.show()
+
+create_LOD_TransferWeightsUI()
+
+
+
+def OutlinerSelection():
+    itemSelection = cmds.ls(sl=True, type="joint")
+    sc = cmds.selectionConnection(lst=True)
+    if len(itemSelection):
+        for node in itemSelection:
+            print node
+            cmds.selectionConnection(sc, e=True, select=node)
+
+        # Create a window layout to hold outliner
+        newWindow = cmds.window(title="Outliner (Custom)", iconName="Outliner*", rtf=0, h=800, w=500)
+        frame = cmds.frameLayout(labelVisible = False)
+        oped = cmds.outlinerEditor("jointOutliner", showSetMembers=True, unParent=True, mainListConnection=sc, selectionConnection="modelList", eo=True)
+        testWin = cmds.outlinerEditor(oped, e=True, parent=frame)
+
+        # scriptJob to delete selectionConnection
+        print "Deleting of the connection"
+        cmds.scriptJob(uiDeleted = [newWindow, "cmds.deleteUI('"+sc+"')"])
+
+        cmds.showWindow(newWindow)
+
+    else:
+        cmds.warning("Nothing Selected! Custom Outliner will not be created")
+
+OutlinerSelection() 
+
+jointList = cmds.ls(sl=True, type="joint")
+print jointList
+'''
\ No newline at end of file
diff --git a/Engine/Extras/Maya_AnimationRiggingTools/ARTv1/MayaTools/General/Scripts/Tools/ART_addSpaceSwitching.py b/Engine/Extras/Maya_AnimationRiggingTools/ARTv1/MayaTools/General/Scripts/Tools/ART_addSpaceSwitching.py
new file mode 100644
--- /dev/null
+++ b/Engine/Extras/Maya_AnimationRiggingTools/ARTv1/MayaTools/General/Scripts/Tools/ART_addSpaceSwitching.py
@@ -0,0 +1,63 @@
+import maya.cmds as cmds
+
+# Creates a new space switching object for the mainObj using the spaceMasterObj.  
+# If you want this new space to be in control, set the control arg to "True".  False will leave everything as it is and just add the new space without it being turned on.
+
+def addSpaceSwitching(mainObj, spaceMasterObj, control, *args):
+    # find all of the objects we need to do the space switching
+    space_switcher_follow = mainObj+"_space_switcher_follow"
+    space_switcher = mainObj+"_space_switcher"
+    cmds.select(space_switcher)
+    
+    spaceMasterObjClean = spaceMasterObj.replace(":", "_")
+    mainObjClean = mainObj.replace(":", "_")
+     
+    # add attr to the space switcher node
+    cmds.addAttr(space_switcher, ln = "space_"+spaceMasterObjClean, minValue = 0, maxValue = 1, dv = 0, keyable = True)
+    
+    # add constraint to the new object on the follow node
+    constraint = cmds.parentConstraint(spaceMasterObj, space_switcher_follow, mo = True)[0]
+
+    # find the entire list of targets on the constraint node
+    targets = cmds.parentConstraint(constraint, q = True, targetList = True)
+
+    # using the targets list, find the weight value for the node we are working with.
+    weight = 0
+    for i in range(int(len(targets))):
+        if targets[i].find(spaceMasterObj) != -1:
+            weight = i
+            
+    # Connect the attr we added earlier to the weight value in the constraint
+    tempName = spaceMasterObj
+    #print spaceMasterObj.find(":")
+    if spaceMasterObj.find(":") != -1:
+        tempName = spaceMasterObj.partition(":")[2]
+        #print "TEST"
+    
+    #print "tempname = ",tempName
+    cmds.connectAttr(space_switcher+".space_"+spaceMasterObjClean, constraint+"."+tempName+"W"+str(weight))
+
+    # lockNode on space object so it cannot be deleted by the user (if node is not a referenced node)
+    if spaceMasterObj.find(":") == -1:
+        cmds.lockNode(spaceMasterObj)
+
+    # set the values of the space switcher based on user input.
+    if control == True:
+        attrs = cmds.listAttr(space_switcher, k=True)
+        for attr in attrs:
+            if attr.find("space") != -1:
+                cmds.setAttr(space_switcher+"."+attr, 0)
+                
+        cmds.setAttr(space_switcher+"."+"space_"+spaceMasterObjClean, 1)
+
+
+    cmds.select(space_switcher)
+
+
+# Example calls to this script.
+'''
+addSpaceSwitching("Rider:ik_foot_anim_l", "cockpit_pedal_l_anim", True)
+'''
+
+#addSpaceSwitching("Rider:ik_foot_anim_l", "cockpit_foot_pedal_l_anim", True)
+#addSpaceSwitching("cockpit_foot_pedal_l_anim", "Rider:ik_foot_anim_l", True)
